MVVM: Move the ModelViewViewModel plugin to //UE5/Main/

diff --git a/Engine/Plugins/Runtime/ModelViewViewModel/Source/ModelViewViewModel/Private/Bindings/MVVMBindingHelper.cpp b/Engine/Plugins/Runtime/ModelViewViewModel/Source/ModelViewViewModel/Private/Bindings/MVVMBindingHelper.cpp
new file mode 100644
--- /dev/null
+++ b/Engine/Plugins/Runtime/ModelViewViewModel/Source/ModelViewViewModel/Private/Bindings/MVVMBindingHelper.cpp
@@ -0,0 +1,696 @@
+// Copyright Epic Games, Inc. All Rights Reserved.
+
+#include "Bindings/MVVMBindingHelper.h"
+#include "FieldNotification/FieldNotificationHelpers.h"
+#include "UObject/Class.h"
+#include "UObject/PropertyAccessUtil.h"
+#include "UObject/UnrealType.h"
+
+namespace UE::MVVM::BindingHelper
+{
+	namespace Private
+	{
+		static const FName NAME_BlueprintPrivate = "BlueprintPrivate";
+		static const FName NAME_DeprecatedFunction = "DeprecatedFunction";
+		static const FName NAME_BlueprintGetter = "BlueprintGetter";
+		static const FName NAME_BlueprintSetter = "BlueprintSetter";
+
+
+		bool IsValidCommon(const UFunction* InFunction)
+		{
+			bool bResult = InFunction != nullptr
+				&& !InFunction->HasAnyFunctionFlags(FUNC_Net | FUNC_Event | FUNC_EditorOnly)
+				&& InFunction->HasAllFunctionFlags(FUNC_BlueprintCallable);
+
+#if WITH_EDITOR
+			if (bResult)
+			{
+				if (InFunction->HasMetaData(NAME_DeprecatedFunction))
+				{
+					bResult = false;
+				}
+			}
+#endif
+
+			return bResult;
+		}
+	} // namespace
+
+	// CPF_Protected is not used
+	bool IsValidForSourceBinding(const FProperty* InProperty)
+	{
+		return InProperty != nullptr
+			&& !InProperty->HasAnyPropertyFlags(CPF_Deprecated | CPF_EditorOnly)
+			&& InProperty->HasAnyPropertyFlags(CPF_BlueprintVisible | CPF_BlueprintAssignable)
+#if WITH_EDITOR
+			&& !InProperty->GetBoolMetaData(Private::NAME_BlueprintPrivate)
+#endif
+			;
+	}
+
+
+	bool IsValidForDestinationBinding(const FProperty* InProperty)
+	{
+		return InProperty != nullptr
+			&& !InProperty->HasAnyPropertyFlags(CPF_Deprecated | CPF_EditorOnly | CPF_BlueprintReadOnly)
+			&& InProperty->HasAnyPropertyFlags(CPF_BlueprintVisible | CPF_BlueprintAssignable)
+#if WITH_EDITOR
+			&& !InProperty->GetBoolMetaData(Private::NAME_BlueprintPrivate)
+#endif
+			;
+	}
+
+
+	bool IsValidForSourceBinding(const UFunction* InFunction)
+	{
+		//UE::FieldNotification::Helpers::IsValidAsField(InFunction)
+		return Private::IsValidCommon(InFunction)
+			&& InFunction->HasAllFunctionFlags(FUNC_Const)
+			&& InFunction->NumParms == 1
+			&& GetReturnProperty(InFunction) != nullptr;
+	}
+
+
+	bool IsValidForDestinationBinding(const UFunction* InFunction)
+	{
+		// A setter can return a value "bool SetValue(int32)" but is not valid for MVVM
+		//&& ((InFunction->NumParms == 1 && GetReturnProperty() == nullptr) || (InFunction->NumParms == 2 && GetReturnProperty() != nullptr));
+
+		return Private::IsValidCommon(InFunction)
+			&& !InFunction->HasAnyFunctionFlags(FUNC_Const | FUNC_BlueprintPure)
+			&& InFunction->NumParms == 1
+			&& GetFirstArgumentProperty(InFunction) != nullptr;
+	}
+
+
+	bool IsValidForSourceBinding(const FMVVMConstFieldVariant InVariant)
+	{
+		if (InVariant.IsProperty())
+		{
+			return IsValidForSourceBinding(InVariant.GetProperty());
+		}
+		else if (InVariant.IsFunction())
+		{
+			return IsValidForSourceBinding(InVariant.GetFunction());
+		}
+		return false;
+	}
+
+
+	bool IsValidForDestinationBinding(const FMVVMConstFieldVariant InVariant)
+	{
+		if (InVariant.IsProperty())
+		{
+			return IsValidForDestinationBinding(InVariant.GetProperty());
+		}
+		else if (InVariant.IsFunction())
+		{
+			return IsValidForDestinationBinding(InVariant.GetFunction());
+		}
+		return false;
+	}
+
+
+	bool IsValidForRuntimeConversion(const UFunction* InFunction)
+	{
+		if (Private::IsValidCommon(InFunction)
+			&& InFunction->HasAllFunctionFlags(FUNC_Static)
+			&& InFunction->NumParms == 2)
+		{
+			const FProperty* ReturnProperty = GetReturnProperty(InFunction);
+			const FProperty* FirstArgumentProperty = GetFirstArgumentProperty(InFunction);
+			return ReturnProperty && FirstArgumentProperty && ReturnProperty != FirstArgumentProperty;
+		}
+		return false;
+	}
+
+
+#if WITH_EDITOR
+	bool IsAccessibleDirectlyForSourceBinding(const FProperty* InProperty)
+	{
+		return IsValidForSourceBinding(InProperty) && !InProperty->GetBoolMetaData(Private::NAME_BlueprintPrivate);
+	}
+
+
+	bool IsAccessibleDirectlyForDestinationBinding(const FProperty* InProperty)
+	{
+		return IsValidForDestinationBinding(InProperty) && !InProperty->GetBoolMetaData(Private::NAME_BlueprintPrivate);
+	}
+
+
+	bool IsAccessibleWithGetterForSourceBinding(const FProperty* InProperty)
+	{
+		return IsValidForSourceBinding(InProperty) && InProperty->HasMetaData(Private::NAME_BlueprintGetter);
+	}
+
+
+	bool IsAccessibleWithSetterForDestinationBinding(const FProperty* InProperty)
+	{
+		return IsValidForSourceBinding(InProperty) && InProperty->HasMetaData(Private::NAME_BlueprintSetter);
+	}
+#endif //WITH_EDITOR
+
+
+	FMVVMFieldVariant FindFieldByName(const UStruct* Container, FMVVMBindingName BindingName)
+	{
+		if (Container)
+		{
+			if (const UClass* Class = Cast<const UClass>(Container))
+			{
+				if (UFunction* Function = Class->FindFunctionByName(BindingName.ToName()))
+				{
+					return FMVVMFieldVariant(Function);
+				}
+			}
+			if (FProperty* Property = PropertyAccessUtil::FindPropertyByName(BindingName.ToName(), Container))
+			{
+				return FMVVMFieldVariant(Property);
+			}
+		}
+		return FMVVMFieldVariant();
+	}
+
+
+	namespace Private
+	{
+		FString TryGetPropertyTypeCommon(const FProperty* InProperty)
+		{
+			if (InProperty == nullptr)
+			{
+				return TEXT("The property is invalid.");
+			}
+
+			if (InProperty->HasAnyPropertyFlags(CPF_Deprecated))
+			{
+				return FString::Printf(TEXT("The property '%s' is depreated."), *InProperty->GetName());
+			}
+
+			if (InProperty->HasAnyPropertyFlags(CPF_EditorOnly))
+			{
+				return FString::Printf(TEXT("The property '%s' is only available in the editor."), *InProperty->GetName());
+			}
+			
+			if (!InProperty->HasAnyPropertyFlags(CPF_BlueprintVisible))
+			{
+				return FString::Printf(TEXT("The property '%s' is not visible to the Blueprint script."), *InProperty->GetName());
+			}
+
+			return FString();
+		}
+
+		FString TryGetPropertyTypeCommon(const UFunction* InFunction)
+		{
+			if (InFunction == nullptr)
+			{
+				return (TEXT("The function is invalid."));
+			}
+
+			if (!InFunction->HasAllFunctionFlags(FUNC_BlueprintCallable))
+			{
+				return (FString::Printf(TEXT("The function '%s' is not BlueprintCallable."), *InFunction->GetName()));
+			}
+
+			if (InFunction->HasAllFunctionFlags(FUNC_Net))
+			{
+				return (FString::Printf(TEXT("The function '%s' is networked."), *InFunction->GetName()));
+			}
+
+			if (InFunction->HasAllFunctionFlags(FUNC_Event))
+			{
+
+				return (FString::Printf(TEXT("The function '%s' is an event."), *InFunction->GetName()));
+			}
+
+			if (InFunction->HasAllFunctionFlags(FUNC_EditorOnly))
+			{
+				return (FString::Printf(TEXT("The function '%s' is editor only"), *InFunction->GetName()));
+			}
+
+#if WITH_EDITOR
+			{
+				if (InFunction->HasMetaData(NAME_DeprecatedFunction))
+				{
+					return (FString::Printf(TEXT("The function '%s' is deprecated"), *InFunction->GetName()));
+				}
+			}
+#endif
+
+			return FString();
+		}
+	} // namespace
+
+	TValueOrError<const FProperty*, FString> TryGetPropertyTypeForSourceBinding(const FProperty* Property)
+	{
+		FString Result = Private::TryGetPropertyTypeCommon(Property);
+		if (!Result.IsEmpty())
+		{
+			return MakeError(MoveTemp(Result));
+		}
+
+		return MakeValue(Property);
+	}
+
+	TValueOrError<const FProperty*, FString> TryGetPropertyTypeForSourceBinding(const UFunction* Function)
+	{
+		FString Result = Private::TryGetPropertyTypeCommon(Function);
+		if (!Result.IsEmpty())
+		{
+			return MakeError(MoveTemp(Result));
+		}
+
+		if (!Function->HasAllFunctionFlags(FUNC_Const))
+		{
+			return MakeError(FString::Printf(TEXT("The function '%s' is not const."), *Function->GetName()));
+		}
+
+		//if (!InFunction->HasAllFunctionFlags(FUNC_BlueprintPure))
+		//{
+		//	return MakeError(FString::Printf(TEXT("The function '%s' is not pure."), *InFunction->GetName()));
+		//}
+
+		if (Function->NumParms != 1)
+		{
+			return MakeError(FString::Printf(TEXT("The function '%s' doesn't have a return value."), *Function->GetName()));
+		}
+
+		const FProperty* ReturnProperty = GetReturnProperty(Function);
+		if (ReturnProperty == nullptr)
+		{
+			return MakeError(FString::Printf(TEXT("The return value for function '%s' is invalid."), *Function->GetName()));
+		}
+
+		return MakeValue(ReturnProperty);
+	}
+
+	TValueOrError<const FProperty*, FString> TryGetPropertyTypeForSourceBinding(const FMVVMConstFieldVariant& InField)
+	{
+		if (InField.IsEmpty())
+		{
+			return MakeError(TEXT("No source was provided for the binding."));
+		}
+
+		if (InField.IsProperty())
+		{
+			return TryGetPropertyTypeForSourceBinding(InField.GetProperty());
+		}
+		else
+		{
+			check(InField.IsFunction());
+			return TryGetPropertyTypeForSourceBinding(InField.GetFunction());
+		}
+	}
+
+	TValueOrError<const FProperty*, FString> TryGetPropertyTypeForDestinationBinding(const FProperty* Property)
+	{
+		FString Result = Private::TryGetPropertyTypeCommon(Property);
+		if (!Result.IsEmpty())
+		{
+			return MakeError(MoveTemp(Result));
+		}
+
+		if (Property->HasAnyPropertyFlags(CPF_BlueprintReadOnly))
+		{
+			return MakeError(FString::Printf(TEXT("The property '%s' is read only."), *Property->GetName()));
+		}
+
+		return MakeValue(Property);
+	}
+
+	TValueOrError<const FProperty*, FString> TryGetPropertyTypeForDestinationBinding(const UFunction* Function)
+	{
+		FString Result = Private::TryGetPropertyTypeCommon(Function);
+		if (!Result.IsEmpty())
+		{
+			return MakeError(MoveTemp(Result));
+		}
+
+		if (Function->NumParms != 1)
+		{
+			return MakeError(FString::Printf(TEXT("The function '%s' has more than one argument."), *Function->GetName()));
+		}
+
+		if (Function->HasAllFunctionFlags(FUNC_Const))
+		{
+			return MakeError(FString::Printf(TEXT("The function '%s' is const."), *Function->GetName()));
+		}
+
+		if (Function->HasAllFunctionFlags(FUNC_BlueprintPure))
+		{
+			return MakeError(FString::Printf(TEXT("The function '%s' is pure."), *Function->GetName()));
+		}
+
+		const FProperty* ReturnProperty = GetReturnProperty(Function);
+		if (ReturnProperty != nullptr)
+		{
+			return MakeError(FString::Printf(TEXT("The function '%s' has a return value."), *Function->GetName()));
+		}
+
+		const FProperty* FirstArgumentProperty = GetFirstArgumentProperty(Function);
+		if (FirstArgumentProperty == nullptr)
+		{
+			return MakeError(FString::Printf(TEXT("The function '%s' doesn't not have a valid (single) argument."), *Function->GetName()));
+		}
+
+		return MakeValue(FirstArgumentProperty);
+	}
+
+	TValueOrError<const FProperty*, FString> TryGetPropertyTypeForDestinationBinding(const FMVVMConstFieldVariant& InField)
+	{
+		if (InField.IsEmpty())
+		{
+			return MakeError(TEXT("No destination was provided for the binding."));
+		}
+
+		if (InField.IsProperty())
+		{
+			return TryGetPropertyTypeForDestinationBinding(InField.GetProperty());
+		}
+		else
+		{
+			check(InField.IsFunction());
+
+			return TryGetPropertyTypeForDestinationBinding(InField.GetFunction());
+		}
+	}
+
+	TValueOrError<FConversionFunctionArguments, FString> TryGetPropertyTypeForConversionFunction(const UFunction* InFunction)
+	{
+		FString CommonResult = Private::TryGetPropertyTypeCommon(InFunction);
+		if (!CommonResult.IsEmpty())
+		{
+			return MakeError(MoveTemp(CommonResult));
+		}
+
+		if (!InFunction->HasAllFunctionFlags(FUNC_Static))
+		{
+			return MakeError(FString::Printf(TEXT("The function '%s' is not static."), *InFunction->GetName()));
+		}
+
+		if (InFunction->NumParms != 2)
+		{
+			return MakeError(FString::Printf(TEXT("The function '%s' does not have the correct number of arguments."), *InFunction->GetName()));
+		}
+
+		FConversionFunctionArguments Result;
+		Result.ReturnProperty = GetReturnProperty(InFunction);
+		if (Result.ReturnProperty == nullptr)
+		{
+			return MakeError(FString::Printf(TEXT("The return value for function '%s' is invalid."), *InFunction->GetName()));
+		}
+
+		Result.ArgumentProperty = GetFirstArgumentProperty(InFunction);
+		if (Result.ArgumentProperty == nullptr)
+		{
+			return MakeError(FString::Printf(TEXT("The function '%s' doesn't not have a valid argument."), *InFunction->GetName()));
+		}
+
+		return MakeValue(Result);
+	}
+
+
+	namespace Private
+	{
+		bool IsObjectPropertyCompatible(const FProperty* Source, const FProperty* Destination)
+		{
+			const FObjectPropertyBase* SourceObjectProperty = CastField<const FObjectPropertyBase>(Source);
+			const FObjectPropertyBase* DestinationObjectProperty = CastField<const FObjectPropertyBase>(Destination);
+			return SourceObjectProperty
+				&& DestinationObjectProperty
+				&& SourceObjectProperty->PropertyClass
+				&& DestinationObjectProperty->PropertyClass
+				&& SourceObjectProperty->PropertyClass->IsChildOf(DestinationObjectProperty->PropertyClass);
+		}
+
+
+		bool IsRealNumberConversion(const FProperty* Source, const FProperty* Destination)
+		{
+			return (CastField<FDoubleProperty>(Source) && CastField<FFloatProperty>(Destination))
+				|| (CastField<FFloatProperty>(Source) && CastField<FDoubleProperty>(Destination));
+		}
+
+		void ConvertRealNumber(const FProperty* Source, const FProperty* Destination, void* Data)
+		{
+			check(Source);
+			check(Destination);
+			check(Data);
+			check(Source->ArrayDim == Destination->ArrayDim);
+
+			if (const FDoubleProperty* SrcDoubleProp = CastField<FDoubleProperty>(Source))
+			{
+				if (const FFloatProperty* DestFloatProp = CastField<FFloatProperty>(Destination))
+				{
+					//double to float
+					for (int32 Idx = 0; Idx < Destination->ArrayDim; ++Idx)
+					{
+						const void* SrcElemValue = static_cast<const uint8*>(Data) + (Source->ElementSize * Idx);
+						void* DestElemValue = static_cast<uint8*>(Data) + (Destination->ElementSize * Idx);
+
+						const double Value = SrcDoubleProp->GetFloatingPointPropertyValue(SrcElemValue);
+						DestFloatProp->SetFloatingPointPropertyValue(DestElemValue, Value);
+					}
+				}
+			}
+			else if (const FFloatProperty* SrcFloatProp = CastField<FFloatProperty>(Source))
+			{
+				if (const FDoubleProperty* DestDoubleProp = CastField<FDoubleProperty>(Destination))
+				{
+					//float to double (backward)
+					for (int32 Idx = Destination->ArrayDim - 1; Idx >= 0; --Idx)
+					{
+						const void* SrcElemValue = static_cast<const uint8*>(Data) + (Source->ElementSize * Idx);
+						void* DestElemValue = static_cast<uint8*>(Data) + (Destination->ElementSize * Idx);
+
+						const double Value = SrcFloatProp->GetFloatingPointPropertyValue(SrcElemValue);
+						DestDoubleProp->SetFloatingPointPropertyValue(DestElemValue, Value);
+					}
+				}
+			}
+		}
+	} // namespace
+
+	bool ArePropertiesCompatible(const FProperty* Source, const FProperty* Destination)
+	{
+		if (Source == nullptr || Destination == nullptr)
+		{
+			return false;
+		}
+
+		return Source->ArrayDim == Destination->ArrayDim
+			&& (Destination->SameType(Source)
+				|| Private::IsRealNumberConversion(Source, Destination)
+				|| Private::IsObjectPropertyCompatible(Source, Destination)
+				);
+	}
+
+
+	const FProperty* GetReturnProperty(const UFunction* InFunction)
+	{
+		check(InFunction);
+		FProperty* Result = InFunction->GetReturnProperty();
+
+		if (Result == nullptr && InFunction->HasAllFunctionFlags(FUNC_HasOutParms))
+		{
+			for (TFieldIterator<FProperty> It(InFunction); It && (It->PropertyFlags & CPF_Parm); ++It)
+			{
+				if (It->HasAllPropertyFlags(CPF_OutParm) && !It->HasAnyPropertyFlags(CPF_Deprecated | CPF_EditorOnly | CPF_ConstParm))
+				{
+					Result = *It;
+					break;
+				}
+			}
+		}
+
+		return Result;
+	}
+
+
+	// We only accept copy argument or const ref
+	const FProperty* GetFirstArgumentProperty(const UFunction* InFunction)
+	{
+		check(InFunction);
+		for (TFieldIterator<FProperty> It(InFunction); It && (It->PropertyFlags & CPF_Parm); ++It)
+		{
+			if (It->HasAllPropertyFlags(CPF_Parm) && !It->HasAnyPropertyFlags(CPF_Deprecated | CPF_EditorOnly | CPF_ReturnParm))
+			{
+				if (It->HasAllPropertyFlags(CPF_OutParm) && !It->HasAllPropertyFlags(CPF_ConstParm))
+				{
+					continue;
+				}
+				return *It;
+			}
+		}
+
+		return nullptr;
+	}
+
+	void ExecuteBinding_NoCheck(const FFieldContext& Source, const FFieldContext& Destination)
+	{
+		check(!Source.GetObjectVariant().IsNull() && !Destination.GetObjectVariant().IsNull());
+		check(!Source.GetFieldVariant().IsEmpty() && !Destination.GetFieldVariant().IsEmpty());
+
+		const bool bIsSourceBindingIsProperty = Source.GetFieldVariant().IsProperty();
+		const FProperty* GetterType = bIsSourceBindingIsProperty ? Source.GetFieldVariant().GetProperty() : GetReturnProperty(Source.GetFieldVariant().GetFunction());
+		check(GetterType);
+
+		const bool bIsDestinationBindingIsProperty = Destination.GetFieldVariant().IsProperty();
+		const FProperty* SetterType = bIsDestinationBindingIsProperty ? Destination.GetFieldVariant().GetProperty() : GetFirstArgumentProperty(Destination.GetFieldVariant().GetFunction());
+		check(SetterType);
+
+		check(ArePropertiesCompatible(GetterType, SetterType));
+
+		const int32 AllocationSize = FMath::Max(GetterType->GetSize(), SetterType->GetSize());
+		const int32 AllocationMinAlignment = FMath::Max(GetterType->GetMinAlignment(), SetterType->GetMinAlignment());
+		void* DataPtr = FMemory_Alloca_Aligned(AllocationSize, AllocationMinAlignment);
+		GetterType->InitializeValue(DataPtr);
+
+		if (bIsSourceBindingIsProperty)
+		{
+			GetterType->GetValue_InContainer(Source.GetObjectVariant().GetData(), DataPtr);
+		}
+		else
+		{
+			check(GetterType->GetSize() == Source.GetFieldVariant().GetFunction()->ParmsSize); // only 1 param
+			check(Source.GetObjectVariant().IsUObject());
+			Source.GetObjectVariant().GetUObject()->ProcessEvent(Source.GetFieldVariant().GetFunction(), DataPtr);
+		}
+
+		Private::ConvertRealNumber(GetterType, SetterType, DataPtr);
+
+		if (bIsDestinationBindingIsProperty)
+		{
+			SetterType->SetValue_InContainer(Destination.GetObjectVariant().GetData(), DataPtr);
+		}
+		else
+		{
+			check(SetterType->GetSize() == Destination.GetFieldVariant().GetFunction()->ParmsSize); // only 1 param
+			check(Destination.GetObjectVariant().IsUObject());
+			Destination.GetObjectVariant().GetUObject()->ProcessEvent(Destination.GetFieldVariant().GetFunction(), DataPtr);
+		}
+
+		GetterType->DestroyValue(DataPtr);
+	}
+
+
+	void ExecuteBinding_NoCheck(const FFieldContext& Source, const FFieldContext& Destination, const FFunctionContext& ConversionFunction)
+	{
+		check(!Source.GetObjectVariant().IsNull() && !Destination.GetObjectVariant().IsNull());
+		check(!Source.GetFieldVariant().IsEmpty() && !Destination.GetFieldVariant().IsEmpty());
+		check(ConversionFunction.GetFunction() && ConversionFunction.GetObject());
+
+		void* ConversionFunctionDataPtr = FMemory_Alloca_Aligned(ConversionFunction.GetFunction()->ParmsSize, ConversionFunction.GetFunction()->GetMinAlignment());
+		for (TFieldIterator<FProperty> It(ConversionFunction.GetFunction()); It && It->HasAnyPropertyFlags(CPF_Parm); ++It)
+		{
+			It->InitializeValue_InContainer(ConversionFunctionDataPtr);
+		}
+
+		// Get the value
+		{
+			const FProperty* ArgumentConversionProperty = GetFirstArgumentProperty(ConversionFunction.GetFunction());
+			check(ArgumentConversionProperty);
+
+			const bool bIsSourceBindingIsProperty = Source.GetFieldVariant().IsProperty();
+			const FProperty* GetterType = bIsSourceBindingIsProperty ? Source.GetFieldVariant().GetProperty() : GetReturnProperty(Source.GetFieldVariant().GetFunction());
+			check(GetterType);
+			check(ArePropertiesCompatible(GetterType, ArgumentConversionProperty));
+
+			if (Private::IsRealNumberConversion(GetterType, ArgumentConversionProperty))
+			{
+				// we need to do a copy because we may destroy the ReturnConversionProperty (imagine the return type is a TArray)
+				const int32 AllocationSize = FMath::Max(GetterType->GetSize(), ArgumentConversionProperty->GetSize());
+				const int32 AllocationMinAlignment = FMath::Max(GetterType->GetMinAlignment(), ArgumentConversionProperty->GetMinAlignment());
+				void* GetterDataPtr = FMemory_Alloca_Aligned(AllocationSize, AllocationMinAlignment);
+				GetterType->InitializeValue(GetterDataPtr); // probably not needed since they are double/float
+
+				if (bIsSourceBindingIsProperty)
+				{
+					GetterType->GetValue_InContainer(Source.GetObjectVariant().GetData(), GetterDataPtr);
+				}
+				else
+				{
+					check(GetterType->GetSize() == Source.GetFieldVariant().GetFunction()->ParmsSize); // only 1 param
+					check(Source.GetObjectVariant().IsUObject());
+					Source.GetObjectVariant().GetUObject()->ProcessEvent(Source.GetFieldVariant().GetFunction(), GetterDataPtr);
+				}
+
+				Private::ConvertRealNumber(GetterType, ArgumentConversionProperty, GetterDataPtr);
+				ArgumentConversionProperty->CopyCompleteValue(ArgumentConversionProperty->ContainerPtrToValuePtr<void>(ConversionFunctionDataPtr), GetterDataPtr);
+
+				GetterType->DestroyValue(GetterDataPtr);
+
+			}
+			else
+			{
+				// Re use the same buffer, no need to create a new copy
+				void* SourceDataPtr = ArgumentConversionProperty->ContainerPtrToValuePtr<void>(ConversionFunctionDataPtr);
+				if (bIsSourceBindingIsProperty)
+				{
+					GetterType->GetValue_InContainer(Source.GetObjectVariant().GetData(), SourceDataPtr);
+				}
+				else
+				{
+					check(GetterType->GetSize() == Source.GetFieldVariant().GetFunction()->ParmsSize); // only 1 param
+					check(Source.GetObjectVariant().IsUObject());
+					Source.GetObjectVariant().GetUObject()->ProcessEvent(Source.GetFieldVariant().GetFunction(), SourceDataPtr);
+				}
+			}
+		}
+
+		ConversionFunction.GetObject()->ProcessEvent(ConversionFunction.GetFunction(), ConversionFunctionDataPtr);
+
+		{
+			const FProperty* ReturnConversionProperty = GetReturnProperty(ConversionFunction.GetFunction());
+			check(ReturnConversionProperty);
+
+			const bool bIsDestinationBindingIsProperty = Destination.GetFieldVariant().IsProperty();
+			const FProperty* SetterType = bIsDestinationBindingIsProperty ? Destination.GetFieldVariant().GetProperty() : GetFirstArgumentProperty(Destination.GetFieldVariant().GetFunction());
+			check(SetterType);
+			check(ArePropertiesCompatible(ReturnConversionProperty, SetterType));
+
+			if (Private::IsRealNumberConversion(ReturnConversionProperty, SetterType))
+			{
+				// we need to do a copy because we may destroy the ReturnConversionProperty (imagine the return type is a TArray)
+				const int32 AllocationSize = FMath::Max(SetterType->GetSize(), ReturnConversionProperty->GetSize());
+				const int32 AllocationMinAlignment = FMath::Max(SetterType->GetMinAlignment(), ReturnConversionProperty->GetMinAlignment());
+				void* SetterDataPtr = FMemory_Alloca_Aligned(AllocationSize, AllocationMinAlignment);
+				SetterType->InitializeValue(SetterDataPtr); // probably not needed since they are double/float
+
+				ReturnConversionProperty->CopyCompleteValue(SetterDataPtr, ReturnConversionProperty->ContainerPtrToValuePtr<void>(ConversionFunctionDataPtr));
+				Private::ConvertRealNumber(ReturnConversionProperty, SetterType, SetterDataPtr);
+
+				if (bIsDestinationBindingIsProperty)
+				{
+					SetterType->SetValue_InContainer(Destination.GetObjectVariant().GetData(), SetterDataPtr);
+				}
+				else
+				{
+					check(SetterType->GetSize() == Destination.GetFieldVariant().GetFunction()->ParmsSize); // only 1 param
+					check(Destination.GetObjectVariant().IsUObject());
+					Destination.GetObjectVariant().GetUObject()->ProcessEvent(Destination.GetFieldVariant().GetFunction(), SetterDataPtr);
+				}
+
+				SetterType->DestroyValue(SetterDataPtr);
+			}
+			else
+			{
+				// Re use the same buffer, no need to create a new copy
+				void* DestinationDataPtr = ReturnConversionProperty->ContainerPtrToValuePtr<void>(ConversionFunctionDataPtr);
+
+				if (bIsDestinationBindingIsProperty)
+				{
+					SetterType->SetValue_InContainer(Destination.GetObjectVariant().GetData(), DestinationDataPtr);
+				}
+				else
+				{
+					check(SetterType->GetSize() == Destination.GetFieldVariant().GetFunction()->ParmsSize); // only 1 param
+					check(Destination.GetObjectVariant().IsUObject());
+					Destination.GetObjectVariant().GetUObject()->ProcessEvent(Destination.GetFieldVariant().GetFunction(), DestinationDataPtr);
+				}
+			}
+		}
+
+		for (TFieldIterator<FProperty> It(ConversionFunction.GetFunction()); It && It->HasAnyPropertyFlags(CPF_Parm); ++It)
+		{
+			It->DestroyValue_InContainer(ConversionFunctionDataPtr);
+		}
+	}
+}
diff --git a/Engine/Plugins/Runtime/ModelViewViewModel/Source/ModelViewViewModel/Private/Bindings/MVVMFieldPathHelper.cpp b/Engine/Plugins/Runtime/ModelViewViewModel/Source/ModelViewViewModel/Private/Bindings/MVVMFieldPathHelper.cpp
new file mode 100644
--- /dev/null
+++ b/Engine/Plugins/Runtime/ModelViewViewModel/Source/ModelViewViewModel/Private/Bindings/MVVMFieldPathHelper.cpp
@@ -0,0 +1,210 @@
+// Copyright Epic Games, Inc. All Rights Reserved.
+
+#include "Bindings/MVVMFieldPathHelper.h"
+
+#include "Bindings/MVVMBindingHelper.h"
+
+namespace UE::MVVM::FieldPathHelper
+{
+
+namespace Private
+{
+
+static const FName NAME_BlueprintGetter = "BlueprintGetter";
+static const FName NAME_BlueprintSetter = "BlueprintSetter";
+
+} // namespace
+
+
+TValueOrError<TArray<FMVVMFieldVariant>, FString> GenerateFieldPathList(TSubclassOf<UObject> InFrom, FStringView InFieldPath, bool bForReading)
+{
+	if (InFrom.Get() == nullptr)
+	{
+		return MakeError(TEXT("The source class is invalid."));
+	}
+	if (InFieldPath.IsEmpty())
+	{
+		return MakeError(TEXT("The FieldPath is empty."));
+	}
+	if (InFieldPath[InFieldPath.Len() - 1] == TEXT('.'))
+	{
+		return MakeError(TEXT("The field path cannot end with a '.' character."));
+	}
+
+	auto FindContainer = [](const FProperty* Property) -> TValueOrError<UStruct*, FString>
+	{
+		const FObjectPropertyBase* ObjectProperty = CastField<const FObjectPropertyBase>(Property);
+		const FStructProperty* StructProperty = CastField<const FStructProperty>(Property);
+
+		if (ObjectProperty)
+		{
+			return MakeValue(ObjectProperty->PropertyClass);
+		}
+		else if (StructProperty)
+		{
+			return MakeValue(StructProperty->Struct);
+		}
+		return MakeError(FString::Printf(TEXT("Field can only be object properties or struct properties. %s is a %s"), *Property->GetName(), *Property->GetClass()->GetName()));
+	};
+
+	auto TransformWithAccessor = [](UStruct* CurrentContainer, FMVVMFieldVariant CurrentField, bool bForReading) -> TValueOrError<FMVVMFieldVariant, FString>
+	{
+#if WITH_EDITORONLY_DATA
+		if (bForReading)
+		{
+			if (!CurrentField.GetProperty()->HasGetter())
+			{
+				const FString& BlueprintGetter = CurrentField.GetProperty()->GetMetaData(Private::NAME_BlueprintGetter);
+				if (!BlueprintGetter.IsEmpty())
+				{
+					FMVVMFieldVariant NewField = BindingHelper::FindFieldByName(CurrentContainer, FMVVMBindingName(*BlueprintGetter));
+					if (NewField.IsFunction())
+					{
+						CurrentField = NewField;
+					}
+					else
+					{
+						return MakeError(FString::Printf(TEXT("The BlueprintGetter %s could not be found on object %s."), *BlueprintGetter, *CurrentContainer->GetName()));
+					}
+				}
+			}
+		}
+		else
+		{
+			if (!CurrentField.GetProperty()->HasSetter())
+			{
+				const FString& BlueprintSetter = CurrentField.GetProperty()->GetMetaData(Private::NAME_BlueprintSetter);
+				if (!BlueprintSetter.IsEmpty())
+				{
+					FMVVMFieldVariant NewField = BindingHelper::FindFieldByName(CurrentContainer, FMVVMBindingName(*BlueprintSetter));
+					if (NewField.IsFunction())
+					{
+						CurrentField = NewField;
+					}
+					else
+					{
+						return MakeError(FString::Printf(TEXT("The BlueprintSetter %s could not be found on object %s."), *BlueprintSetter, *CurrentContainer->GetName()));
+					}
+				}
+			}
+		}
+#endif
+		return MakeValue(CurrentField);
+	};
+
+	TArray<FMVVMFieldVariant> Result;
+	UStruct* CurrentContainer = InFrom.Get();
+
+	// Split the string into property or function names
+	//ie. myvar.myfunction.myvar
+	int32 FoundIndex = INDEX_NONE;
+	while (InFieldPath.FindChar(TEXT('.'), FoundIndex))
+	{
+		FMVVMFieldVariant Field = BindingHelper::FindFieldByName(CurrentContainer, FMVVMBindingName(FName(FoundIndex, InFieldPath.GetData())));
+		if (Field.IsEmpty())
+		{
+			return MakeError(FString::Printf(TEXT("The field %s could not be found on container %s."), *FName(FoundIndex, InFieldPath.GetData()).ToString(), CurrentContainer ? *CurrentContainer->GetName() : TEXT("<none>")));
+		}
+		else if (Field.IsProperty())
+		{
+			TValueOrError<FMVVMFieldVariant, FString> TransformedField = TransformWithAccessor(CurrentContainer, Field, true);
+			if (TransformedField.HasError())
+			{
+				return MakeError(TransformedField.StealError());
+			}
+			Field = TransformedField.StealValue();
+			check(!Field.IsEmpty());
+
+			if (Field.IsProperty())
+			{
+				TValueOrError<UStruct*, FString> FoundContainer = FindContainer(Field.GetProperty());
+				if (FoundContainer.HasError())
+				{
+					return MakeError(FoundContainer.StealError());
+				}
+				CurrentContainer = FoundContainer.GetValue();
+			}
+		}
+		
+		if (Field.IsFunction())
+		{
+			const FProperty* ReturnProperty = BindingHelper::GetReturnProperty(Field.GetFunction());
+			TValueOrError<UStruct*, FString> FoundContainer = FindContainer(ReturnProperty);
+			if (FoundContainer.HasError())
+			{
+				return MakeError(FoundContainer.StealError());
+			}
+			CurrentContainer = FoundContainer.GetValue();
+		}
+
+		InFieldPath = InFieldPath.RightChop(FoundIndex + 1);
+		Result.Add(Field);
+	}
+
+	// The last field can be anything (that is what we are going to bind to)
+	if (InFieldPath.Len() > 0)
+	{
+		FMVVMFieldVariant Field = BindingHelper::FindFieldByName(CurrentContainer, FMVVMBindingName(InFieldPath.GetData()));
+		if (Field.IsEmpty())
+		{
+			return MakeError(FString::Printf(TEXT("The field %s could not be found on container %s."), InFieldPath.GetData(), CurrentContainer ? *CurrentContainer->GetName() : TEXT("<none>")));
+		}
+
+		if (Field.IsProperty())
+		{
+			TValueOrError<FMVVMFieldVariant, FString> TransformedField = TransformWithAccessor(CurrentContainer, Field, bForReading);
+			if (TransformedField.HasError())
+			{
+				return MakeError(TransformedField.StealError());
+			}
+			Field = TransformedField.StealValue();
+		}
+		check(!Field.IsEmpty());
+		Result.Add(Field);
+	}
+
+	return MakeValue(MoveTemp(Result));
+}
+
+
+TValueOrError<TArray<FMVVMFieldVariant>, FString> GenerateConversionFunctionFieldPathList(TSubclassOf<UObject> From, FStringView FieldPath)
+{
+	return MakeError(TEXT("Not implemented yet"));
+}
+
+
+TValueOrError<UObject*, void> EvaluateObjectProperty(const FFieldContext& InSource)
+{
+	if (InSource.GetObjectVariant().IsNull())
+	{
+		return MakeError();
+	}
+
+	const bool bIsProperty = InSource.GetFieldVariant().IsProperty();
+	const FProperty* GetterType = bIsProperty ? InSource.GetFieldVariant().GetProperty() : BindingHelper::GetReturnProperty(InSource.GetFieldVariant().GetFunction());
+	check(GetterType);
+
+	const FObjectPropertyBase* ObjectProperty = CastField<const FObjectPropertyBase>(GetterType);
+	if (ObjectProperty == nullptr)
+	{
+		return MakeError();
+	}
+
+	if (bIsProperty)
+	{
+		return MakeValue(ObjectProperty->GetObjectPropertyValue_InContainer(InSource.GetObjectVariant().GetData()));
+	}
+	else
+	{
+		check(InSource.GetObjectVariant().IsUObject());
+		UFunction* Function = InSource.GetFieldVariant().GetFunction();
+		void* DataPtr = FMemory_Alloca_Aligned(Function->ParmsSize, Function->GetMinAlignment());
+		ObjectProperty->InitializeValue(DataPtr);
+		InSource.GetObjectVariant().GetUObject()->ProcessEvent(Function, DataPtr);
+		UObject* Result = ObjectProperty->GetObjectPropertyValue_InContainer(DataPtr);
+		ObjectProperty->DestroyValue(DataPtr);
+		return MakeValue(Result);
+	}
+}
+
+} // namespace
diff --git a/Engine/Plugins/Runtime/ModelViewViewModel/Source/ModelViewViewModel/Private/MVVMSubsystem.cpp b/Engine/Plugins/Runtime/ModelViewViewModel/Source/ModelViewViewModel/Private/MVVMSubsystem.cpp
new file mode 100644
--- /dev/null
+++ b/Engine/Plugins/Runtime/ModelViewViewModel/Source/ModelViewViewModel/Private/MVVMSubsystem.cpp
@@ -0,0 +1,301 @@
+// Copyright Epic Games, Inc. All Rights Reserved.
+
+#include "MVVMSubsystem.h"
+
+#include "Bindings/MVVMBindingHelper.h"
+#include "Blueprint/WidgetTree.h"
+#include "View/MVVMView.h"
+#include "MVVMViewModelBase.h"
+#include "Types/MVVMViewModelCollection.h"
+
+
+void UMVVMSubsystem::Initialize(FSubsystemCollectionBase& Collection)
+{
+	Super::Initialize(Collection);
+	GlobalViewModelCollection = NewObject<UMVVMViewModelCollectionObject>(this);
+}
+
+
+void UMVVMSubsystem::Deinitialize()
+{
+	GlobalViewModelCollection = nullptr;
+	Super::Deinitialize();
+}
+
+
+UMVVMView* UMVVMSubsystem::GetViewFromUserWidget(const UUserWidget* UserWidget) const
+{
+	return UserWidget ? UserWidget->GetExtension<UMVVMView>() : nullptr;
+}
+
+
+bool UMVVMSubsystem::IsViewModelValueValidForSourceBinding(const UMVVMViewModelBase* ViewModel, FMVVMBindingName ViewModelPropertyOrFunctionName) const
+{
+	if (ViewModel)
+	{
+		UE::MVVM::FMVVMFieldVariant Binding = UE::MVVM::BindingHelper::FindFieldByName(ViewModel->GetClass(), ViewModelPropertyOrFunctionName);
+		if (Binding.IsFunction())
+		{
+			return UE::MVVM::BindingHelper::IsValidForSourceBinding(Binding.GetFunction());
+		}
+		if (Binding.IsProperty())
+		{
+			return UE::MVVM::BindingHelper::IsValidForSourceBinding(Binding.GetProperty());
+		}
+	}
+	return false;
+}
+
+
+bool UMVVMSubsystem::IsViewModelValueValidForDestinationBinding(const UMVVMViewModelBase* ViewModel, FMVVMBindingName ViewModelPropertyOrFunctionName) const
+{
+	if (ViewModel)
+	{
+		UE::MVVM::FMVVMFieldVariant Binding = UE::MVVM::BindingHelper::FindFieldByName(ViewModel->GetClass(), ViewModelPropertyOrFunctionName);
+		if (Binding.IsFunction())
+		{
+			return UE::MVVM::BindingHelper::IsValidForDestinationBinding(Binding.GetFunction());
+		}
+		if (Binding.IsProperty())
+		{
+			return UE::MVVM::BindingHelper::IsValidForDestinationBinding(Binding.GetProperty());
+		}
+	}
+	return false;
+}
+
+
+bool UMVVMSubsystem::IsViewValueValidForSourceBinding(const UWidget* View, FMVVMBindingName ViewPropertyOrFunctionName) const
+{
+	if (View)
+	{
+		UE::MVVM::FMVVMFieldVariant Binding = UE::MVVM::BindingHelper::FindFieldByName(View->GetClass(), ViewPropertyOrFunctionName);
+		if (Binding.IsFunction())
+		{
+			return UE::MVVM::BindingHelper::IsValidForSourceBinding(Binding.GetFunction());
+		}
+		if (Binding.IsProperty())
+		{
+			return UE::MVVM::BindingHelper::IsValidForSourceBinding(Binding.GetProperty());
+		}
+	}
+	return false;
+}
+
+
+bool UMVVMSubsystem::IsViewValueValidForDestinationBinding(const UWidget* View, FMVVMBindingName ViewPropertyOrFunctionName) const
+{
+	if (View)
+	{
+		UE::MVVM::FMVVMFieldVariant Binding = UE::MVVM::BindingHelper::FindFieldByName(View->GetClass(), ViewPropertyOrFunctionName);
+		if (Binding.IsFunction())
+		{
+			return UE::MVVM::BindingHelper::IsValidForDestinationBinding(Binding.GetFunction());
+		}
+		if (Binding.IsProperty())
+		{
+			return UE::MVVM::BindingHelper::IsValidForDestinationBinding(Binding.GetProperty());
+		}
+	}
+	return false;
+}
+
+
+bool UMVVMSubsystem::DoesWidgetTreeContainedWidget(const UWidgetTree* WidgetTree, const UWidget* ViewWidget) const
+{
+	// Test if the View's Widget is valid.
+	if (WidgetTree != nullptr && ViewWidget != nullptr)
+	{
+		TArray<UWidget*> Widgets;
+		WidgetTree->GetAllWidgets(Widgets);
+		return Widgets.Contains(ViewWidget);
+	}
+	return false;
+}
+
+
+namespace UE::MVVM::Private
+{
+	TArray<FMVVMAvailableBinding> GetAvailableBindings(const UClass* Class, const UE::FieldNotification::IClassDescriptor& ClassDescriptor)
+	{
+		TSet<FName> FieldDescriptors;
+		ClassDescriptor.ForEachField(Class, [&FieldDescriptors](UE::FieldNotification::FFieldId FieldId)
+			{
+				FieldDescriptors.Add(FieldId.GetName());
+				return true;
+			});
+
+		TArray<FMVVMAvailableBinding> Result;
+		Result.Reserve(FieldDescriptors.Num());
+
+		for (TFieldIterator<const UFunction> FunctionItt(Class, EFieldIteratorFlags::IncludeSuper); FunctionItt; ++FunctionItt)
+		{
+			const UFunction* Function = *FunctionItt;
+			check(Function);
+			const bool bIsReadable = FieldDescriptors.Contains(Function->GetFName()) && BindingHelper::IsValidForSourceBinding(Function);
+			const bool bIsWritable = BindingHelper::IsValidForDestinationBinding(Function);
+			if (bIsReadable || bIsWritable)
+			{
+				Result.Add(FMVVMAvailableBinding(FMVVMBindingName(Function->GetFName()), bIsReadable, bIsWritable));
+			}
+		}
+
+#if WITH_EDITOR
+		FName NAME_BlueprintGetter = "BlueprintGetter";
+		FName NAME_BlueprintSetter = "BlueprintSetter";
+#endif
+
+		for (TFieldIterator<const FProperty> PropertyIt(Class, EFieldIteratorFlags::IncludeSuper); PropertyIt; ++PropertyIt)
+		{
+			const FProperty* Property = *PropertyIt;
+			const bool bIsReadable = FieldDescriptors.Contains(Property->GetFName()) && BindingHelper::IsValidForSourceBinding(Property);
+			const bool bIsWritable = BindingHelper::IsValidForDestinationBinding(Property);
+			if (bIsReadable || bIsWritable)
+			{
+#if WITH_EDITOR
+				// Remove the accessor and replace it by the property
+				if (bIsReadable)
+				{
+					const FString& PropertyGetter = Property->GetMetaData(NAME_BlueprintGetter);
+					if (!PropertyGetter.IsEmpty())
+					{
+						FMVVMBindingName BindingName = FMVVMBindingName(*PropertyGetter);
+						int32 BindingIndex = Result.IndexOfByPredicate([BindingName](const FMVVMAvailableBinding& Binding) { return Binding.GetBindingName() == BindingName && Binding.IsReadable() && !Binding.IsWritable(); });
+						if (BindingIndex != INDEX_NONE)
+						{
+							Result.RemoveAtSwap(BindingIndex);
+						}
+					}
+				}
+				if (bIsWritable)
+				{
+					const FString& PropertySetter = Property->GetMetaData(NAME_BlueprintSetter);
+					if (!PropertySetter.IsEmpty())
+					{
+						FMVVMBindingName BindingName = FMVVMBindingName(*PropertySetter);
+						int32 BindingIndex = Result.IndexOfByPredicate([BindingName](const FMVVMAvailableBinding& Binding) { return Binding.GetBindingName() == BindingName && Binding.IsWritable() && !Binding.IsReadable(); });
+						if (BindingIndex != INDEX_NONE)
+						{
+							Result.RemoveAtSwap(BindingIndex);
+						}
+					}
+				}
+#endif
+				Result.Add(FMVVMAvailableBinding(FMVVMBindingName(Property->GetFName()), bIsReadable, bIsWritable));
+			}
+		}
+
+		return Result;
+	}
+
+	template<typename TClassType>
+	TArray<FMVVMAvailableBinding> GetAvailableBindings(TSubclassOf<TClassType> InSubClass)
+	{
+		if (InSubClass.Get() == nullptr)
+		{
+			return TArray<FMVVMAvailableBinding>();
+		}
+		const UE::FieldNotification::IClassDescriptor& ClassDescriptor = InSubClass.GetDefaultObject()->GetFieldNotificationDescriptor();
+		return GetAvailableBindings(InSubClass.Get(), ClassDescriptor);
+	}
+} //namespace
+
+
+TArray<FMVVMAvailableBinding> UMVVMSubsystem::GetViewModelAvailableBindings(TSubclassOf<UMVVMViewModelBase> ViewModelClass) const
+{
+	return UE::MVVM::Private::GetAvailableBindings(ViewModelClass);
+}
+
+
+TArray<FMVVMAvailableBinding> UMVVMSubsystem::GetWidgetAvailableBindings(TSubclassOf<UWidget> WidgetClass) const
+{
+	return UE::MVVM::Private::GetAvailableBindings(WidgetClass);
+}
+
+
+TValueOrError<bool, FString> UMVVMSubsystem::IsBindingValid(FBindingArgs Args) const
+{
+	if (UE::MVVM::IsForwardBinding(Args.Mode))
+	{
+		TValueOrError<bool, FString> ForwardDirection = IsBindingValid(Args.ForwardArgs);
+		if (ForwardDirection.HasError())
+		{
+			return ForwardDirection;
+		}
+	}
+
+	if (UE::MVVM::IsBackwardBinding(Args.Mode))
+	{
+		TValueOrError<bool, FString> BackwardDirection = IsBindingValid(Args.BackwardArgs);
+		if (BackwardDirection.HasError())
+		{
+			return BackwardDirection;
+		}
+	}
+
+	return MakeValue(true);
+}
+
+TValueOrError<bool, FString> UMVVMSubsystem::IsBindingValid(FDirectionalBindingArgs Args) const
+{
+	return IsBindingValid(Args.ToConst());
+}
+
+TValueOrError<bool, FString> UMVVMSubsystem::IsBindingValid(FConstDirectionalBindingArgs Args) const
+{
+	// Test Source
+	TValueOrError<const FProperty*, FString> GetterProperty = UE::MVVM::BindingHelper::TryGetPropertyTypeForSourceBinding(Args.SourceBinding);
+	if (GetterProperty.HasError())
+	{
+		return MakeError(GetterProperty.StealError());
+	}
+
+	const FProperty* SourceProperty = GetterProperty.StealValue();
+	if (SourceProperty == nullptr)
+	{
+		return MakeError(TEXT("There is no value to bind at the source."));
+	}
+
+	// Test Destination
+	TValueOrError<const FProperty*, FString> SetterProperty = UE::MVVM::BindingHelper::TryGetPropertyTypeForDestinationBinding(Args.DestinationBinding);
+	if (SetterProperty.HasError())
+	{
+		return MakeError(SetterProperty.StealError());
+	}
+
+	const FProperty* DestinationProperty = SetterProperty.StealValue();
+	if (DestinationProperty == nullptr)
+	{
+		return MakeError(TEXT("There is no value to bind at the destination."));
+	}
+
+	// Test the conversion function
+	if (Args.ConversionFunction)
+	{
+		TValueOrError<UE::MVVM::BindingHelper::FConversionFunctionArguments, FString> ConversionResult = UE::MVVM::BindingHelper::TryGetPropertyTypeForConversionFunction(Args.ConversionFunction);
+		if (ConversionResult.HasError())
+		{
+			return MakeError(ConversionResult.StealError());
+		}
+
+		const FProperty* ConversionReturnProperty = ConversionResult.GetValue().ReturnProperty;
+		const FProperty* ConversionArgProperty = ConversionResult.GetValue().ArgumentProperty;
+
+		// The compiled version should look like Setter(Conversion(Getter())).
+		if (!UE::MVVM::BindingHelper::ArePropertiesCompatible(ConversionReturnProperty, DestinationProperty))
+		{
+			return MakeError(FString::Printf(TEXT("The Destination '%s' property type does not match the type of the return type from the conversion function '%s'."), *DestinationProperty->GetName(), *ConversionReturnProperty->GetName()));
+		}
+		if (!UE::MVVM::BindingHelper::ArePropertiesCompatible(SourceProperty, ConversionArgProperty))
+		{
+			return MakeError(FString::Printf(TEXT("The Source '%s' property type does not match the type of the argument type of the conversion function '%s'."), *SourceProperty->GetName(), *DestinationProperty->GetName()));
+		}
+	}
+	else if (!UE::MVVM::BindingHelper::ArePropertiesCompatible(SourceProperty, DestinationProperty))
+	{
+		return MakeError(FString::Printf(TEXT("The Source '%s' property type does not match the type of the Destination property '%s'. A conversion function is required."), *SourceProperty->GetName(), *DestinationProperty->GetName()));
+	}
+
+	return MakeValue(true);
+}
+
diff --git a/Engine/Plugins/Runtime/ModelViewViewModel/Source/ModelViewViewModel/Private/MVVMViewModelBase.cpp b/Engine/Plugins/Runtime/ModelViewViewModel/Source/ModelViewViewModel/Private/MVVMViewModelBase.cpp
new file mode 100644
--- /dev/null
+++ b/Engine/Plugins/Runtime/ModelViewViewModel/Source/ModelViewViewModel/Private/MVVMViewModelBase.cpp
@@ -0,0 +1,84 @@
+// Copyright Epic Games, Inc. All Rights Reserved.
+
+#include "MVVMViewModelBase.h"
+
+#include "Bindings/MVVMBindingHelper.h"
+
+
+FDelegateHandle UMVVMViewModelBase::AddFieldValueChangedDelegate(UE::FieldNotification::FFieldId InFieldId, FFieldValueChangedDelegate InNewDelegate)
+{
+	FDelegateHandle Result;
+	if (InFieldId.IsValid())
+	{
+		Result = Delegates.Add(this, InFieldId, MoveTemp(InNewDelegate));
+		if (Result.IsValid())
+		{
+			EnabledFieldNotifications.PadToNum(InFieldId.GetIndex() + 1, false);
+			EnabledFieldNotifications[InFieldId.GetIndex()] = true;
+		}
+	}
+	return Result;
+}
+
+
+bool UMVVMViewModelBase::RemoveFieldValueChangedDelegate(UE::FieldNotification::FFieldId InFieldId, FDelegateHandle InHandle)
+{
+	bool bResult = false;
+	if (InFieldId.IsValid() && InHandle.IsValid()  && EnabledFieldNotifications.IsValidIndex(InFieldId.GetIndex()) && EnabledFieldNotifications[InFieldId.GetIndex()])
+	{
+		UE::FieldNotification::FFieldMulticastDelegate::FRemoveFromResult RemoveResult = Delegates.RemoveFrom(this, InFieldId, InHandle);
+		bResult = RemoveResult.bRemoved;
+		EnabledFieldNotifications[InFieldId.GetIndex()] = RemoveResult.bHasOtherBoundDelegates;
+	}
+	return bResult;
+}
+
+
+int32 UMVVMViewModelBase::RemoveAllFieldValueChangedDelegates(const void* InUserObject)
+{
+	int32 bResult = 0;
+	if (InUserObject)
+	{
+		UE::FieldNotification::FFieldMulticastDelegate::FRemoveAllResult RemoveResult = Delegates.RemoveAll(this, InUserObject);
+		bResult = RemoveResult.RemoveCount;
+		EnabledFieldNotifications = RemoveResult.HasFields;
+	}
+	return bResult;
+}
+
+
+int32 UMVVMViewModelBase::RemoveAllFieldValueChangedDelegates(UE::FieldNotification::FFieldId InFieldId, const void* InUserObject)
+{
+	int32 bResult = 0;
+	if (InUserObject)
+	{
+		UE::FieldNotification::FFieldMulticastDelegate::FRemoveAllResult RemoveResult = Delegates.RemoveAll(this, InFieldId, InUserObject);
+		bResult = RemoveResult.RemoveCount;
+		EnabledFieldNotifications = RemoveResult.HasFields;
+	}
+	return bResult;
+}
+
+
+const UE::FieldNotification::IClassDescriptor& UMVVMViewModelBase::GetFieldNotificationDescriptor() const
+{
+	static FFieldNotificationClassDescriptor Local;
+	return Local;
+}
+
+
+void UMVVMViewModelBase::BindingFieldValueChanged(UE::FieldNotification::FFieldId InFieldId)
+{
+	if (InFieldId.IsValid() && EnabledFieldNotifications.IsValidIndex(InFieldId.GetIndex()) && EnabledFieldNotifications[InFieldId.GetIndex()])
+	{
+		Delegates.Broadcast(this, InFieldId);
+	}
+}
+
+
+void UMVVMViewModelBase::BindingValueChanged(FMVVMBindingName ViewModelPropertyName)
+{
+	UE::FieldNotification::FFieldId FieldId = GetFieldNotificationDescriptor().GetField(GetClass(), ViewModelPropertyName.ToName());
+	BindingFieldValueChanged(FieldId);
+}
+
diff --git a/Engine/Plugins/Runtime/ModelViewViewModel/Source/ModelViewViewModel/Private/ModelViewViewModelModule.h b/Engine/Plugins/Runtime/ModelViewViewModel/Source/ModelViewViewModel/Private/ModelViewViewModelModule.h
new file mode 100644
--- /dev/null
+++ b/Engine/Plugins/Runtime/ModelViewViewModel/Source/ModelViewViewModel/Private/ModelViewViewModelModule.h
@@ -0,0 +1,21 @@
+// Copyright Epic Games, Inc. All Rights Reserved.
+
+#pragma once
+
+#include "CoreMinimal.h"
+#include "Modules/ModuleInterface.h"
+#include "Modules/ModuleManager.h"
+
+/**
+ *
+ */
+class FModelViewViewModelModule : public IModuleInterface
+{
+public:
+	FModelViewViewModelModule() = default;
+
+	//~ Begin IModuleInterface interface
+	// virtual void StartupModule() override;
+	// virtual void ShutdownModule() override;
+	//~ End IModuleInterface interface
+};
diff --git a/Engine/Plugins/Runtime/ModelViewViewModel/Source/ModelViewViewModel/Private/View/MVVMView.cpp b/Engine/Plugins/Runtime/ModelViewViewModel/Source/ModelViewViewModel/Private/View/MVVMView.cpp
new file mode 100644
--- /dev/null
+++ b/Engine/Plugins/Runtime/ModelViewViewModel/Source/ModelViewViewModel/Private/View/MVVMView.cpp
@@ -0,0 +1,446 @@
+// Copyright Epic Games, Inc. All Rights Reserved.
+
+#include "View/MVVMView.h"
+#include "View/MVVMViewClass.h"
+
+#include "Blueprint/UserWidget.h"
+#include "Blueprint/WidgetTree.h"
+#include "MVVMMessageLog.h"
+#include "MVVMViewModelBase.h"
+#include "MVVMSubsystem.h"
+
+
+#define LOCTEXT_NAMESPACE "MVVMView"
+
+///////////////////////////////////////////////////////////////////////
+// 
+///////////////////////////////////////////////////////////////////////
+
+void UMVVMView::ConstructView(const UMVVMViewClass* InClassExtension)
+{
+	ensure(ClassExtension == nullptr);
+	ClassExtension = InClassExtension;
+}
+
+
+void UMVVMView::Initialize()
+{
+	check(ClassExtension);
+
+	// Init ViewModel instances
+	for (const FMVVMViewClass_SourceCreator& Item : ClassExtension->GetViewModelCreators())
+	{
+		Item.CreateInstance(ClassExtension, this, GetUserWidget());
+	}
+
+	const TArrayView<const FMVVMViewClass_CompiledBinding> CompiledBindings = ClassExtension->GetCompiledBindings();
+	for (int32 Index = 0; Index < CompiledBindings.Num(); ++Index)
+	{
+		const FMVVMViewClass_CompiledBinding& Binding = CompiledBindings[Index];
+		if (Binding.IsEnabledByDefault())
+		{
+			EnableLibraryBinding(Binding, Index);
+		}
+	}
+}
+
+
+void UMVVMView::Destruct()
+{
+	for (FRegisteredSource& Source : AllSources)
+	{
+		UObject* SourceObject = Source.Source.Get();
+		if (Source.Count > 0 && SourceObject)
+		{
+			TScriptInterface<INotifyFieldValueChanged> SourceAsInterface = SourceObject;
+			checkf(SourceAsInterface.GetInterface(), TEXT("It was added as a INotifyFieldValueChanged. It should still be."));
+			SourceAsInterface->RemoveAllFieldValueChangedDelegates(this);
+		}
+	}
+	AllSources.Reset();
+	EnabledLibraryBindings.Reset();
+}
+
+
+void UMVVMView::SetViewModel(FName ViewModelName, UMVVMViewModelBase* NewValue)
+{
+	if (!ViewModelName.IsNone() && ClassExtension != nullptr)
+	{
+		FObjectPropertyBase* FoundObjectProperty = FindFProperty<FObjectPropertyBase>(GetUserWidget()->GetClass(), ViewModelName);
+		if (FoundObjectProperty == nullptr)
+		{
+			UE::MVVM::FMessageLog Log(GetUserWidget());
+			Log.Error(LOCTEXT("SetViewModelInvalidName", "Invalid viewmodel name for SetViewModel."));
+			return;
+		}
+
+		{
+			bool bFound = false;
+			for (const FMVVMViewClass_SourceCreator& Item : ClassExtension->GetViewModelCreators())
+			{
+				if (Item.GetSourcePropertyName() == ViewModelName)
+				{
+					bFound = true;
+					if (NewValue && !NewValue->GetClass()->IsChildOf(Item.GetSourceClass()))
+					{
+						UE::MVVM::FMessageLog Log(GetUserWidget());
+						Log.Error(LOCTEXT("SetViewModelInvalidValueType", "The new viewmodel is not of the expected type."));
+						return;
+					}
+					break;
+				}
+			}
+
+			if (!bFound)
+			{
+				UE::MVVM::FMessageLog Log(GetUserWidget());
+				Log.Error(LOCTEXT("SetViewModelViewModelNameNotFound", "The viewmodel name could not be found."));
+				return;
+			}
+		}
+
+
+		UObject* PreviousValue = FoundObjectProperty->GetObjectPropertyValue_InContainer(GetUserWidget());
+		if (PreviousValue != NewValue)
+		{
+			FMemMark Mark(FMemStack::Get());
+			TArray<int32, TMemStackAllocator<>> BindingToReenabled;
+			const TArrayView<const FMVVMViewClass_CompiledBinding> CompiledBindings = ClassExtension->GetCompiledBindings();
+
+			// Unregister any from that source
+			if (PreviousValue)
+			{
+				BindingToReenabled.Reserve(CompiledBindings.Num());
+
+				for (int32 Index = 0; Index < CompiledBindings.Num(); ++Index)
+				{
+					const FMVVMViewClass_CompiledBinding& Binding = CompiledBindings[Index];
+					if (Binding.GetSourceObjectPropertyName() == ViewModelName && IsLibraryBindingEnabled(Index))
+					{
+						DisableLibraryBinding(Binding, Index);
+						BindingToReenabled.Add(Index);
+					}
+				}
+			}
+
+			FoundObjectProperty->SetObjectPropertyValue_InContainer(GetUserWidget(), NewValue);
+
+			if (NewValue)
+			{
+				// Register back any binding that was previously enabled
+				if (PreviousValue)
+				{
+					for (int32 Index = 0; Index < BindingToReenabled.Num(); ++Index)
+					{
+						const FMVVMViewClass_CompiledBinding& Binding = CompiledBindings[Index];
+						EnableLibraryBinding(Binding, Index);
+					}
+				}
+				else
+				{
+					// Enabled the default bindings
+					for (int32 Index = 0; Index < CompiledBindings.Num(); ++Index)
+					{
+						const FMVVMViewClass_CompiledBinding& Binding = CompiledBindings[Index];
+						if (Binding.IsEnabledByDefault() && Binding.GetSourceObjectPropertyName() == ViewModelName)
+						{
+							EnableLibraryBinding(Binding, Index);
+						}
+					}
+				}
+			}
+		}
+	}
+}
+
+
+namespace UE::MVVM::Private
+{
+	using FRecursiveDetctionElement = TTuple<UObject*, int32>;
+	TArray<FRecursiveDetctionElement> RecursiveDetector;
+}
+
+void UMVVMView::HandledLibraryBindingValueChanged(UObject* InViewModelOrWidget, UE::FieldNotification::FFieldId InFieldId, int32 InCompiledBindingIndex) const
+{
+	check(InViewModelOrWidget);
+	check(InFieldId.IsValid());
+
+	if (ensure(ClassExtension))
+	{
+		checkf(ClassExtension->GetCompiledBindings().IsValidIndex(InCompiledBindingIndex), TEXT("The binding at index '%d' does not exist. The binding was probably not cleared on destroyed."), InCompiledBindingIndex);
+		const FMVVMViewClass_CompiledBinding& Binding = ClassExtension->GetCompiledBinding(InCompiledBindingIndex);
+
+		// Test for recursivity
+		if (UE::MVVM::Private::RecursiveDetector.FindByPredicate([InViewModelOrWidget, InCompiledBindingIndex](const UE::MVVM::Private::FRecursiveDetctionElement& Element)
+			{
+				return Element.Get<0>() == InViewModelOrWidget && Element.Get<1>() == InCompiledBindingIndex;
+			}) != nullptr)
+		{
+			ensureAlwaysMsgf(false, TEXT("Recursive binding detected"));
+			//Todo add more infos. Callstack maybe? Log the chain?
+			UE::MVVM::FMessageLog Log(GetUserWidget());
+			Log.Error(LOCTEXT("RecursionDetected", "A recusion binding was detected (ie. A->B->C->A->B->C) at runtime."));
+			return;
+		}
+
+		{
+			UE::MVVM::Private::RecursiveDetector.Emplace(InViewModelOrWidget, InCompiledBindingIndex);
+
+			ExecuteLibraryBinding(Binding, InViewModelOrWidget);
+
+			UE::MVVM::Private::RecursiveDetector.Pop();
+		}
+	}
+}
+
+
+void UMVVMView::ExecuteLibraryBinding(const FMVVMViewClass_CompiledBinding& Binding) const
+{
+	check(ClassExtension);
+	check(GetUserWidget());
+
+	TValueOrError<void, FMVVMCompiledBindingLibrary::EExecutionFailingReason> ExecutionResult = ClassExtension->GetBindingLibrary().Execute(GetUserWidget(), Binding.GetBinding());
+	if (ExecutionResult.HasError())
+	{
+		UE::MVVM::FMessageLog Log(GetUserWidget());
+		Log.Error(FText::Format(LOCTEXT("ExecuteBindingGenericExecute", "The binding '{0}' was not executed. {1}."), FText::FromString(Binding.ToString()), LOCTEXT("todo", "Reason is todo.")));
+	}
+	else if (bLogBinding)
+	{
+		UE::MVVM::FMessageLog Log(GetUserWidget());
+		Log.Info(FText::Format(LOCTEXT("ExecuteLibraryBinding", "Execute binding '{0}'."), FText::FromString(Binding.ToString())));
+	}
+}
+
+
+void UMVVMView::ExecuteLibraryBinding(const FMVVMViewClass_CompiledBinding& Binding, UObject* Source) const
+{
+	check(ClassExtension);
+	check(GetUserWidget());
+	check(Source);
+
+	TValueOrError<void, FMVVMCompiledBindingLibrary::EExecutionFailingReason> ExecutionResult = ClassExtension->GetBindingLibrary().ExecuteWithSource(GetUserWidget(), Binding.GetBinding(), Source);
+	if (ExecutionResult.HasError())
+	{
+		UE::MVVM::FMessageLog Log(GetUserWidget());
+		Log.Error(FText::Format(LOCTEXT("ExecuteBindingGenericExecute", "The binding '{0}' was not executed. {1}."), FText::FromString(Binding.ToString()), LOCTEXT("todo", "Reason is todo.")));
+	}
+	else if (bLogBinding)
+	{
+		UE::MVVM::FMessageLog Log(GetUserWidget());
+		Log.Info(FText::Format(LOCTEXT("ExecuteLibraryBinding", "Execute binding '{0}'."), FText::FromString(Binding.ToString())));
+	}
+}
+
+
+//void UMVVMView::SetLibraryBindingEnabled(FGuid ViewModelId, FMVVMBindingName BindingName, bool bEnable)
+//{
+//	if (ensure(ClassExtension))
+//	{
+//		int32 LibraryBindingIndex = ClassExtension->IndexOfCompiledBinding(ViewModelId, BindingName);
+//		if (LibraryBindingIndex != INDEX_NONE)
+//		{
+//			const bool bIsEnabled = IsLibraryBindingEnabled(LibraryBindingIndex);
+//			if (bIsEnabled != bEnable)
+//			{
+//				const FMVVMViewClass_CompiledBinding& Item = ClassExtension->GetCompiledBindings(LibraryBindingIndex);
+//				if (bEnable)
+//				{
+//					EnableLibraryBinding(Item, LibraryBindingIndex);
+//				}
+//				else
+//				{
+//					DisableLibraryBinding(Item, LibraryBindingIndex);
+//				}
+//			}
+//		}
+//	}
+//}
+
+
+//bool UMVVMView::IsLibraryBindingEnabled(FGuid ViewModelId, FMVVMBindingName BindingName) const
+//{
+//	if (ensure(ClassExtension))
+//	{
+//		int32 LibraryBindingIndex = ClassExtension->IndexOfCompiledBinding(ViewModelId, BindingName);
+//		return IsLibraryBindingEnabled(LibraryBindingIndex);
+//	}
+//	return false;
+//}
+
+
+bool UMVVMView::IsLibraryBindingEnabled(int32 InBindindIndex) const
+{
+	return EnabledLibraryBindings.IsValidIndex(InBindindIndex) && EnabledLibraryBindings[InBindindIndex];
+}
+
+
+void UMVVMView::EnableLibraryBinding(const FMVVMViewClass_CompiledBinding& Binding, int32 BindingIndex)
+{
+	EnabledLibraryBindings.PadToNum(BindingIndex + 1, false);
+	check(EnabledLibraryBindings[BindingIndex] == false);
+
+	bool bRegistered = true;
+	if (!Binding.IsOneTime())
+	{
+		bRegistered = RegisterLibraryBinding(Binding, BindingIndex);
+	}
+
+	if (bRegistered && bLogBinding)
+	{
+		UE::MVVM::FMessageLog Log(GetUserWidget());
+		Log.Info(FText::Format(LOCTEXT("EnableLibraryBinding", "Enable binding '{0}'."), FText::FromString(Binding.ToString())));
+	}
+
+	EnabledLibraryBindings[BindingIndex] = bRegistered;
+
+	if (bRegistered && Binding.NeedsExecutionAtInitialization())
+	{
+		ExecuteLibraryBinding(Binding);
+	}
+}
+
+
+void UMVVMView::DisableLibraryBinding(const FMVVMViewClass_CompiledBinding& Binding, int32 BindingIndex)
+{
+	check(IsLibraryBindingEnabled(BindingIndex));
+
+	EnabledLibraryBindings[BindingIndex] = false;
+	UnregisterLibraryBinding(Binding);
+
+	if (bLogBinding)
+	{
+		UE::MVVM::FMessageLog Log(GetUserWidget());
+		Log.Info(FText::Format(LOCTEXT("DisableLibraryBinding", "Disable binding '{0}'."), FText::FromString(Binding.ToString())));
+	}
+}
+
+
+bool UMVVMView::RegisterLibraryBinding(const FMVVMViewClass_CompiledBinding& Binding, int32 BindingIndex)
+{
+	check(ClassExtension);
+
+	TValueOrError<UE::FieldNotification::FFieldId, void> FieldIdResult = ClassExtension->GetBindingLibrary().GetFieldId(Binding.GetSourceFieldId());
+	if (FieldIdResult.HasError())
+	{
+		UE::MVVM::FMessageLog Log(GetUserWidget());
+		Log.Error(FText::Format(LOCTEXT("RegisterBindingInvalidFieldId", "Can't register binding '{0}'. The FieldId is invalid."), FText::FromString(Binding.ToString())));
+		return false;
+	}
+
+	UE::FieldNotification::FFieldId FieldId = FieldIdResult.StealValue();
+	if (!FieldId.IsValid())
+	{
+		UE::MVVM::FMessageLog Log(GetUserWidget());
+		Log.Error(FText::Format(LOCTEXT("RegisterBindingNoFieldId", "Can't register binding '{0}'. The FieldId was not found on the source."), FText::FromString(Binding.ToString())));
+		return false;
+	}
+
+	// The source may not have been created because the property path was wrong.
+	TScriptInterface<INotifyFieldValueChanged> Source = FindSource(Binding, true);
+	if (Source.GetInterface() == nullptr)
+	{
+		UE::MVVM::FMessageLog Log(GetUserWidget());
+		Log.Error(FText::Format(LOCTEXT("RegisterBindingInvalidSourceInstance", "Can't register binding '{0}'. The source is invalid."), FText::FromString(Binding.ToString())));
+		return false;
+	}
+
+	// Only bind if the source and the destination are valid.
+	if (ClassExtension->GetBindingLibrary().EvaluateFieldPath(GetUserWidget(), Binding.GetBinding().GetSourceFieldPath()).HasError())
+	{
+		UE::MVVM::FMessageLog Log(GetUserWidget());
+		Log.Warning(FText::Format(LOCTEXT("RegisterBindingInvalidDestination", "Can't register binding '{0}'. The source was not evaluated."), FText::FromString(Binding.ToString())));
+		return false;
+	}
+	if (ClassExtension->GetBindingLibrary().EvaluateFieldPath(GetUserWidget(), Binding.GetBinding().GetDestinationFieldPath()).HasError())
+	{
+		UE::MVVM::FMessageLog Log(GetUserWidget());
+		Log.Warning(FText::Format(LOCTEXT("RegisterBindingInvalidDestination", "Can't register binding '{0}'. The destination was not evaluated."), FText::FromString(Binding.ToString())));
+		return false;
+	}
+
+
+	UE::FieldNotification::FFieldMulticastDelegate::FDelegate Delegate = UE::FieldNotification::FFieldMulticastDelegate::FDelegate::CreateUObject(this, &UMVVMView::HandledLibraryBindingValueChanged, BindingIndex);
+	bool bResult = Source->AddFieldValueChangedDelegate(FieldId, MoveTemp(Delegate)).IsValid();
+	if (bResult)
+	{
+		FRegisteredSource* FoundSource = AllSources.FindByPredicate([&Source](const FRegisteredSource& Other) { return Other.Source == Source.GetObject(); });
+		if (FoundSource)
+		{
+			++FoundSource->Count;
+		}
+		else
+		{
+			FRegisteredSource Item;
+			Item.Source = Source.GetObject();
+			Item.Count = 1;
+			AllSources.Add(MoveTemp(Item));
+		}
+	}
+	return bResult;
+}
+
+
+void UMVVMView::UnregisterLibraryBinding(const FMVVMViewClass_CompiledBinding& Binding)
+{
+	TScriptInterface<INotifyFieldValueChanged> Source = FindSource(Binding, true);
+	if (Source.GetInterface() && ClassExtension)
+	{
+		TValueOrError<UE::FieldNotification::FFieldId, void> FieldIdResult = ClassExtension->GetBindingLibrary().GetFieldId(Binding.GetSourceFieldId());
+		if (ensureMsgf(FieldIdResult.HasValue() && FieldIdResult.GetValue().IsValid(), TEXT("If the binding was enabled then the FieldId should exist.")))
+		{
+			Source->RemoveAllFieldValueChangedDelegates(FieldIdResult.GetValue(), this);
+		}
+
+		int32 FoundIndex = AllSources.IndexOfByPredicate([&Source](const FRegisteredSource& Other) { return Other.Source == Source.GetObject(); });
+		if (ensureMsgf(FoundIndex != INDEX_NONE, TEXT("If the binding was enabled, then the source should also be there.")))
+		{
+			FRegisteredSource& FoundSource = AllSources[FoundIndex];
+			ensureMsgf(FoundSource.Count > 0, TEXT("The count should match the number of RegisterLibraryBinding and UnregisterLibraryBinding"));
+			--FoundSource.Count;
+			if (FoundSource.Count == 0)
+			{
+				AllSources.RemoveAtSwap(FoundIndex);
+			}
+		}
+	}
+}
+
+
+TScriptInterface<INotifyFieldValueChanged> UMVVMView::FindSource(const FMVVMViewClass_CompiledBinding& Binding, bool bAllowNull) const
+{
+	UUserWidget* UserWidget = GetUserWidget();
+	check(UserWidget);
+	check(ClassExtension);
+
+	const FObjectPropertyBase* SourceObjectProperty = CastField<FObjectPropertyBase>(UserWidget->GetClass()->FindPropertyByName(Binding.GetSourceObjectPropertyName()));
+	if (SourceObjectProperty == nullptr)
+	{
+		UE::MVVM::FMessageLog Log(GetUserWidget());
+		Log.Error(FText::Format(LOCTEXT("FindSourceInvalidPropertyName", "We could not evalute the source for binding '{0}'. The property name is invalid."), FText::FromString(Binding.ToString())));
+		return TScriptInterface<INotifyFieldValueChanged>();
+	}
+
+	UObject* Source = SourceObjectProperty->GetObjectPropertyValue_InContainer(UserWidget);
+	if (Source == nullptr)
+	{
+		if (!bAllowNull)
+		{
+			UE::MVVM::FMessageLog Log(GetUserWidget());
+			Log.Warning(FText::Format(LOCTEXT("FindSourceInvalidUObject", "We could not evalute the source for binding '{0}'. The path point to an invalid object."), FText::FromString(Binding.ToString())));
+		}
+		return TScriptInterface<INotifyFieldValueChanged>();
+	}
+
+	if (!Source->Implements<UNotifyFieldValueChanged>())
+	{
+		UE::MVVM::FMessageLog Log(GetUserWidget());
+		Log.Error(FText::Format(LOCTEXT("FindSourceNotTheCorrectInterface", "We could not evalute the source for binding '{0}'. The object {1} doesn't implements INotifyFieldValueChanged."), FText::FromString(Binding.ToString()), FText::FromName(Source->GetFName())));
+		return TScriptInterface<INotifyFieldValueChanged>();
+	}
+
+	return TScriptInterface<INotifyFieldValueChanged>(Source);
+}
+
+
+#undef LOCTEXT_NAMESPACE
\ No newline at end of file
diff --git a/Engine/Plugins/Runtime/ModelViewViewModel/Source/ModelViewViewModel/Private/View/MVVMViewClass.cpp b/Engine/Plugins/Runtime/ModelViewViewModel/Source/ModelViewViewModel/Private/View/MVVMViewClass.cpp
new file mode 100644
--- /dev/null
+++ b/Engine/Plugins/Runtime/ModelViewViewModel/Source/ModelViewViewModel/Private/View/MVVMViewClass.cpp
@@ -0,0 +1,169 @@
+// Copyright Epic Games, Inc. All Rights Reserved.
+
+#include "View/MVVMViewClass.h"
+#include "View/MVVMView.h"
+
+#include "Bindings/MVVMBindingHelper.h"
+#include "Bindings/MVVMFieldPathHelper.h"
+#include "Blueprint/UserWidget.h"
+#include "Blueprint/WidgetTree.h"
+#include "Engine/Engine.h"
+#include "FieldNotification/IFieldValueChanged.h"
+#include "MVVMMessageLog.h"
+#include "MVVMViewModelBase.h"
+#include "MVVMSubsystem.h"
+
+
+#define LOCTEXT_NAMESPACE "MVVMViewClass"
+
+///////////////////////////////////////////////////////////////////////
+// 
+///////////////////////////////////////////////////////////////////////
+
+FMVVMViewClass_SourceCreator FMVVMViewClass_SourceCreator::MakeInstance(FName InName, UClass* InNotifyFieldValueChangedClass)
+{
+	FMVVMViewClass_SourceCreator Result;
+	if (ensure(InNotifyFieldValueChangedClass && InNotifyFieldValueChangedClass->ImplementsInterface(UNotifyFieldValueChanged::StaticClass())))
+	{
+		Result.PropertyName = InName;
+		Result.ExpectedSourceType = InNotifyFieldValueChangedClass;
+		Result.bCreateInstance = true;
+		Result.bOptional = false;
+	}
+	return Result;
+}
+
+
+FMVVMViewClass_SourceCreator FMVVMViewClass_SourceCreator::MakeFieldPath(FName InName, UClass* InNotifyFieldValueChangedClass, FMVVMVCompiledFieldPath InFieldPath)
+{
+	FMVVMViewClass_SourceCreator Result;
+	if (ensure(InNotifyFieldValueChangedClass&& InNotifyFieldValueChangedClass->ImplementsInterface(UNotifyFieldValueChanged::StaticClass())))
+	{
+		if (ensure(InFieldPath.IsValid()))
+		{
+			Result.PropertyName = InName;
+			Result.ExpectedSourceType = InNotifyFieldValueChangedClass;
+			Result.FieldPath = InFieldPath;
+		}
+	}
+	return Result;
+}
+
+
+FMVVMViewClass_SourceCreator FMVVMViewClass_SourceCreator::MakeGlobalContext(FName InName, FMVVMViewModelContext InContext)
+{
+	FMVVMViewClass_SourceCreator Result;
+	if (ensure(InContext.ContextClass && InContext.ContextClass->ImplementsInterface(UNotifyFieldValueChanged::StaticClass())))
+	{
+		Result.PropertyName = InName;
+		Result.ExpectedSourceType = InContext.ContextClass;
+		Result.GlobalViewModelInstance = MoveTemp(InContext);
+	}
+	return Result;
+}
+
+
+UObject* FMVVMViewClass_SourceCreator::CreateInstance(const UMVVMViewClass* InViewClass, UMVVMView* InView, UUserWidget* InUserWidget) const
+{
+	check(InViewClass);
+	check(InView);
+	check(InUserWidget);
+
+	UObject* Result = nullptr;
+
+	FObjectPropertyBase* FoundObjectProperty = FindFProperty<FObjectPropertyBase>(InUserWidget->GetClass(), PropertyName);
+	if (ensureAlwaysMsgf(FoundObjectProperty, TEXT("The compiler should have added the property")))
+	{
+		auto AssignProperty = [FoundObjectProperty, InUserWidget](UObject* NewObject)
+		{
+			check(NewObject);
+			if (ensure(NewObject->GetClass()->IsChildOf(FoundObjectProperty->PropertyClass)))
+			{
+				FoundObjectProperty->SetObjectPropertyValue_InContainer(InUserWidget, NewObject);
+			}
+		};
+
+		if (bCreateInstance)
+		{
+			if (ExpectedSourceType.Get() != nullptr)
+			{
+				Result = NewObject<UObject>(InUserWidget, ExpectedSourceType.Get(), NAME_None, RF_Transient);
+				AssignProperty(Result);
+			}
+			else if (!bOptional)
+			{
+				UE::MVVM::FMessageLog Log(InUserWidget);
+				Log.Error(FText::Format(LOCTEXT("CreateInstanceCreateInstance", "The source '{0}' could not be created. The class is not loaded."), FText::FromName(PropertyName)));
+			}
+		}
+		else if (GlobalViewModelInstance.IsValid())
+		{
+			UMVVMViewModelBase* FoundViewModelInstance = GEngine->GetEngineSubsystem<UMVVMSubsystem>()->GetGlobalViewModelCollection()->FindViewModelInstance(GlobalViewModelInstance);
+			if (FoundViewModelInstance)
+			{
+				ensureMsgf(FoundViewModelInstance->IsA(GlobalViewModelInstance.ContextClass), TEXT("The Global View Model Instance is not of the expected type."));
+				Result = FoundViewModelInstance;
+				AssignProperty(Result);
+			}
+			else if (!bOptional)
+			{
+				UE::MVVM::FMessageLog Log(InUserWidget);
+				Log.Error(FText::Format(LOCTEXT("CreateInstanceFailedGlobal", "The viewmodel '{0}' was not found in the global view model collection."), FText::FromName(GlobalViewModelInstance.ContextName)));
+			}
+		}
+		else if (FieldPath.IsValid())
+		{
+			TValueOrError<UE::MVVM::FFieldContext, void> FieldPathResult = InViewClass->GetBindingLibrary().EvaluateFieldPath(InUserWidget, FieldPath);
+			if (FieldPathResult.HasValue())
+			{
+				TValueOrError<UObject*, void> ObjectResult = UE::MVVM::FieldPathHelper::EvaluateObjectProperty(FieldPathResult.GetValue());
+				if (ObjectResult.HasValue())
+				{
+					Result = ObjectResult.GetValue();
+					AssignProperty(ObjectResult.GetValue());
+				}
+			}
+
+			if (Result == nullptr && !bOptional)
+			{
+				UE::MVVM::FMessageLog Log(InUserWidget);
+				Log.Error(FText::Format(LOCTEXT("CreateInstanceInvalidBiding", "The source '{0}' was evaluated to be invalid at initialization."), FText::FromName(PropertyName)));
+			}
+		}
+	}
+
+	return Result;
+}
+
+
+///////////////////////////////////////////////////////////////////////
+// 
+///////////////////////////////////////////////////////////////////////
+
+FString FMVVMViewClass_CompiledBinding::ToString() const
+{
+	return TEXT("[PH]");
+}
+
+
+///////////////////////////////////////////////////////////////////////
+// 
+///////////////////////////////////////////////////////////////////////
+
+void UMVVMViewClass::Initialize(UUserWidget* UserWidget)
+{
+	check(UserWidget->GetExtension<UMVVMView>() == nullptr);
+
+	UMVVMView* View = UserWidget->AddExtension<UMVVMView>();
+	if (ensure(View))
+	{
+		if (!bLoaded)
+		{
+			BindingLibrary.Load();
+		}
+
+		View->ConstructView(this);
+	}
+}
+
+#undef LOCTEXT_NAMESPACE
diff --git a/Engine/Plugins/Runtime/ModelViewViewModel/Source/ModelViewViewModel/Public/Bindings/MVVMBindingHelper.h b/Engine/Plugins/Runtime/ModelViewViewModel/Source/ModelViewViewModel/Public/Bindings/MVVMBindingHelper.h
new file mode 100644
--- /dev/null
+++ b/Engine/Plugins/Runtime/ModelViewViewModel/Source/ModelViewViewModel/Public/Bindings/MVVMBindingHelper.h
@@ -0,0 +1,137 @@
+// Copyright Epic Games, Inc. All Rights Reserved.
+
+#pragma once
+
+#include "CoreMinimal.h"
+
+#include "Templates/ValueOrError.h"
+#include "Types/MVVMFieldContext.h"
+#include "Types/MVVMFunctionContext.h"
+#include "Types/MVVMFieldVariant.h"
+#include "Types/MVVMObjectVariant.h"
+#include "Types/MVVMBindingName.h"
+
+class FProperty;
+class UFunction;
+
+namespace UE::MVVM::BindingHelper
+{
+
+	/** Is the Property usable as a source (readable) by the binding system. It may required a Getter to read it. */
+	UE_NODISCARD MODELVIEWVIEWMODEL_API bool IsValidForSourceBinding(const FProperty* InProperty);
+
+	/**
+	 * Is the Function usable as a source (readable) by the binding system.
+	 * @note It may be a BlueprintGetter and binding to the Property would be better in the editor.
+	 */
+	UE_NODISCARD MODELVIEWVIEWMODEL_API bool IsValidForSourceBinding(const UFunction* InFunction);
+
+	/**
+	 * Is the Field usable as a source (readable) by the binding system.
+	 * @note It may be a BlueprintGetter and binding to the Property would be better in the editor.
+	 */
+	UE_NODISCARD MODELVIEWVIEWMODEL_API bool IsValidForSourceBinding(const FMVVMConstFieldVariant InField);
+
+	/** Is the Property usable as a destination (settable) by the binding system. It may required a Setter to write it. */
+	UE_NODISCARD MODELVIEWVIEWMODEL_API bool IsValidForDestinationBinding(const FProperty* InProperty);
+
+	/**
+	 * Is the Function usable as a destination (settable) by the binding system.
+	 * @note It may be a BlueprintSetter and binding to the Property would be better in the editor.
+	 */
+	UE_NODISCARD MODELVIEWVIEWMODEL_API bool IsValidForDestinationBinding(const UFunction* InFunction);
+
+	/**
+	 * Is the Field usable as a destination (settable) by the binding system.
+	 * @note It may be a BlueprintSetter and binding to the Property would be better in the editor.
+	 */
+	UE_NODISCARD MODELVIEWVIEWMODEL_API bool IsValidForDestinationBinding(const FMVVMConstFieldVariant InFunction);
+
+	/**
+	 * Is the Function usable as a conversion function by the binding system.
+	 */
+	UE_NODISCARD MODELVIEWVIEWMODEL_API bool IsValidForRuntimeConversion(const UFunction* InFunction);
+
+#if WITH_EDITOR
+	/** Is the Property usable as a source by the binding system and can it be read directly or it requires a Getter. */
+	UE_NODISCARD MODELVIEWVIEWMODEL_API bool IsAccessibleDirectlyForSourceBinding(const FProperty* InProperty);
+
+	/** Is the Property usable as a destination by the binding system and can it be write directly or it requires a Setter. */
+	UE_NODISCARD MODELVIEWVIEWMODEL_API bool IsAccessibleDirectlyForDestinationBinding(const FProperty* InProperty);
+
+	/** Is the Property usable as a source by the binding system and a Getter exists. */
+	UE_NODISCARD MODELVIEWVIEWMODEL_API bool IsAccessibleWithGetterForSourceBinding(const FProperty* InProperty);
+
+	/** Is the Property usable as a destination by the binding system and a Setter exists. */
+	UE_NODISCARD MODELVIEWVIEWMODEL_API bool IsAccessibleWithSetterForDestinationBinding(const FProperty* InProperty);
+#endif
+
+	/**
+	 * Returns the Property or the Function that matches that BindingName.
+	 * @note It doesn't check for BlueprintGetter or BlueprintSetter.
+	 */
+	UE_NODISCARD MODELVIEWVIEWMODEL_API FMVVMFieldVariant FindFieldByName(const UStruct* Container, FMVVMBindingName BindingName);
+
+	UE_NODISCARD MODELVIEWVIEWMODEL_API TValueOrError<const FProperty*, FString> TryGetPropertyTypeForSourceBinding(const FMVVMConstFieldVariant& InField);
+	UE_NODISCARD MODELVIEWVIEWMODEL_API TValueOrError<const FProperty*, FString> TryGetPropertyTypeForDestinationBinding(const FMVVMConstFieldVariant& InField);
+	UE_NODISCARD MODELVIEWVIEWMODEL_API TValueOrError<const FProperty*, FString> TryGetPropertyTypeForSourceBinding(const FProperty* InProperty);
+	UE_NODISCARD MODELVIEWVIEWMODEL_API TValueOrError<const FProperty*, FString> TryGetPropertyTypeForDestinationBinding(const FProperty* InProperty);
+	UE_NODISCARD MODELVIEWVIEWMODEL_API TValueOrError<const FProperty*, FString> TryGetPropertyTypeForSourceBinding(const UFunction* InFunction);
+	UE_NODISCARD MODELVIEWVIEWMODEL_API TValueOrError<const FProperty*, FString> TryGetPropertyTypeForDestinationBinding(const UFunction* InFunction);
+
+	struct FConversionFunctionArguments
+	{
+		const FProperty* ReturnProperty = nullptr;
+		const FProperty* ArgumentProperty = nullptr;
+	};
+	/** Returns the Return and 1st argument for a conversion function. */
+	UE_NODISCARD MODELVIEWVIEWMODEL_API TValueOrError<FConversionFunctionArguments, FString> TryGetPropertyTypeForConversionFunction(const UFunction* InFunction);
+
+	/** Are type the same or could be converted at runtime. */
+	UE_NODISCARD MODELVIEWVIEWMODEL_API bool ArePropertiesCompatible(const FProperty* Source, const FProperty* Destination);
+
+	/**
+	 * Returns the "return" property of the function. It can be a none-const ref/out argument.
+	 * int Foo(double) -> returns int
+	 * void Foo(int&, double) -> returns int
+	 * void Foo(const int&, double) -> returns null
+	 */
+	UE_NODISCARD MODELVIEWVIEWMODEL_API const FProperty* GetReturnProperty(const UFunction* InFunction);
+
+	/**
+	 * Returns the first argument property. Ref/Out argument have to be const.
+	 * int Foo(double) -> returns double
+	 * void Foo(int& Out, double) -> returns double
+	 * void Foo(const int& Out, double) -> returns int
+	 */
+	UE_NODISCARD MODELVIEWVIEWMODEL_API const FProperty* GetFirstArgumentProperty(const UFunction* InFunction);
+
+	/**
+	 * Execute a binding that can be
+	 *   LocalValue = Src.Property; Dst.Property = LocalValue;
+	 *   LocalValue = Src.Property; Dst.SetProperty(LocalValue);
+	 *   LocalValue = Src.GetProperty(); Dst.Property = LocalValue;
+	 *   LocalValue = Src.GetProperty(); Dst.SetProperty(LocalValue);
+	 * with conversion from float to double between the getter and the setter.
+	 * 
+	 * @note No test is performed to see if the Src Property can be safely assign to the Destination Property. Use with caution.
+	 * @note No test is performed to see if the Src.GetProperty can be safely executed. Use with caution.
+	 * @note No test is performed to see if the Dst.SetProperty can be safely executed. Use with caution.
+	 */
+	MODELVIEWVIEWMODEL_API void ExecuteBinding_NoCheck(const FFieldContext& Source, const FFieldContext& Destination);
+
+	/**
+	 * Execute a binding that can be
+	 *   LocalValue = Src.Property; LocalValue = ConcFuncOwner.ConversionFunction(LocalValue); Dst.Property = LocalValue;
+	 *   LocalValue = Src.Property; LocalValue = ConcFuncOwner.ConversionFunction(LocalValue); Dst.SetProperty(LocalValue);
+	 *   LocalValue = Src.GetProperty(); LocalValue = ConcFuncOwner.ConversionFunction(LocalValue); Dst.Property = LocalValue;
+	 *   LocalValue = Src.GetProperty(); LocalValue = ConcFuncOwner.ConversionFunction(LocalValue); Dst.SetProperty(LocalValue);
+	 * with conversion from float to double between the getter and the conversion function and/or the conversion function and the setter.
+	 *
+	 * @note No test is performed to see if the Src Property can be safely assign to the Destination Property. Use with caution.
+	 * @note No test is performed to see if the Src.GetProperty can be safely executed. Use with caution.
+	 * @note No test is performed to see if the Dst.SetProperty can be safely executed. Use with caution.
+	 * @note No test is performed to see if the ConversionFunction can be safely executed. Use with caution.
+	 */
+	MODELVIEWVIEWMODEL_API void ExecuteBinding_NoCheck(const FFieldContext& Source, const FFieldContext& Destination, const FFunctionContext& ConversionFunction);
+} //namespace
diff --git a/Engine/Plugins/Runtime/ModelViewViewModel/Source/ModelViewViewModel/Public/Bindings/MVVMCompiledBindingLibrary.h b/Engine/Plugins/Runtime/ModelViewViewModel/Source/ModelViewViewModel/Public/Bindings/MVVMCompiledBindingLibrary.h
new file mode 100644
--- /dev/null
+++ b/Engine/Plugins/Runtime/ModelViewViewModel/Source/ModelViewViewModel/Public/Bindings/MVVMCompiledBindingLibrary.h
@@ -0,0 +1,315 @@
+// Copyright Epic Games, Inc. All Rights Reserved.
+
+#pragma once
+
+#include "CoreMinimal.h"
+#include "FieldNotification/FieldId.h"
+#include "Templates/ValueOrError.h"
+#include "Types/MVVMFieldContext.h"
+#include "Types/MVVMFunctionContext.h"
+#include "Types/MVVMFieldVariant.h"
+#include "Types/MVVMObjectVariant.h"
+#include "UObject/FieldPath.h"
+#include "UObject/Object.h"
+
+#include "MVVMCompiledBindingLibrary.generated.h"
+
+struct FMVVMCompiledBindingLibrary;
+namespace UE::MVVM
+{
+	class FCompiledBindingLibraryCompiler;
+}
+
+
+/** The info to fetch a list of FProperty or UFunction from a Class that will be needed by bindings. */
+USTRUCT()
+struct MODELVIEWVIEWMODEL_API FMVVMVCompiledFields
+{
+	GENERATED_BODY()
+	friend UE::MVVM::FCompiledBindingLibraryCompiler;
+
+public:
+	int32 GetPropertyNum() const
+	{
+		return ClassOrScriptStruct ? NumberOfProperties : 0;
+	}
+
+	int32 GetFunctionNum() const
+	{
+		return ClassOrScriptStruct ? NumberOfFunctions : 0;
+	}
+
+	int32 GetFieldIdNum() const
+	{
+		return ClassOrScriptStruct ? NumberOfFieldIds : 0;
+	}
+
+	FName GetPropertyName(TArrayView<FName> Names, int32 Index) const
+	{
+		check(Index < NumberOfProperties && Index >= 0);
+		check(Names.IsValidIndex(LibraryStartIndex + Index));
+		return Names[LibraryStartIndex + Index];
+	}
+
+	FName GetFunctionName(TArrayView<FName> Names, int32 Index) const
+	{
+		check(Index < NumberOfFunctions && Index >= 0);
+		check(Names.IsValidIndex(LibraryStartIndex + NumberOfProperties + Index));
+		return Names[LibraryStartIndex + NumberOfProperties + Index];
+	}
+
+	FName GetFieldIdName(TArrayView<FName> Names, int32 Index) const
+	{
+		check(Index < NumberOfFieldIds && Index >= 0);
+		check(Names.IsValidIndex(LibraryStartIndex + NumberOfProperties + NumberOfFunctions + Index));
+		return Names[LibraryStartIndex + NumberOfProperties + NumberOfFunctions + Index];
+	}
+
+	const UStruct* GetStruct() const
+	{
+		return ClassOrScriptStruct;
+	}
+
+	/** Find the FProperty from the  Class.*/
+	FProperty* GetProperty(FName PropertyName) const;
+
+	/** Find the UFunction from the Class.*/
+	UFunction* GetFunction(FName FunctionName) const;
+
+	/** Find the FFieldId from the Class.*/
+	UE::FieldNotification::FFieldId GetFieldId(FName FieldName) const;
+
+private:
+	UPROPERTY()
+	TObjectPtr<const UStruct> ClassOrScriptStruct = nullptr;
+
+	UPROPERTY()
+	int16 LibraryStartIndex = 0;
+
+	UPROPERTY()
+	int16 NumberOfProperties = 0;
+
+	UPROPERTY()
+	int16 NumberOfFunctions = 0;
+	
+	UPROPERTY()
+	int16 NumberOfFieldIds = 0;
+};
+
+
+/** Contains the info to evaluate a path for a specific library. */
+USTRUCT()
+struct FMVVMVCompiledFieldPath
+{
+	GENERATED_BODY()
+
+	friend FMVVMCompiledBindingLibrary;
+	friend UE::MVVM::FCompiledBindingLibraryCompiler;
+
+	bool IsValid() const
+	{
+		return Num >= 0;
+	}
+
+private:
+	using IndexType = int16;
+
+	UPROPERTY()
+	int16 StartIndex = INDEX_NONE;
+
+	UPROPERTY()
+	int16 Num = INDEX_NONE;
+
+#if WITH_EDITORONLY_DATA
+	UPROPERTY(meta = (IgnoreForMemberInitializationTest))
+	FGuid CompiledBindingLibraryId;
+#endif
+};
+
+
+/** Contains the FieldId index. */
+USTRUCT()
+struct FMVVMVCompiledFieldId
+{
+	GENERATED_BODY()
+
+	friend FMVVMCompiledBindingLibrary;
+	friend UE::MVVM::FCompiledBindingLibraryCompiler;
+
+	bool IsValid() const
+	{
+		return FieldIdIndex >= 0;
+	}
+
+private:
+	using IndexType = int16;
+
+	UPROPERTY()
+	int16 FieldIdIndex = INDEX_NONE;
+
+#if WITH_EDITORONLY_DATA
+	UPROPERTY(meta = (IgnoreForMemberInitializationTest))
+	FGuid CompiledBindingLibraryId;
+#endif
+};
+
+
+/** */
+USTRUCT()
+struct FMVVMCompiledLoadedPropertyOrFunctionIndex
+{
+	GENERATED_BODY()
+
+	friend FMVVMCompiledBindingLibrary;
+	friend UE::MVVM::FCompiledBindingLibraryCompiler;
+
+	FMVVMCompiledLoadedPropertyOrFunctionIndex()
+	{
+		bIsObjectProperty = false;
+		bIsScriptStructProperty = false;
+		bIsProperty = false;
+	}
+
+private:
+	UPROPERTY()
+	int16 Index = INDEX_NONE;
+
+	/** Is the property or the return property of the UFunction is a FObjectPropertyBase */
+	UPROPERTY()
+	uint8 bIsObjectProperty : 1;
+
+	/** Is the property or the return property of the UFunction is a FStructProperty */
+	UPROPERTY()
+	uint8 bIsScriptStructProperty : 1;
+
+	/** Is the index in LoadedProperties or the index is in LoadedFunctions. */
+	UPROPERTY()
+	uint8 bIsProperty : 1;
+};
+
+
+/** Contains a single combination to execute a binding for a specific library. */
+USTRUCT()
+struct FMVVMVCompiledBinding
+{
+	GENERATED_BODY()
+
+	friend FMVVMCompiledBindingLibrary;
+	friend UE::MVVM::FCompiledBindingLibraryCompiler;
+
+public:
+	bool IsValid() const
+	{
+		return SourceFieldPath.IsValid() && DestinationFieldPath.IsValid();
+	}
+
+	FMVVMVCompiledFieldPath GetSourceFieldPath() const
+	{
+		return SourceFieldPath;
+	}
+
+	FMVVMVCompiledFieldPath GetDestinationFieldPath() const
+	{
+		return DestinationFieldPath;
+	}
+
+	FMVVMVCompiledFieldPath GetConversionFunctionFieldPath() const
+	{
+		return ConversionFunctionFieldPath;
+	}
+
+private:
+	using IndexType = int16;
+
+	UPROPERTY()
+	FMVVMVCompiledFieldPath SourceFieldPath; // todo: make this an array to support multiple input to conversion functions
+
+	UPROPERTY()
+	FMVVMVCompiledFieldPath DestinationFieldPath;
+
+	UPROPERTY()
+	FMVVMVCompiledFieldPath ConversionFunctionFieldPath;
+
+#if WITH_EDITORONLY_DATA
+	UPROPERTY(meta = (IgnoreForMemberInitializationTest))
+	FGuid CompiledBindingLibraryId;
+#endif
+};
+
+
+/** Library of all the compiled bindings. */
+USTRUCT()
+struct MODELVIEWVIEWMODEL_API FMVVMCompiledBindingLibrary
+{
+	GENERATED_BODY()
+	friend UE::MVVM::FCompiledBindingLibraryCompiler;
+
+public:
+	FMVVMCompiledBindingLibrary();
+
+public:
+	/** */
+	enum class EExecutionFailingReason
+	{
+		IncompatibleLibrary,
+		InvalidSource,
+		InvalidDestination,
+		InvalidConversionFunction,
+	};
+
+	/** Fetch the FProperty and UFunction. */
+	void Load();
+
+	/**
+	 * Execute a binding, in one direction.
+	 * The ExecutionSource is the View (UserWidget) instance.
+	 * The Source of the binding (ViewModel, UserWidget, Widget, ...) can be provided if already known.
+	 * The Destination or the binding (ViewModel, UserWIdget, Widget, ...) can be provided if already known.
+	 */
+	TValueOrError<void, EExecutionFailingReason> Execute(UObject* ExecutionSource, const FMVVMVCompiledBinding& Binding) const;
+	TValueOrError<void, EExecutionFailingReason> ExecuteWithSource(UObject* ExecutionSource, const FMVVMVCompiledBinding& Binding, UObject* Source) const;
+
+	/**
+	 * Evaluate the path to find the container use by a binding.
+	 * ExecutionSource is the View (UserWidget) instance.
+	 * @returns the container and the last segment of the path. The last segment can be a property or a function.
+	 */
+	TValueOrError<UE::MVVM::FFieldContext, void> EvaluateFieldPath(UObject* ExecutionSource, const FMVVMVCompiledFieldPath& FieldPath) const;
+
+	/** Get the loaded FieldId. */
+	TValueOrError<UE::FieldNotification::FFieldId, void> GetFieldId(const FMVVMVCompiledFieldId& FieldId) const;
+
+	/** Return a readable version of the FieldPath. */
+	TValueOrError<FString, FString> FieldPathToString(const FMVVMVCompiledFieldPath& FieldPath) const;
+
+private:
+	TValueOrError<void, EExecutionFailingReason> ExecuteImpl(UObject* ExecutionSource, const FMVVMVCompiledBinding& Binding, UObject* Source) const;
+	TValueOrError<void, EExecutionFailingReason> ExecuteImpl(UE::MVVM::FFieldContext& Source, UE::MVVM::FFieldContext& Destination, UE::MVVM::FFunctionContext& ConversionFunction) const;
+
+	TValueOrError<UE::MVVM::FMVVMFieldVariant, void> GetFinalFieldFromPathImpl(const FMVVMVCompiledFieldPath& FieldPath) const;
+
+private:
+	TArray<FProperty*> LoadedProperties;
+
+	UPROPERTY(Transient)
+	TArray<UFunction*> LoadedFunctions;
+
+	TArray<UE::FieldNotification::FFieldId> LoadedFieldIds;
+
+	UPROPERTY()
+	TArray<FMVVMCompiledLoadedPropertyOrFunctionIndex> FieldPaths;
+
+private:
+#if WITH_EDITORONLY_DATA
+	UPROPERTY(meta = (IgnoreForMemberInitializationTest))
+	FGuid CompiledBindingLibraryId;
+#endif
+
+	/** Only needed for loading the FProperty/UFunction. */
+	UPROPERTY()
+	TArray<FMVVMVCompiledFields> CompiledFields;
+
+	/** Only needed for loading the FProperty/UFunction. */
+	UPROPERTY()
+	TArray<FName> CompiledFieldNames;
+};
diff --git a/Engine/Plugins/Runtime/ModelViewViewModel/Source/ModelViewViewModel/Public/Bindings/MVVMFieldPathHelper.h b/Engine/Plugins/Runtime/ModelViewViewModel/Source/ModelViewViewModel/Public/Bindings/MVVMFieldPathHelper.h
new file mode 100644
--- /dev/null
+++ b/Engine/Plugins/Runtime/ModelViewViewModel/Source/ModelViewViewModel/Public/Bindings/MVVMFieldPathHelper.h
@@ -0,0 +1,24 @@
+// Copyright Epic Games, Inc. All Rights Reserved.
+
+#pragma once
+
+#include "CoreMinimal.h"
+#include "Templates/SubclassOf.h"
+#include "Templates/ValueOrError.h"
+#include "Types/MVVMFieldContext.h"
+#include "UObject/FieldPath.h"
+#include "UObject/Object.h"
+
+
+namespace UE::MVVM::FieldPathHelper
+{
+	/** */
+	UE_NODISCARD MODELVIEWVIEWMODEL_API TValueOrError<TArray<FMVVMFieldVariant>, FString> GenerateFieldPathList(TSubclassOf<UObject> From, FStringView FieldPath, bool bForSourceBinding);
+
+	/** */
+	UE_NODISCARD MODELVIEWVIEWMODEL_API TValueOrError<TArray<FMVVMFieldVariant>, FString> GenerateConversionFunctionFieldPathList(TSubclassOf<UObject> From, FStringView FieldPath);
+
+	/** */
+	UE_NODISCARD MODELVIEWVIEWMODEL_API TValueOrError<UObject*, void> EvaluateObjectProperty(const FFieldContext& InSource);
+
+} // namespace
diff --git a/Engine/Plugins/Runtime/ModelViewViewModel/Source/ModelViewViewModel/Public/MVVMSubsystem.h b/Engine/Plugins/Runtime/ModelViewViewModel/Source/ModelViewViewModel/Public/MVVMSubsystem.h
new file mode 100644
--- /dev/null
+++ b/Engine/Plugins/Runtime/ModelViewViewModel/Source/ModelViewViewModel/Public/MVVMSubsystem.h
@@ -0,0 +1,106 @@
+// Copyright Epic Games, Inc. All Rights Reserved.
+
+#pragma once
+
+#include "Subsystems/EngineSubsystem.h"
+#include "Templates/ValueOrError.h"
+
+#include "Types/MVVMBindingName.h"
+#include "Types/MVVMFieldVariant.h"
+#include "Types/MVVMBindingMode.h"
+#include "Types/MVVMViewModelCollection.h"
+
+#include "MVVMSubsystem.generated.h"
+
+class UMVVMView;
+class UMVVMViewModelBase;
+class UUserWidget;
+class UWidget;
+class UWidgetTree;
+
+/** */
+UCLASS()
+class MODELVIEWVIEWMODEL_API UMVVMSubsystem : public UEngineSubsystem
+{
+	GENERATED_BODY()
+
+public:
+	//~ Begin UEngineSubsystem interface
+	virtual void Initialize(FSubsystemCollectionBase& Collection) override;
+	virtual void Deinitialize() override;
+	//~ End UEngineSubsystem interface
+
+public:
+	UFUNCTION(BlueprintCallable, Category="MVVM")
+	UMVVMView* GetViewFromUserWidget(const UUserWidget* UserWidget) const;
+
+	UFUNCTION(BlueprintCallable, Category="MVVM")
+	bool IsViewModelValueValidForSourceBinding(const UMVVMViewModelBase* ViewModel, FMVVMBindingName ViewModelPropertyOrFunctionName) const;
+
+	UFUNCTION(BlueprintCallable, Category = "MVVM")
+	bool IsViewModelValueValidForDestinationBinding(const UMVVMViewModelBase* ViewModel, FMVVMBindingName ViewModelPropertyOrFunctionName) const;
+
+	UFUNCTION(BlueprintCallable, Category = "MVVM")
+	bool IsViewValueValidForSourceBinding(const UWidget* View, FMVVMBindingName ViewPropertyOrFunctionName) const;
+
+	UFUNCTION(BlueprintCallable, Category = "MVVM")
+	bool IsViewValueValidForDestinationBinding(const UWidget* View, FMVVMBindingName ViewPropertyOrFunctionName) const;
+
+	UFUNCTION(BlueprintCallable, Category = "MVVM")
+	bool DoesWidgetTreeContainedWidget(const UWidgetTree* WidgetTree, const UWidget* ViewWidget) const;
+
+	/** Returns the list of all the binding that are available for the ViewModel. */
+	UFUNCTION(BlueprintCallable, Category = "MVVM")
+	TArray<FMVVMAvailableBinding> GetViewModelAvailableBindings(TSubclassOf<UMVVMViewModelBase> ViewModelClass) const;
+
+	/** Returns the list of all the binding that are available for the Widget. */
+	UFUNCTION(BlueprintCallable, Category = "MVVM")
+	TArray<FMVVMAvailableBinding> GetWidgetAvailableBindings(TSubclassOf<UWidget> WidgetClass) const;
+
+public:
+	UFUNCTION(BlueprintCallable, Category = "MVVM")
+	UMVVMViewModelCollectionObject* GetGlobalViewModelCollection() const
+	{
+		return GlobalViewModelCollection;
+	}
+
+public:
+
+	struct FConstDirectionalBindingArgs
+	{
+		UE::MVVM::FMVVMConstFieldVariant SourceBinding;
+		UE::MVVM::FMVVMConstFieldVariant DestinationBinding;
+		const UFunction* ConversionFunction = nullptr;
+	};
+
+	struct FDirectionalBindingArgs
+	{
+		UE::MVVM::FMVVMFieldVariant SourceBinding;
+		UE::MVVM::FMVVMFieldVariant DestinationBinding;
+		UFunction* ConversionFunction = nullptr;
+
+		FConstDirectionalBindingArgs ToConst() const
+		{
+			FConstDirectionalBindingArgs ConstArgs;
+			ConstArgs.SourceBinding = SourceBinding;
+			ConstArgs.DestinationBinding = DestinationBinding;
+			ConstArgs.ConversionFunction = ConversionFunction;
+			return MoveTemp(ConstArgs);
+		}
+	};
+
+	struct FBindingArgs
+	{
+		EMVVMBindingMode Mode = EMVVMBindingMode::OneWayToDestination;
+		FDirectionalBindingArgs ForwardArgs;
+		FDirectionalBindingArgs BackwardArgs;
+	};
+
+	UE_NODISCARD TValueOrError<bool, FString> IsBindingValid(FConstDirectionalBindingArgs Args) const;
+	UE_NODISCARD TValueOrError<bool, FString> IsBindingValid(FDirectionalBindingArgs Args) const;
+	UE_NODISCARD TValueOrError<bool, FString> IsBindingValid(FBindingArgs Args) const;
+
+private:
+	UPROPERTY(Transient)
+	UMVVMViewModelCollectionObject* GlobalViewModelCollection;
+};
diff --git a/Engine/Plugins/Runtime/ModelViewViewModel/Source/ModelViewViewModel/Public/MVVMViewModelBase.h b/Engine/Plugins/Runtime/ModelViewViewModel/Source/ModelViewViewModel/Public/MVVMViewModelBase.h
new file mode 100644
--- /dev/null
+++ b/Engine/Plugins/Runtime/ModelViewViewModel/Source/ModelViewViewModel/Public/MVVMViewModelBase.h
@@ -0,0 +1,92 @@
+// Copyright Epic Games, Inc. All Rights Reserved.
+
+#pragma once
+
+#include "CoreMinimal.h"
+#include "UObject/Object.h"
+
+#include "FieldNotification/FieldNotificationDeclaration.h"
+#include "FieldNotification/FieldMulticastDelegate.h"
+#include "FieldNotification/IFieldValueChanged.h"
+#include "Types/MVVMAvailableBinding.h"
+#include "Types/MVVMBindingName.h"
+
+#include "MVVMViewModelBase.generated.h"
+
+//UENUM()
+//enum class EMVVMNGetterType : uint8
+//{
+//	/** The Source Property has a Delegate associated with it that will notify us when the value change. */
+//	Delegate,
+//	/** The value is fetch when requested. */
+//	FetchOnRequest,
+//	/** The value is fetch every frame and compared. */ 
+//	FetchOnTickAndCompared,
+//	/** The value is fetch every frame and is considered changed. */
+//	FetchOnTick,
+//};
+
+///** The user will notify the ViewModel when the value change via the PropertyValueChanged event. */
+//Event,
+//Interface,		// The Source Context implement the MVVMModelInterface.
+					//Use the interface to register to the SourceContext's OnPropertyValueChanged.
+//PushModel,		// The property has a NetId generated by UBT. Used to notify us when the value changed.
+
+#define UE_MVVM_NOTIFY_FIELD_VALUE_CHANGED(MemberName) \
+	BindingFieldValueChanged(ThisClass::FFieldNotificationClassDescriptor::MemberName)
+
+#define UE_MVVM_SET_PROPERTY_VALUE(MemberName, NewValue) \
+	SetPropertyValue(MemberName, NewValue, ThisClass::FFieldNotificationClassDescriptor::MemberName)
+
+
+/** Sub class of this will be generated from the BP to cached the source data. */
+UCLASS(Blueprintable, Abstract, DisplayName="MVVM ViewModel")
+class MODELVIEWVIEWMODEL_API UMVVMViewModelBase : public UObject, public INotifyFieldValueChanged
+{
+	GENERATED_BODY()
+
+public:
+	struct FFieldNotificationClassDescriptor : public UE::FieldNotification::IClassDescriptor
+	{
+	};
+
+public:
+	//~ Begin INotifyFieldValueChanged Interface
+	virtual FDelegateHandle AddFieldValueChangedDelegate(UE::FieldNotification::FFieldId InFieldId, FFieldValueChangedDelegate InNewDelegate) override final;
+	virtual bool RemoveFieldValueChangedDelegate(UE::FieldNotification::FFieldId InFieldId, FDelegateHandle InHandle) override final;
+	virtual int32 RemoveAllFieldValueChangedDelegates(const void* InUserObject) override final;
+	virtual int32 RemoveAllFieldValueChangedDelegates(UE::FieldNotification::FFieldId InFieldId, const void* InUserObject) override final;
+	virtual const UE::FieldNotification::IClassDescriptor& GetFieldNotificationDescriptor() const override;
+	//~ End INotifyFieldValueChanged Interface
+
+protected:
+	// todo make a 2k node to create a binding base on a variable or a function
+	UFUNCTION(BlueprintCallable, Category="MVVM")
+	static FMVVMBindingName MakeBindingName_Temporary(FName BindingName)
+	{
+		return FMVVMBindingName(BindingName);
+	}
+
+protected:
+	UFUNCTION(BlueprintCallable, Category = "MVVM")
+	void BindingValueChanged(FMVVMBindingName PropertyName);
+	void BindingFieldValueChanged(UE::FieldNotification::FFieldId InFieldId);
+
+protected:
+	template<typename T>
+	bool SetPropertyValue(T& Value, const T& NewValue, UE::FieldNotification::FFieldId FieldId)
+	{
+		if (Value == NewValue)
+		{
+			return false;
+		}
+
+		Value = NewValue;
+		BindingFieldValueChanged(FieldId);
+		return true;
+	}
+
+private:
+	UE::FieldNotification::FFieldMulticastDelegate Delegates;
+	TBitArray<> EnabledFieldNotifications;
+};
diff --git a/Engine/Plugins/Runtime/ModelViewViewModel/Source/ModelViewViewModel/Public/Types/MVVMAvailableBinding.h b/Engine/Plugins/Runtime/ModelViewViewModel/Source/ModelViewViewModel/Public/Types/MVVMAvailableBinding.h
new file mode 100644
--- /dev/null
+++ b/Engine/Plugins/Runtime/ModelViewViewModel/Source/ModelViewViewModel/Public/Types/MVVMAvailableBinding.h
@@ -0,0 +1,73 @@
+// Copyright Epic Games, Inc. All Rights Reserved.
+
+#pragma once
+
+#include "CoreMinimal.h"
+#include "UObject/Object.h"
+#include "Types/MVVMBindingName.h"
+
+#include "MVVMAvailableBinding.generated.h"
+
+
+USTRUCT(BlueprintType)
+struct FMVVMAvailableBinding
+{
+	GENERATED_BODY()
+
+public:
+	FMVVMAvailableBinding() = default;
+
+	explicit FMVVMAvailableBinding(FMVVMBindingName InBindingName, bool bReadable, bool bWritable)
+		: BindingName(InBindingName)
+		, bIsReadable(bReadable)
+		, bIsWritable(bWritable)
+	{
+
+	}
+
+public:
+	FMVVMBindingName GetBindingName() const
+	{
+		return BindingName;
+	}
+
+	bool IsReadable() const
+	{
+		return bIsReadable;
+	}
+
+	bool IsWritable() const
+	{
+		return bIsWritable;
+	}
+
+	bool IsValid() const
+	{
+		return BindingName.IsValid();
+	}
+
+	bool operator==(const FMVVMAvailableBinding& Other) const 
+	{
+		return BindingName == Other.BindingName && bIsReadable == Other.bIsReadable && bIsWritable == Other.bIsWritable;
+	}
+
+	bool operator!=(const FMVVMAvailableBinding& Other) const
+	{
+		return BindingName == Other.BindingName && bIsReadable == Other.bIsReadable && bIsWritable == Other.bIsWritable;
+	}
+
+	friend int32 GetTypeHash(const FMVVMAvailableBinding& Value)
+	{
+		return GetTypeHash(Value.BindingName);
+	}
+
+private:
+	UPROPERTY()
+	FMVVMBindingName BindingName;
+
+	UPROPERTY()
+	bool bIsReadable = false;
+
+	UPROPERTY()
+	bool bIsWritable = false;
+};
diff --git a/Engine/Plugins/Runtime/ModelViewViewModel/Source/ModelViewViewModel/Public/Types/MVVMBindingMode.h b/Engine/Plugins/Runtime/ModelViewViewModel/Source/ModelViewViewModel/Public/Types/MVVMBindingMode.h
new file mode 100644
--- /dev/null
+++ b/Engine/Plugins/Runtime/ModelViewViewModel/Source/ModelViewViewModel/Public/Types/MVVMBindingMode.h
@@ -0,0 +1,39 @@
+// Copyright Epic Games, Inc. All Rights Reserved.
+
+#pragma once
+
+#include "CoreMinimal.h"
+#include "MVVMBindingMode.generated.h"
+
+
+/** */
+UENUM()
+enum class EMVVMBindingMode : uint8
+{
+	OneTimeToDestination = 0,
+	OneWayToDestination,
+	TwoWay,
+	OneTimeToSource UMETA(Hidden),
+	OneWayToSource,
+};
+
+
+/** */
+UENUM()
+enum class EMVVMViewBindingUpdateMode : uint8
+{
+	/** Execute the binding as soon as the source value changes. */
+	Immediate,
+	///** Execute the binding at the end of the frame before drawing when the source value changes. */
+	//Delayed,
+	///** Always execute the binding at the end of the frame before drawing. */
+	//Debug,
+};
+
+
+namespace UE::MVVM
+{
+	UE_NODISCARD MODELVIEWVIEWMODEL_API bool IsForwardBinding(EMVVMBindingMode Mode);
+	UE_NODISCARD MODELVIEWVIEWMODEL_API bool IsBackwardBinding(EMVVMBindingMode Mode);
+	UE_NODISCARD MODELVIEWVIEWMODEL_API bool IsOneTimeBinding(EMVVMBindingMode Mode);
+}
\ No newline at end of file
