MVVM: Move the ModelViewViewModel plugin to //UE5/Main/

diff --git a/Engine/Plugins/Runtime/ModelViewViewModel/Source/ModelViewViewModel/Public/Types/MVVMFieldVariant.h b/Engine/Plugins/Runtime/ModelViewViewModel/Source/ModelViewViewModel/Public/Types/MVVMFieldVariant.h
new file mode 100644
--- /dev/null
+++ b/Engine/Plugins/Runtime/ModelViewViewModel/Source/ModelViewViewModel/Public/Types/MVVMFieldVariant.h
@@ -0,0 +1,183 @@
+// Copyright Epic Games, Inc. All Rights Reserved.
+
+#pragma once
+
+#include "CoreMinimal.h"
+#include "Misc/TVariant.h"
+#include <type_traits>
+
+class FProperty;
+class UFunction;
+
+namespace UE::MVVM
+{
+
+	/** */
+	template<bool bConst>
+	struct TMVVMFieldVariant
+	{
+	private:
+		friend TMVVMFieldVariant<!bConst>;
+		using FunctionType = std::conditional_t<bConst, const UFunction, UFunction>;
+		using PropertyType = std::conditional_t<bConst, const FProperty, FProperty>;
+		using VariantType = TVariant<FEmptyVariantState, PropertyType*, FunctionType*>;
+
+	public:
+		TMVVMFieldVariant() = default;
+
+		explicit TMVVMFieldVariant(PropertyType* InValue)
+		{
+			Binding = VariantType(TInPlaceType<PropertyType*>(), InValue);
+		}
+
+		explicit TMVVMFieldVariant(FunctionType* InValue)
+		{
+			Binding = VariantType(TInPlaceType<FunctionType*>(), InValue);
+		}
+
+		UE_NODISCARD bool IsProperty() const
+		{
+			return Binding.template IsType<PropertyType*>();
+		}
+
+		UE_NODISCARD PropertyType* GetProperty() const
+		{
+			return Binding.template Get<PropertyType*>();
+		}
+
+		void SetProperty(PropertyType* InValue)
+		{
+			Binding.template Set<PropertyType*>(InValue);
+		}
+
+		UE_NODISCARD bool IsFunction() const
+		{
+			return Binding.template IsType<FunctionType*>();
+		}
+
+		UE_NODISCARD FunctionType* GetFunction() const
+		{
+			return Binding.template Get<FunctionType*>();
+		}
+
+		void SetFunction(FunctionType* InValue)
+		{
+			Binding.template Set<FunctionType*>(InValue);
+		}
+
+		UE_NODISCARD FName GetName() const
+		{
+			if (IsProperty())
+			{
+				PropertyType* Property = GetProperty();
+				return Property ? Property->GetFName() : FName();
+			}
+			else if (IsFunction())
+			{
+				FunctionType* Function = GetFunction();
+				return Function ? Function->GetFName() : FName();
+			}
+			return FName();
+		}
+
+		UE_NODISCARD UStruct* GetOwner() const
+		{
+			if (IsProperty())
+			{
+				PropertyType* Property = GetProperty();
+				return Property ? Property->GetOwnerStruct() : nullptr;
+			}
+			else if (IsFunction())
+			{
+				FunctionType* Function = GetFunction();
+				return Function ? Function->GetOwnerClass() : nullptr;
+			}
+			return nullptr;
+		}
+
+		UE_NODISCARD bool IsEmpty() const
+		{
+			return Binding.template IsType<FEmptyVariantState>();
+		}
+
+		void Reset()
+		{
+			Binding = VariantType();
+		}
+
+		template<bool bOtherConst>
+		bool operator==(const TMVVMFieldVariant<bOtherConst>& B) const
+		{
+			if (Binding.GetIndex() != B.Binding.GetIndex())
+			{
+				return false;
+			}
+			if (IsEmpty())
+			{
+				return true;
+			}
+			if (IsProperty())
+			{
+				return GetProperty() == B.GetProperty();
+			}
+			return GetFunction() == B.GetFunction();
+		}
+
+		template<bool bOtherConst>
+		bool operator!=(const TMVVMFieldVariant<bOtherConst>& B) const
+		{
+			return !(*this == B);
+		}
+
+		friend int32 GetTypeHash(const TMVVMFieldVariant<bConst>& Variant)
+		{
+			if (Variant.IsProperty())
+			{
+				return GetTypeHash(Variant.GetProperty());
+			}
+			if (Variant.IsFunction())
+			{
+				return GetTypeHash(Variant.GetFunction());
+			}
+			return 0;
+		}
+
+	private:
+		VariantType Binding;
+	};
+
+	/** */
+	struct FMVVMFieldVariant : public TMVVMFieldVariant<false>
+	{
+	public:
+		using TMVVMFieldVariant<false>::TMVVMFieldVariant;
+		using TMVVMFieldVariant<false>::operator==;
+	};
+
+	/** */
+	struct FMVVMConstFieldVariant : public TMVVMFieldVariant<true>
+	{
+	public:
+		using TMVVMFieldVariant<true>::TMVVMFieldVariant;
+		using TMVVMFieldVariant<true>::operator==;
+
+		FMVVMConstFieldVariant(const FMVVMFieldVariant& OtherVariant)
+		{
+			if (OtherVariant.IsProperty())
+			{
+				SetProperty(OtherVariant.GetProperty());
+			}
+			else if (OtherVariant.IsFunction())
+			{
+				SetFunction(OtherVariant.GetFunction());
+			}
+		}
+
+		FMVVMConstFieldVariant& operator=(const FMVVMFieldVariant& OtherVariant)
+		{
+			*this = FMVVMConstFieldVariant(OtherVariant);
+			return *this;
+		}
+	};
+
+} //namespace
diff --git a/Engine/Plugins/Runtime/ModelViewViewModel/Source/ModelViewViewModel/Public/Types/MVVMViewModelCollection.h b/Engine/Plugins/Runtime/ModelViewViewModel/Source/ModelViewViewModel/Public/Types/MVVMViewModelCollection.h
new file mode 100644
--- /dev/null
+++ b/Engine/Plugins/Runtime/ModelViewViewModel/Source/ModelViewViewModel/Public/Types/MVVMViewModelCollection.h
@@ -0,0 +1,81 @@
+// Copyright Epic Games, Inc. All Rights Reserved.
+
+#pragma once
+
+#include "CoreMinimal.h"
+#include "Templates/SubclassOf.h"
+#include "Types/MVVMViewModelContext.h"
+#include "Types/MVVMViewModelContextInstance.h"
+
+#include "MVVMViewModelCollection.generated.h"
+
+class UMVVMViewModelBase;
+
+/** */
+USTRUCT()
+struct MODELVIEWVIEWMODEL_API FMVVMViewModelCollection
+{
+	GENERATED_BODY()
+
+public:
+	UMVVMViewModelBase* FindViewModelInstance(FMVVMViewModelContext Context) const;
+
+	bool AddInstance(FMVVMViewModelContext Context, UMVVMViewModelBase* ViewModel);
+	bool RemoveInstance(FMVVMViewModelContext Context);
+	int32 RemoveAllInstances(UMVVMViewModelBase* ViewModel);
+
+	void Reset();
+
+	FSimpleMulticastDelegate& OnCollectionChanged()
+	{
+		return OnCollectionChangedDelegate;
+	}
+
+private:
+	UPROPERTY()
+	mutable TArray<FMVVMViewModelContextInstance> ViewModelInstances;
+
+	FSimpleMulticastDelegate OnCollectionChangedDelegate;
+};
+
+
+/** */
+UCLASS(meta = (DisplayName = "MVVM View Model Collection Object"))
+class UMVVMViewModelCollectionObject : public UObject
+{
+	GENERATED_BODY()
+
+public:
+	UFUNCTION(BlueprintCallable, Category = "MVVM")
+	UMVVMViewModelBase* FindViewModelInstance(FMVVMViewModelContext Context) const
+	{
+		return ViewModelCollection.FindViewModelInstance(Context);
+	}
+	
+	UFUNCTION(BlueprintCallable, Category = "MVVM")
+	bool AddViewModelInstance(FMVVMViewModelContext Context, UMVVMViewModelBase* ViewModel)
+	{
+		return ViewModelCollection.AddInstance(Context, ViewModel);
+	}
+
+	UFUNCTION(BlueprintCallable, Category = "MVVM")
+	bool RemoveViewModel(FMVVMViewModelContext Context)
+	{
+		return ViewModelCollection.RemoveInstance(Context);
+	}
+
+	UFUNCTION(BlueprintCallable, Category = "MVVM")
+	int32 RemoveAllViewModelInstance(UMVVMViewModelBase* ViewModel)
+	{
+		return ViewModelCollection.RemoveAllInstances(ViewModel);
+	}
+
+	FSimpleMulticastDelegate& OnCollectionChanged()
+	{
+		return ViewModelCollection.OnCollectionChanged();
+	}
+
+private:
+	UPROPERTY(Transient)
+	FMVVMViewModelCollection ViewModelCollection;
+};
diff --git a/Engine/Plugins/Runtime/ModelViewViewModel/Source/ModelViewViewModel/Public/Types/MVVMViewModelContext.h b/Engine/Plugins/Runtime/ModelViewViewModel/Source/ModelViewViewModel/Public/Types/MVVMViewModelContext.h
new file mode 100644
--- /dev/null
+++ b/Engine/Plugins/Runtime/ModelViewViewModel/Source/ModelViewViewModel/Public/Types/MVVMViewModelContext.h
@@ -0,0 +1,32 @@
+// Copyright Epic Games, Inc. All Rights Reserved.
+
+#pragma once
+
+#include "CoreMinimal.h"
+#include "Templates/SubclassOf.h"
+
+#include "MVVMViewModelBase.h"
+
+#include "MVVMViewModelContext.generated.h"
+
+class UMVVMViewModelBase;
+
+/** */
+USTRUCT(BlueprintType)
+struct MODELVIEWVIEWMODEL_API FMVVMViewModelContext
+{
+	GENERATED_BODY()
+
+public:
+	UPROPERTY(BlueprintReadWrite, EditDefaultsOnly, Category="MVVM")
+	TSubclassOf<UMVVMViewModelBase> ContextClass;
+
+	UPROPERTY(BlueprintReadWrite, EditDefaultsOnly, Category="MVVM")
+	FName ContextName;
+
+public:
+	bool IsValid() const;
+	bool operator== (const FMVVMViewModelContext& Other) const;
+	bool IsCompatibleWith(const FMVVMViewModelContext& Other) const;
+	bool IsCompatibleWith(const UMVVMViewModelBase* Other) const;
+};
diff --git a/Engine/Plugins/Runtime/ModelViewViewModel/Source/ModelViewViewModel/Public/View/MVVMView.h b/Engine/Plugins/Runtime/ModelViewViewModel/Source/ModelViewViewModel/Public/View/MVVMView.h
new file mode 100644
--- /dev/null
+++ b/Engine/Plugins/Runtime/ModelViewViewModel/Source/ModelViewViewModel/Public/View/MVVMView.h
@@ -0,0 +1,89 @@
+// Copyright Epic Games, Inc. All Rights Reserved.
+
+#pragma once
+
+#include "CoreMinimal.h"
+#include "Templates/SubclassOf.h"
+#include "UObject/Object.h"
+
+#include "Extensions/UserWidgetExtension.h"
+#include "FieldNotification/FieldId.h"
+#include "FieldNotification/IFieldValueChanged.h"
+#include "Types/MVVMBindingName.h"
+
+#include "MVVMView.generated.h"
+
+class UMVVMViewClass;
+struct FMVVMViewClass_CompiledBinding;
+class UMVVMViewModelBase;
+class UWidget;
+
+
+/**
+ * Instance UMVVMClassExtension_View for the UUserWdiget
+ */
+UCLASS(Transient, DisplayName="MVVM View")
+class MODELVIEWVIEWMODEL_API UMVVMView : public UUserWidgetExtension
+{
+	GENERATED_BODY()
+
+public:
+	void ConstructView(const UMVVMViewClass* ClassExtension);
+
+	//~ Begin UUserWidgetExtension implementation
+	virtual void Initialize() override;
+	virtual void Destruct() override;
+
+// todo needed when we will support different type of update
+	//virtual bool RequiresTick() const override;
+	//virtual void Tick(const FGeometry& MyGeometry, float InDeltaTime) override;
+	//~ End UUserWidgetExtension implementation
+
+// todo a way to identify a binding from outside. maybe a unique name in the editor?
+	//UFUNCTION(BlueprintCallable, Category = "MVVM")
+	//void SetLibraryBindingEnabled(FGuid ViewModelId, FMVVMBindingName BindingName, bool bEnable);
+
+	//UFUNCTION(BlueprintCallable, Category = "MVVM")
+	//bool IsLibraryBindingEnabled(FGuid ViewModelId, FMVVMBindingName BindingName) const;
+
+	UFUNCTION(BlueprintCallable, Category = "MVVM")
+	void SetViewModel(FName ViewModelName, UMVVMViewModelBase* ViewModel);
+
+private:
+	void HandledLibraryBindingValueChanged(UObject* InViewModel, UE::FieldNotification::FFieldId InFieldId, int32 InCompiledBindingIndex) const;
+
+	void ExecuteLibraryBinding(const FMVVMViewClass_CompiledBinding& Binding) const;
+	void ExecuteLibraryBinding(const FMVVMViewClass_CompiledBinding& Binding, UObject* Source) const;
+
+	void EnableLibraryBinding(const FMVVMViewClass_CompiledBinding& Item, int32 BindingIndex);
+	void DisableLibraryBinding(const FMVVMViewClass_CompiledBinding& Item, int32 BindingIndex);
+	bool IsLibraryBindingEnabled(int32 InBindindIndex) const;
+
+	bool RegisterLibraryBinding(const FMVVMViewClass_CompiledBinding& Binding, int32 BindingIndex);
+	void UnregisterLibraryBinding(const FMVVMViewClass_CompiledBinding& Binding);
+
+	TScriptInterface<INotifyFieldValueChanged> FindSource(const FMVVMViewClass_CompiledBinding& Binding, bool bAllowNull) const;
+
+private:
+// todo support dynamic runtime binding.
+	/** Binding that are added dynamically at runtime. */
+	//TArray<FMVVMView_Binding> RegisteredDynamicBindings;
+
+	struct FRegisteredSource
+	{
+		TWeakObjectPtr<UObject> Source;
+		int32 Count = 0;
+	};
+
+	UPROPERTY(Transient)
+	TObjectPtr<const UMVVMViewClass> ClassExtension;
+
+	/** A list of all source and the amount of registered binding. */
+	TArray<FRegisteredSource> AllSources;
+
+	/** The binding that are enabled for the instance. */
+	TBitArray<> EnabledLibraryBindings;
+
+	/** Should log when a binding is executed */
+	bool bLogBinding = false;
+};
diff --git a/Engine/Plugins/Runtime/ModelViewViewModel/Source/ModelViewViewModel/Public/View/MVVMViewClass.h b/Engine/Plugins/Runtime/ModelViewViewModel/Source/ModelViewViewModel/Public/View/MVVMViewClass.h
new file mode 100644
--- /dev/null
+++ b/Engine/Plugins/Runtime/ModelViewViewModel/Source/ModelViewViewModel/Public/View/MVVMViewClass.h
@@ -0,0 +1,228 @@
+// Copyright Epic Games, Inc. All Rights Reserved.
+
+#pragma once
+
+#include "CoreMinimal.h"
+#include "Templates/SubclassOf.h"
+#include "UObject/Object.h"
+
+#include "Blueprint/UserWidget.h"
+#include "Bindings/MVVMCompiledBindingLibrary.h"
+#include "Extensions/WidgetBlueprintGeneratedClassExtension.h"
+#include "FieldNotification/FieldId.h"
+#include "Types/MVVMBindingMode.h"
+#include "Types/MVVMBindingName.h"
+#include "Types/MVVMViewModelContext.h"
+#include "View/MVVMView.h"
+
+#include "MVVMViewClass.generated.h"
+
+
+class UMVVMUserWidgetBinding;
+class UMVVMView;
+class UMVVMViewClass;
+class UMVVMViewModelBlueprintExtension;
+class UUserWidget;
+
+namespace UE::MVVM::Private
+{
+	struct FMVVMViewBlueprintCompiler;
+}
+
+/**
+ * Shared data to find or create a ViewModel at runtime.
+ */
+USTRUCT()
+struct FMVVMViewClass_SourceCreator
+{
+	GENERATED_BODY()
+
+	friend UE::MVVM::Private::FMVVMViewBlueprintCompiler;
+
+public:
+	MODELVIEWVIEWMODEL_API static FMVVMViewClass_SourceCreator MakeInstance(FName Name, UClass* NotifyFieldValueChangedClass);
+	MODELVIEWVIEWMODEL_API static FMVVMViewClass_SourceCreator MakeFieldPath(FName Name, UClass* NotifyFieldValueChangedClass, FMVVMVCompiledFieldPath FieldPath);
+	MODELVIEWVIEWMODEL_API static FMVVMViewClass_SourceCreator MakeGlobalContext(FName Name, FMVVMViewModelContext Context);
+
+	UObject* CreateInstance(const UMVVMViewClass* ViewClass, UMVVMView* View, UUserWidget* UserWidget) const;
+
+	UClass* GetSourceClass() const
+	{
+		return ExpectedSourceType.Get();
+	}
+
+	FName GetSourcePropertyName() const
+	{
+		return PropertyName;
+	}
+
+private:
+	/** Class type to create a source at runtime. */
+	UPROPERTY()
+	TSubclassOf<UObject> ExpectedSourceType;
+
+	/** Info to find the ViewModel instance at runtime. */
+	UPROPERTY()
+	FMVVMViewModelContext GlobalViewModelInstance;
+
+	/**
+	 * A resolvable path to retrieve the source instance at runtime.
+	 * It can be a path "Property = Object.Function.Object".
+	 * It can be a UFunction's name of a FProperty's name.
+	 */
+	UPROPERTY()
+	FMVVMVCompiledFieldPath FieldPath;
+
+	UPROPERTY()
+	FName PropertyName;
+	
+	UPROPERTY()
+	bool bCreateInstance = false;
+
+	UPROPERTY()
+	bool bOptional = false;
+};
+
+
+/**
+ * A compiled and shared binding for ViewModel<->View
+ */
+USTRUCT()
+struct MODELVIEWVIEWMODEL_API FMVVMViewClass_CompiledBinding
+{
+	GENERATED_BODY()
+
+	friend UE::MVVM::Private::FMVVMViewBlueprintCompiler;
+
+public:
+	/** @return The id for the FieldId on the source (if forward) or the destination (if backward). */
+	FMVVMVCompiledFieldId GetSourceFieldId() const
+	{
+		return FieldId;
+	}
+
+	/**
+	 * @return The property name of the object that contains the SourceFieldId.
+	 * It implements INotifyFieldValueChanged if it's not a One Time.
+	 */
+	FName GetSourceObjectPropertyName() const
+	{
+		return SourcePropertyName;
+	}
+
+	/** @return the binding. From source to destination (if forward) or from destination to source (if backward). */
+	const FMVVMVCompiledBinding& GetBinding() const
+	{
+		return Binding;
+	}
+
+	/** @return true if this Binding should be executed at initialization. */
+	bool NeedsExecutionAtInitialization() const
+	{
+		return (Flags & EBindingFlags::ForwardBinding) != 0;
+	}
+
+	/** @return true if this Binding should be executed once (at initialization) but should not be executed when the SourceFieldId value changes. */
+	bool IsOneTime() const
+	{
+		return (Flags & EBindingFlags::OneTime) != 0;
+	}
+
+	/** @return true if the binding is enabled by default. */
+	bool IsEnabledByDefault() const
+	{
+		return (Flags & EBindingFlags::EnabledByDefault) != 0;
+	}
+
+	/** @return a human readable version of the binding that can be use for debugging purposes. */
+	FString ToString() const;
+
+private:
+	UPROPERTY()
+	FMVVMVCompiledFieldId FieldId;
+
+	UPROPERTY()
+	FName SourcePropertyName;
+
+	UPROPERTY()
+	FMVVMVCompiledBinding Binding;
+
+	/** How the binding should be executed. */
+	UPROPERTY()
+	EMVVMViewBindingUpdateMode UpdateMode = EMVVMViewBindingUpdateMode::Immediate;
+
+	enum EBindingFlags
+	{
+		None = 0,
+		ForwardBinding = 1 << 0, // True when the binding goes from Source to Destination.
+		TwoWayBinding = 1 << 1, // The binding is one part of a 2 ways binding.
+		OneTime = 1 << 2,
+		EnabledByDefault = 1 << 3,
+	};
+
+	UPROPERTY()
+	uint8 Flags = EBindingFlags::None;
+};
+
+
+/**
+ * Shared between every instances of the same View class.
+ */
+UCLASS()
+class MODELVIEWVIEWMODEL_API UMVVMViewClass : public UWidgetBlueprintGeneratedClassExtension
+{
+	GENERATED_BODY()
+
+	friend UE::MVVM::Private::FMVVMViewBlueprintCompiler;
+
+public:
+	virtual void Initialize(UUserWidget* UserWidget) override;
+
+public:
+	/** Get the list of the needed ViewModel. */
+	const TArrayView<const FMVVMViewClass_SourceCreator> GetViewModelCreators() const
+	{
+		return MakeArrayView(SourceCreators);
+	}
+
+	/** Get the container of all the bindings. */
+	const FMVVMCompiledBindingLibrary& GetBindingLibrary() const
+	{
+		return BindingLibrary;
+	}
+
+	/**  */
+	const TArrayView<const FMVVMViewClass_CompiledBinding> GetCompiledBindings() const
+	{
+		return MakeArrayView(CompiledBindings);
+	}
+
+	/** */
+	const FMVVMViewClass_CompiledBinding& GetCompiledBinding(int32 Index) const
+	{
+		return CompiledBindings[Index];
+	}
+
+private:
+	/**  */
+	TArrayView<FMVVMViewClass_CompiledBinding> GetCompiledBindings()
+	{
+		return MakeArrayView(CompiledBindings);
+	}
+
+private:
+	/** Data to retrieve/create the sources (could be viewmodel, widget, ...). */
+	UPROPERTY()
+	TArray<FMVVMViewClass_SourceCreator> SourceCreators;
+
+	/** */
+	UPROPERTY()
+	TArray<FMVVMViewClass_CompiledBinding> CompiledBindings;
+
+	/** All the bindings shared between all the View instance. */
+	UPROPERTY()
+	FMVVMCompiledBindingLibrary BindingLibrary;
+
+	/** */
+	bool bLoaded = false;
+};
diff --git a/Engine/Plugins/Runtime/ModelViewViewModel/Source/ModelViewViewModelBlueprint/Private/Bindings/MVVMCompiledBindingLibraryCompiler.cpp b/Engine/Plugins/Runtime/ModelViewViewModel/Source/ModelViewViewModelBlueprint/Private/Bindings/MVVMCompiledBindingLibraryCompiler.cpp
new file mode 100644
--- /dev/null
+++ b/Engine/Plugins/Runtime/ModelViewViewModel/Source/ModelViewViewModelBlueprint/Private/Bindings/MVVMCompiledBindingLibraryCompiler.cpp
@@ -0,0 +1,734 @@
+// Copyright Epic Games, Inc. All Rights Reserved.
+
+#include "Bindings/MVVMCompiledBindingLibraryCompiler.h"
+
+#include "Bindings/MVVMBindingHelper.h"
+#include "Bindings/MVVMFieldPathHelper.h"
+#include "Engine/Engine.h"
+#include "MVVMSubsystem.h"
+#include "Types/MVVMObjectVariant.h"
+#include <limits>
+
+
+namespace UE::MVVM::Private
+{
+
+static const FName NAME_BlueprintGetter = "BlueprintGetter";
+
+
+/** */
+struct FRawFieldId
+{
+	UClass* NotifyFieldValudChangedClass;
+	UE::FieldNotification::FFieldId FieldId;
+
+	int32 LoadedFieldIdIndex = INDEX_NONE;
+	FCompiledBindingLibraryCompiler::FFieldIdHandle IdHandle;
+	FMVVMVCompiledFieldId CompiledFieldId;
+};
+
+
+/** */
+struct FRawField
+{
+	FMVVMConstFieldVariant Field;
+	bool bPropertyIsObjectProperty = false; // either the FProperty or the return value of the UFunction
+	bool bPropertyIsStructProperty = false;
+	int32 LoadedPropertyOrFunctionIndex = INDEX_NONE;
+
+public:
+	bool IsSameField(const FRawField& Other) const
+	{
+		return Field == Other.Field;
+	}
+};
+
+
+/** */
+struct FRawFieldPath
+{
+	TArray<int32> RawFieldIndexes;
+	bool bIsReadable = false;
+	bool bIsWritable = false;
+
+	FCompiledBindingLibraryCompiler::FFieldPathHandle PathHandle;
+	FMVVMVCompiledFieldPath CompiledFieldPath;
+
+public:
+	bool IsSameFieldPath(const FRawFieldPath& Other) const
+	{
+		return Other.RawFieldIndexes == RawFieldIndexes;
+	}
+};
+
+
+/** */
+struct FRawBinding
+{
+	TArray<FCompiledBindingLibraryCompiler::FFieldPathHandle> SourcePathHandles;
+	FCompiledBindingLibraryCompiler::FFieldPathHandle DestinationPathHandle;
+	FCompiledBindingLibraryCompiler::FFieldPathHandle ConversionFunctionPathHandle;
+
+	FCompiledBindingLibraryCompiler::FBindingHandle BindingHandle;
+	FMVVMVCompiledBinding CompiledBinding;
+
+public:
+	bool IsSameBinding(const FRawBinding& Binding) const
+	{
+		return Binding.SourcePathHandles == SourcePathHandles
+			&& Binding.DestinationPathHandle == DestinationPathHandle
+			&& Binding.ConversionFunctionPathHandle == ConversionFunctionPathHandle;
+	}
+};
+
+/** */
+class FCompiledBindingLibraryCompilerImpl
+{
+public:
+	TArray<FRawFieldId> FieldIds;
+	TArray<FRawField> Fields;
+	TArray<FRawFieldPath> FieldPaths;
+	TArray<FRawBinding> Bindings;
+	bool bCompiled = false;
+
+public:
+	int32 AddUniqueField(FMVVMConstFieldVariant InFieldVariant)
+	{
+		int32 FoundFieldPath = Fields.IndexOfByPredicate([InFieldVariant](const Private::FRawField& Other)
+			{
+				return Other.Field == InFieldVariant;
+			});
+		if (FoundFieldPath == INDEX_NONE)
+		{
+			FRawField RawField;
+			RawField.Field = InFieldVariant;
+			check(!InFieldVariant.IsEmpty());
+			const FProperty* FieldProperty = InFieldVariant.IsProperty() ? InFieldVariant.GetProperty() : BindingHelper::GetReturnProperty(InFieldVariant.GetFunction());
+			// FieldProperty can be null if it's a setter function
+			RawField.bPropertyIsObjectProperty = CastField<FObjectPropertyBase>(FieldProperty) != nullptr;
+			RawField.bPropertyIsStructProperty = CastField<FStructProperty>(FieldProperty) != nullptr;
+
+			FoundFieldPath = Fields.Add(RawField);
+		}
+		return FoundFieldPath;
+	}
+};
+
+} //namespace
+
+
+
+namespace UE::MVVM
+{
+
+
+int32 FCompiledBindingLibraryCompiler::FBindingHandle::IdGenerator = 0;
+int32 FCompiledBindingLibraryCompiler::FFieldPathHandle::IdGenerator = 0;
+int32 FCompiledBindingLibraryCompiler::FFieldIdHandle::IdGenerator = 0;
+
+
+/**
+ *
+ */
+FCompiledBindingLibraryCompiler::FCompiledBindingLibraryCompiler()
+	: Impl(MakePimpl<Private::FCompiledBindingLibraryCompilerImpl>())
+{
+
+}
+
+TValueOrError<FCompiledBindingLibraryCompiler::FFieldIdHandle, FString> FCompiledBindingLibraryCompiler::AddFieldId(TSubclassOf<UObject> SourceClass, FName FieldId)
+{
+	Impl->bCompiled = false;
+
+	if (!SourceClass->ImplementsInterface(UNotifyFieldValueChanged::StaticClass()))
+	{
+		return MakeError(FString::Printf(TEXT("'%s' doesn't implement the NotifyFieldValueChanged interface.")
+			, *SourceClass->GetName()));
+	}
+
+	const TScriptInterface<INotifyFieldValueChanged> ScriptObject = SourceClass->GetDefaultObject();
+	if (ensure(ScriptObject.GetInterface()))
+	{
+		UE::FieldNotification::FFieldId FoundFieldId = ScriptObject->GetFieldNotificationDescriptor().GetField(SourceClass, FieldId);
+		if (!FoundFieldId.IsValid())
+		{
+			return MakeError(FString::Printf(TEXT("The FieldNotify '%s' is not support by '%s'.")
+				, *FieldId.ToString()
+				, *SourceClass->GetName()));
+		}
+
+		int32 FoundFieldIdIndex = Impl->FieldIds.IndexOfByPredicate([FoundFieldId, SourceClass](const Private::FRawFieldId& Other)
+			{
+				return Other.FieldId == FoundFieldId && Other.NotifyFieldValudChangedClass == SourceClass.Get();
+			});
+		if (FoundFieldIdIndex == INDEX_NONE)
+		{
+			Private::FRawFieldId RawFieldId;
+			RawFieldId.NotifyFieldValudChangedClass = SourceClass.Get();
+			RawFieldId.FieldId = FoundFieldId;
+			RawFieldId.IdHandle = FFieldIdHandle::MakeHandle();
+			FoundFieldIdIndex = Impl->FieldIds.Add(MoveTemp(RawFieldId));
+		}
+
+		return MakeValue(Impl->FieldIds[FoundFieldIdIndex].IdHandle);
+	}
+
+	return MakeError(TEXT("Unexpected case with AddFieldId."));
+}
+
+
+TValueOrError<FCompiledBindingLibraryCompiler::FFieldPathHandle, FString> FCompiledBindingLibraryCompiler::AddFieldPath(TSubclassOf<UObject> InSourceClass, FStringView InFieldPath, bool bInRead)
+{
+	Impl->bCompiled = false;
+
+	TValueOrError<TArray<FMVVMFieldVariant>, FString> GeneratedField = FieldPathHelper::GenerateFieldPathList(InSourceClass, InFieldPath, bInRead);
+	if (GeneratedField.HasError())
+	{
+		return MakeError(GeneratedField.StealError());
+	}
+
+	return AddFieldPath(MakeArrayView(GeneratedField.GetValue()), bInRead);
+}
+
+
+TValueOrError<FCompiledBindingLibraryCompiler::FFieldPathHandle, FString> FCompiledBindingLibraryCompiler::AddFieldPath(TArrayView<FMVVMFieldVariant> InFieldPath, bool bInRead)
+{
+	Impl->bCompiled = false;
+
+	auto ValidateContainer = [](const FProperty* Property, bool bShouldBeInsideContainer, bool bIsObjectOrScriptStruct) -> FString
+	{	
+		const UStruct* OwnerStruct = Property->GetOwnerStruct();
+		if (OwnerStruct == nullptr)
+		{
+			return FString::Printf(TEXT("The field %s has an invalid owner struct."), *Property->GetName());
+		}
+
+		if (bShouldBeInsideContainer)
+		{
+			if (!Cast<UScriptStruct>(OwnerStruct) && !Cast<UClass>(OwnerStruct))
+			{
+				return FString::Printf(TEXT("The field %s doesn't have a valid owner for that path."), *Property->GetName());
+			}
+		}
+
+		if (bIsObjectOrScriptStruct)
+		{
+			if (const FObjectPropertyBase* ObjectProperty = CastField<const FObjectPropertyBase>(Property))
+			{
+				return FString();
+			}
+			else if (const FStructProperty* StructProperty = CastField<const FStructProperty>(Property))
+			{
+				if (StructProperty->HasGetter() || Property->HasMetaData(Private::NAME_BlueprintGetter))
+				{
+					return FString::Printf(TEXT("Property %s has getter accessor. Accessor not supported on FStructProperty since it would create a temporary structure and we would not able to return a valid container from that structure."), *StructProperty->GetName());
+				}
+				return FString();
+			}
+			return FString::Printf(TEXT("Field can only be object properties or struct properties. %s is a %s"), *Property->GetName(), *Property->GetClass()->GetName());
+		}
+
+		return FString();
+	};
+
+	TArray<int32> RawFieldIndexes;
+	RawFieldIndexes.Reserve(InFieldPath.Num());
+
+	for (int32 Index = 0; Index < InFieldPath.Num(); ++Index)
+	{
+		FMVVMConstFieldVariant FieldVariant = InFieldPath[Index];
+		const bool bIsLast = Index == InFieldPath.Num() - 1;
+		if (FieldVariant.IsProperty())
+		{
+			// They must all be readable except the last item if we are writing to the property.
+			if (bIsLast && !bInRead)
+			{
+				if (!BindingHelper::IsValidForDestinationBinding(FieldVariant.GetProperty()))
+				{
+					return MakeError(FString::Printf(TEXT("Property %s is not writable at runtime."), *FieldVariant.GetProperty()->GetName()));
+				}
+			}
+			else if (!BindingHelper::IsValidForSourceBinding(FieldVariant.GetProperty()))
+			{
+				return MakeError(FString::Printf(TEXT("Property %s is not readable at runtime."), *FieldVariant.GetProperty()->GetName()));
+			}
+
+			FString ValidatedStr = ValidateContainer(FieldVariant.GetProperty(), true, !bIsLast);
+			if (!ValidatedStr.IsEmpty())
+			{
+				return MakeError(ValidatedStr);
+			}
+
+			RawFieldIndexes.Add(Impl->AddUniqueField(FieldVariant));
+		}
+		else if (FieldVariant.IsFunction())
+		{
+			if (bIsLast && !bInRead)
+			{
+				if (!BindingHelper::IsValidForDestinationBinding(FieldVariant.GetFunction()))
+				{
+					return MakeError(FString::Printf(TEXT("Function %s is not writable at runtime."), *FieldVariant.GetFunction()->GetName()));
+				}
+			}
+			else if (!BindingHelper::IsValidForSourceBinding(FieldVariant.GetFunction()))
+			{
+				return MakeError(FString::Printf(TEXT("Function %s is not readable at runtime."), *FieldVariant.GetFunction()->GetName()));
+			}
+
+			if (bIsLast && !bInRead)
+			{
+				const FProperty* FirstProperty = BindingHelper::GetFirstArgumentProperty(FieldVariant.GetFunction());
+				ValidateContainer(FirstProperty, false, bIsLast);
+				RawFieldIndexes.Add(Impl->AddUniqueField(FieldVariant));
+			}
+			else
+			{
+				const FProperty* ReturnProperty = BindingHelper::GetReturnProperty(FieldVariant.GetFunction());
+				ValidateContainer(ReturnProperty, false, bIsLast);
+				RawFieldIndexes.Add(Impl->AddUniqueField(FieldVariant));
+			}
+		}
+		else
+		{
+			return MakeError(TEXT("There is an invalid field in the field path."));
+		}
+	}
+
+	int32 FoundFieldPath = Impl->FieldPaths.IndexOfByPredicate([&RawFieldIndexes](const Private::FRawFieldPath& Other)
+		{
+			return Other.RawFieldIndexes == RawFieldIndexes;
+		});
+	if (FoundFieldPath != INDEX_NONE)
+	{
+		Impl->FieldPaths[FoundFieldPath].bIsReadable = Impl->FieldPaths[FoundFieldPath].bIsReadable || bInRead;
+		Impl->FieldPaths[FoundFieldPath].bIsWritable = Impl->FieldPaths[FoundFieldPath].bIsWritable || !bInRead;
+		return MakeValue(Impl->FieldPaths[FoundFieldPath].PathHandle);
+	}
+
+	Private::FRawFieldPath RawFieldPath;
+	RawFieldPath.RawFieldIndexes = RawFieldIndexes;
+	RawFieldPath.PathHandle = FFieldPathHandle::MakeHandle();
+	RawFieldPath.bIsReadable = bInRead;
+	RawFieldPath.bIsWritable = !bInRead;
+	FoundFieldPath = Impl->FieldPaths.Add(MoveTemp(RawFieldPath));
+
+	return MakeValue(Impl->FieldPaths[FoundFieldPath].PathHandle);
+}
+
+
+TValueOrError<FCompiledBindingLibraryCompiler::FFieldPathHandle, FString> FCompiledBindingLibraryCompiler::AddObjectFieldPath(TSubclassOf<UObject> InSourceClass, FStringView InFieldPath, UClass* ExpectedType, bool bInRead)
+{
+	Impl->bCompiled = false;
+
+	check(ExpectedType);
+
+	TValueOrError<TArray<FMVVMFieldVariant>, FString> GeneratedField = FieldPathHelper::GenerateFieldPathList(InSourceClass, InFieldPath, bInRead);
+	if (GeneratedField.HasError())
+	{
+		return MakeError(GeneratedField.StealError());
+	}
+	if (GeneratedField.GetValue().Num() == 0)
+	{
+		return MakeError(FString::Printf(TEXT("The field does not returned a '%s'."), *ExpectedType->GetName()));
+	}
+
+	const FObjectPropertyBase* ObjectPropertyBase = nullptr;
+	if (GeneratedField.GetValue().Last().IsProperty())
+	{
+		ObjectPropertyBase = CastField<const FObjectPropertyBase>(GeneratedField.GetValue().Last().GetProperty());
+	}
+	else if (GeneratedField.GetValue().Last().IsFunction())
+	{
+		ObjectPropertyBase = CastField<const FObjectPropertyBase>(BindingHelper::GetReturnProperty(GeneratedField.GetValue().Last().GetFunction()));
+	}
+
+	if (ObjectPropertyBase == nullptr)
+	{
+		return MakeError(FString::Printf(TEXT("The field does not returned a '%s'."), *ExpectedType->GetName()));
+	}
+	if (ObjectPropertyBase->PropertyClass == nullptr || !ExpectedType->IsChildOf(ObjectPropertyBase->PropertyClass))
+	{
+		return MakeError(FString::Printf(TEXT("The field does not returned a '%s'."), *ExpectedType->GetName()));
+	}
+
+	return AddFieldPath(MakeArrayView(GeneratedField.GetValue()), bInRead);
+}
+
+
+TValueOrError<FCompiledBindingLibraryCompiler::FFieldPathHandle, FString> FCompiledBindingLibraryCompiler::AddConversionFunctionFieldPath(TSubclassOf<UObject> SourceClass, FStringView FieldPath)
+{
+	return MakeError(TEXT("Not implemented yet"));
+}
+
+
+TValueOrError<FCompiledBindingLibraryCompiler::FBindingHandle, FString> FCompiledBindingLibraryCompiler::AddBinding(FFieldPathHandle InSourceHandle, FFieldPathHandle InDestinationHandle)
+{
+	return AddBinding(InSourceHandle, InDestinationHandle, FFieldPathHandle());
+}
+
+
+TValueOrError<FCompiledBindingLibraryCompiler::FBindingHandle, FString> FCompiledBindingLibraryCompiler::AddBinding(FFieldPathHandle InSourceHandle, FFieldPathHandle InDestinationHandle, FFieldPathHandle InConversionFunctionHandle)
+{
+	return AddBinding(MakeArrayView(&InSourceHandle, 1), InDestinationHandle, InConversionFunctionHandle);
+}
+
+
+TValueOrError<FCompiledBindingLibraryCompiler::FBindingHandle, FString> FCompiledBindingLibraryCompiler::AddBinding(TArrayView<FFieldPathHandle> InSourceHandles, FFieldPathHandle InDestinationHandle, FFieldPathHandle InConversionFunctionHandle)
+{
+	Impl->bCompiled = false;
+
+	UMVVMSubsystem::FConstDirectionalBindingArgs DirectionBindingArgs;
+
+	ensureMsgf(InSourceHandles.Num() == 1, TEXT("Conversion function with more than one arguements is not yet supported."));
+
+	{
+		FFieldPathHandle FirstSourceHandle = InSourceHandles[0];
+		const int32 FoundSourceFieldPath = Impl->FieldPaths.IndexOfByPredicate([FirstSourceHandle](const Private::FRawFieldPath& Other)
+			{
+				return Other.PathHandle == FirstSourceHandle;
+			});
+		if (FoundSourceFieldPath == INDEX_NONE)
+		{
+			return MakeError(TEXT("The source handle is invalid."));
+		}
+
+		Private::FRawFieldPath& SourceRawFieldPath = Impl->FieldPaths[FoundSourceFieldPath];
+		if (!SourceRawFieldPath.bIsReadable)
+		{
+			return MakeError(TEXT("The source handle was not constructed as a readable path."));
+		}
+		if (SourceRawFieldPath.RawFieldIndexes.Num() == 0)
+		{
+			return MakeError(TEXT("The source handle was not registered correctly."));
+		}
+
+		Private::FRawField& RawField = Impl->Fields[SourceRawFieldPath.RawFieldIndexes.Last()];
+		if (RawField.Field.IsEmpty())
+		{
+			return MakeError(TEXT("The source handle was not registered correctly."));
+		}
+
+		DirectionBindingArgs.SourceBinding = RawField.Field;
+	}
+
+	{
+		const int32 FoundDestinationFieldPath = Impl->FieldPaths.IndexOfByPredicate([InDestinationHandle](const Private::FRawFieldPath& Other)
+			{
+				return Other.PathHandle == InDestinationHandle;
+			});
+		if (FoundDestinationFieldPath == INDEX_NONE)
+		{
+			return MakeError(TEXT("The destination handle is invalid."));
+		}
+
+		Private::FRawFieldPath& DestinationRawFieldPath = Impl->FieldPaths[FoundDestinationFieldPath];
+		if (!DestinationRawFieldPath.bIsWritable)
+		{
+			return MakeError(TEXT("The destination handle was not constructed as a writable path."));
+		}
+		if (DestinationRawFieldPath.RawFieldIndexes.Num() == 0)
+		{
+			return MakeError(TEXT("The destination handle was not registered correctly."));
+		}
+
+		Private::FRawField& RawField = Impl->Fields[DestinationRawFieldPath.RawFieldIndexes.Last()];
+		if (RawField.Field.IsEmpty())
+		{
+			return MakeError(TEXT("The destination handle was not registered correctly."));
+		}
+
+		DirectionBindingArgs.DestinationBinding = RawField.Field;
+	}
+
+	if (InConversionFunctionHandle.IsValid())
+	{
+		const int32 FoundFunctionFieldPath = Impl->FieldPaths.IndexOfByPredicate([InConversionFunctionHandle](const Private::FRawFieldPath& Other)
+			{
+				return Other.PathHandle == InConversionFunctionHandle;
+			});
+		if (FoundFunctionFieldPath == INDEX_NONE)
+		{
+			return MakeError(TEXT("The function handle is invalid."));
+		}
+
+		Private::FRawFieldPath& ConversionFunctinRawFieldPath = Impl->FieldPaths[FoundFunctionFieldPath];
+		if (ConversionFunctinRawFieldPath.RawFieldIndexes.Num() == 0)
+		{
+			return MakeError(TEXT("The function handle was not registered as a function."));
+		}
+
+		Private::FRawField& RawField = Impl->Fields[ConversionFunctinRawFieldPath.RawFieldIndexes.Last()];
+		if (!RawField.Field.IsFunction())
+		{
+			return MakeError(TEXT("The function handle was not registered as a function."));
+		}
+
+		DirectionBindingArgs.ConversionFunction = RawField.Field.GetFunction();
+	}
+
+	TValueOrError<bool, FString> IsValidBinding = GEngine->GetEngineSubsystem<UMVVMSubsystem>()->IsBindingValid(DirectionBindingArgs);
+	if (IsValidBinding.HasError())
+	{
+		return MakeError(IsValidBinding.StealError());
+	}
+
+	Private::FRawBinding NewBinding;
+	NewBinding.SourcePathHandles = InSourceHandles;
+	NewBinding.DestinationPathHandle = InDestinationHandle;
+	NewBinding.ConversionFunctionPathHandle = InConversionFunctionHandle;
+	const int32 FoundIndex = Impl->Bindings.IndexOfByPredicate([&NewBinding](const Private::FRawBinding& Binding)
+		{
+			return NewBinding.IsSameBinding(Binding);
+		});
+
+	if (FoundIndex != INDEX_NONE)
+	{
+		return MakeValue(Impl->Bindings[FoundIndex].BindingHandle);
+	}
+	else
+	{
+		FCompiledBindingLibraryCompiler::FBindingHandle ResultBindingHandle = FBindingHandle::MakeHandle();
+		NewBinding.BindingHandle = ResultBindingHandle;
+		Impl->Bindings.Add(MoveTemp(NewBinding));
+		return MakeValue(ResultBindingHandle);
+	}
+}
+
+
+TValueOrError<FCompiledBindingLibraryCompiler::FCompileResult, FString> FCompiledBindingLibraryCompiler::Compile()
+{
+	Impl->bCompiled = false;
+
+	struct FCompiledClassInfo
+	{
+		TArray<int32> RawFieldIndex;
+		TArray<int32> RawFieldIdIndex;
+	};
+
+	// create the list of UClass
+	TMap<const UStruct*, FCompiledClassInfo> MapOfFieldInClass;
+	{
+		for (int32 Index = 0; Index < Impl->Fields.Num(); ++Index)
+		{
+			const Private::FRawField& RawField = Impl->Fields[Index];
+			check(!RawField.Field.IsEmpty());
+
+			const UStruct* Owner = RawField.Field.IsProperty() ? RawField.Field.GetProperty()->GetOwnerStruct() : RawField.Field.GetFunction()->GetOwnerClass();
+			check(Owner);
+			FCompiledClassInfo& ClassInfo = MapOfFieldInClass.FindOrAdd(Owner);
+
+			// Test if the Field is not there more than one
+			{
+				FMVVMConstFieldVariant FieldToTest = RawField.Field;
+				const TArray<Private::FRawField>& ListOfFields = Impl->Fields;
+				const bool bContains = ClassInfo.RawFieldIndex.ContainsByPredicate([FieldToTest, &ListOfFields](int32 OtherIndex)
+					{
+						return ListOfFields[OtherIndex].Field == FieldToTest;
+					});
+				check(!bContains);
+			}
+
+			ClassInfo.RawFieldIndex.Add(Index);
+		}
+	}
+	{
+		for (int32 Index = 0; Index < Impl->FieldIds.Num(); ++Index)
+		{
+			const Private::FRawFieldId& RawFieldId = Impl->FieldIds[Index];
+			check(RawFieldId.FieldId.IsValid());
+			check(RawFieldId.NotifyFieldValudChangedClass);
+
+			FCompiledClassInfo& ClassInfo = MapOfFieldInClass.FindOrAdd(RawFieldId.NotifyFieldValudChangedClass);
+
+			// Test if the Field is not there more than one
+			{
+				UE::FieldNotification::FFieldId FieldIdToTest = RawFieldId.FieldId;
+				const TArray<Private::FRawFieldId>& ListOfFieldIds = Impl->FieldIds;
+				const bool bContains = ClassInfo.RawFieldIdIndex.ContainsByPredicate([FieldIdToTest, &ListOfFieldIds](int32 OtherIndex)
+					{
+						return ListOfFieldIds[OtherIndex].FieldId == FieldIdToTest;
+					});
+				check(!bContains);
+			}
+
+			ClassInfo.RawFieldIdIndex.Add(Index);
+		}
+	}
+
+
+	// Todo optimize that list to group common type. ie UWidget::ToolTip == UProgressBar::ToolTip. We can merge UWidget in UProgressBar.
+	//The difficulty is with type that may be not loaded at runtime and would create runtime issue with type that would be loaded otherwise.
+
+	FCompileResult Result;
+
+	// Create FMVVMCompiledBindingLibrary::CompiledFields and FMVVMCompiledBindingLibrary::CompiledFieldNames
+	int32 TotalNumberOfProperties = 0;
+	int32 TotalNumberOfFunctions = 0;
+	int32 TotalNumberOfFieldIds = 0;
+	for (TPair<const UStruct*, FCompiledClassInfo>& StructCompiledFields : MapOfFieldInClass)
+	{
+		FMVVMVCompiledFields CompiledFields;
+		CompiledFields.ClassOrScriptStruct = StructCompiledFields.Key;
+		check(StructCompiledFields.Key);
+
+		TArray<FName> PropertyNames;
+		TArray<FName> FunctionNames;
+		TArray<FName> FieldIdNames;
+		for (const int32 FieldIndex : StructCompiledFields.Value.RawFieldIndex)
+		{
+			check(Impl->Fields.IsValidIndex(FieldIndex));
+			Private::FRawField& RawField = Impl->Fields[FieldIndex];
+			const FMVVMConstFieldVariant& Field = RawField.Field;
+
+			if (Field.IsProperty())
+			{
+				Result.Library.LoadedProperties.Add(const_cast<FProperty*>(Field.GetProperty()));
+				PropertyNames.Add(Field.GetName());
+				RawField.LoadedPropertyOrFunctionIndex = TotalNumberOfProperties;
+				++TotalNumberOfProperties;
+			}
+			else
+			{
+				check(Field.IsFunction());
+				Result.Library.LoadedFunctions.Add(const_cast<UFunction*>(Field.GetFunction()));
+				FunctionNames.Add(Field.GetName());
+				RawField.LoadedPropertyOrFunctionIndex = TotalNumberOfFunctions;
+				++TotalNumberOfFunctions;
+			}
+		}
+
+		for (const int32 FieldIdIndex : StructCompiledFields.Value.RawFieldIdIndex)
+		{
+			check(Impl->FieldIds.IsValidIndex(FieldIdIndex));
+			Private::FRawFieldId& RawFieldId = Impl->FieldIds[FieldIdIndex];
+			
+			Result.Library.LoadedFieldIds.Add(RawFieldId.FieldId);
+			FieldIdNames.Add(RawFieldId.FieldId.GetName());
+			RawFieldId.LoadedFieldIdIndex = TotalNumberOfFieldIds;
+			++TotalNumberOfFieldIds;
+		}
+
+		if (PropertyNames.Num() > std::numeric_limits<FMVVMVCompiledBinding::IndexType>::max())
+		{
+			return MakeError(FString::Printf(TEXT("There are too many properties binded to struct '%s'"), *StructCompiledFields.Key->GetName()));
+		}
+		CompiledFields.NumberOfProperties = static_cast<int16>(PropertyNames.Num());
+
+		if (FunctionNames.Num() > std::numeric_limits<FMVVMVCompiledBinding::IndexType>::max())
+		{
+			return MakeError(FString::Printf(TEXT("There are too many functions binded to struct '%s'"), *StructCompiledFields.Key->GetName()));
+		}
+		CompiledFields.NumberOfFunctions = static_cast<int16>(FunctionNames.Num());
+
+		if (FieldIdNames.Num() > std::numeric_limits<FMVVMVCompiledBinding::IndexType>::max())
+		{
+			return MakeError(FString::Printf(TEXT("There are too many field ids binded to struct '%s'"), *StructCompiledFields.Key->GetName()));
+		}
+		CompiledFields.NumberOfFieldIds = static_cast<int16>(FieldIdNames.Num());
+
+		int32 LibraryStartIndex = Result.Library.CompiledFieldNames.Num();
+		if (LibraryStartIndex > std::numeric_limits<FMVVMVCompiledBinding::IndexType>::max())
+		{
+			return MakeError(FString::Printf(TEXT("There are too many properties and functions binded for the library")));
+		}
+		CompiledFields.LibraryStartIndex = static_cast<int16>(LibraryStartIndex);
+			
+		Result.Library.CompiledFieldNames.Append(PropertyNames);
+		PropertyNames.Reset();
+		Result.Library.CompiledFieldNames.Append(FunctionNames);
+		FunctionNames.Reset();
+		Result.Library.CompiledFieldNames.Append(FieldIdNames);
+		FieldIdNames.Reset();
+		if (Result.Library.CompiledFieldNames.Num() > std::numeric_limits<FMVVMVCompiledBinding::IndexType>::max())
+		{
+			return MakeError(FString::Printf(TEXT("There are too many properties binded for the library")));
+		}
+
+		Result.Library.CompiledFields.Add(CompiledFields);
+
+		check(Result.Library.LoadedProperties.Num() + Result.Library.LoadedFunctions.Num() + Result.Library.LoadedFieldIds.Num()  == Result.Library.CompiledFieldNames.Num());
+		check(Result.Library.LoadedProperties.Num() == TotalNumberOfProperties);
+		check(Result.Library.LoadedFunctions.Num() == TotalNumberOfFunctions);
+		check(Result.Library.LoadedFieldIds.Num() == TotalNumberOfFieldIds);
+	}
+
+	// Create FMVVMCompiledBindingLibrary::FieldPaths
+	for (Private::FRawFieldPath& FieldPath : Impl->FieldPaths)
+	{
+		FieldPath.CompiledFieldPath.CompiledBindingLibraryId = Result.Library.CompiledBindingLibraryId;
+		FieldPath.CompiledFieldPath.StartIndex = INDEX_NONE;
+		FieldPath.CompiledFieldPath.Num = FieldPath.RawFieldIndexes.Num();
+		if (FieldPath.RawFieldIndexes.Num())
+		{
+			FieldPath.CompiledFieldPath.StartIndex = Result.Library.FieldPaths.Num();
+			for (const int32 RawFieldIndex : FieldPath.RawFieldIndexes)
+			{
+				const Private::FRawField& RawField = Impl->Fields[RawFieldIndex];
+				check(!RawField.Field.IsEmpty());
+
+				FMVVMCompiledLoadedPropertyOrFunctionIndex FieldIndex;
+				FieldIndex.Index = RawField.LoadedPropertyOrFunctionIndex;
+				FieldIndex.bIsObjectProperty = RawField.bPropertyIsObjectProperty;
+				FieldIndex.bIsScriptStructProperty = RawField.bPropertyIsStructProperty;
+				FieldIndex.bIsProperty = RawField.Field.IsProperty();
+				Result.Library.FieldPaths.Add(FieldIndex);
+
+				if (FieldIndex.bIsProperty)
+				{
+					check(Result.Library.LoadedProperties.IsValidIndex(FieldIndex.Index));
+				}
+				else
+				{
+					check(Result.Library.LoadedFunctions.IsValidIndex(FieldIndex.Index));
+				}
+			}
+		}
+
+		Result.FieldPaths.Add(FieldPath.PathHandle, FieldPath.CompiledFieldPath);
+	}
+
+	// Create FieldId
+	for (Private::FRawFieldId& FieldId: Impl->FieldIds)
+	{
+		FieldId.CompiledFieldId.CompiledBindingLibraryId = Result.Library.CompiledBindingLibraryId;
+		FieldId.CompiledFieldId.FieldIdIndex = FieldId.LoadedFieldIdIndex;
+
+		Result.FieldIds.Add(FieldId.IdHandle, FieldId.CompiledFieldId);
+	}
+
+	auto GetCompiledFiledPath = [this](const FFieldPathHandle Handle)
+	{
+		const Private::FRawFieldPath* FoundBinding = Impl->FieldPaths.FindByPredicate([Handle](const Private::FRawFieldPath& Other)
+			{
+				return Other.PathHandle == Handle;
+			});
+		if (FoundBinding)
+		{
+			return FoundBinding->CompiledFieldPath;
+		}
+		return FMVVMVCompiledFieldPath();
+	};
+
+
+	// Create the requested FMVVMVCompiledBinding
+	for (Private::FRawBinding& Binding : Impl->Bindings)
+	{
+		Binding.CompiledBinding.CompiledBindingLibraryId = Result.Library.CompiledBindingLibraryId;
+		check(Binding.CompiledBinding.CompiledBindingLibraryId.IsValid());
+
+		Binding.CompiledBinding.SourceFieldPath = GetCompiledFiledPath(Binding.SourcePathHandles[0]);
+		check(Binding.CompiledBinding.SourceFieldPath.IsValid());
+
+		Binding.CompiledBinding.DestinationFieldPath = GetCompiledFiledPath(Binding.DestinationPathHandle);
+		check(Binding.CompiledBinding.DestinationFieldPath.IsValid());
+
+		Binding.CompiledBinding.ConversionFunctionFieldPath = GetCompiledFiledPath(Binding.ConversionFunctionPathHandle);
+
+		Result.Bindings.Add(Binding.BindingHandle, Binding.CompiledBinding);
+	}
+
+	Impl->bCompiled = true;
+	return MakeValue(MoveTemp(Result));
+}
+
+} //namespace
diff --git a/Engine/Plugins/Runtime/ModelViewViewModel/Source/ModelViewViewModelBlueprint/Private/MVVMBlueprintView.cpp b/Engine/Plugins/Runtime/ModelViewViewModel/Source/ModelViewViewModelBlueprint/Private/MVVMBlueprintView.cpp
new file mode 100644
--- /dev/null
+++ b/Engine/Plugins/Runtime/ModelViewViewModel/Source/ModelViewViewModelBlueprint/Private/MVVMBlueprintView.cpp
@@ -0,0 +1,188 @@
+// Copyright Epic Games, Inc. All Rights Reserved.
+
+#include "MVVMBlueprintView.h"
+
+#include "MVVMWidgetBlueprintExtension_View.h"
+#include "Blueprint/WidgetTree.h"
+#include "Kismet2/BlueprintEditorUtils.h"
+#include "UObject/UnrealType.h"
+
+
+FMVVMBlueprintViewModelContext* UMVVMBlueprintView::FindViewModel(FGuid ViewModelId)
+{
+	return AvailableViewModels.FindByPredicate([ViewModelId](const FMVVMBlueprintViewModelContext& Other)
+		{
+			return Other.GetViewModelId() == ViewModelId;
+		});
+}
+
+const FMVVMBlueprintViewModelContext* UMVVMBlueprintView::FindViewModel(FGuid ViewModelId) const
+{
+	return const_cast<UMVVMBlueprintView*>(this)->FindViewModel(ViewModelId);
+}
+
+void UMVVMBlueprintView::AddViewModel(const FMVVMBlueprintViewModelContext& NewContext)
+{
+	AvailableViewModels.Add(NewContext);
+	FBlueprintEditorUtils::MarkBlueprintAsStructurallyModified(GetOuterUMVVMWidgetBlueprintExtension_View()->GetWidgetBlueprint());
+	OnViewModelsUpdated.Broadcast();
+}
+
+
+void UMVVMBlueprintView::RemoveViewModel(FGuid ViewModelId)
+{
+	int32 Count = AvailableViewModels.RemoveAll([ViewModelId](const FMVVMBlueprintViewModelContext& VM)
+		{
+			return VM.GetViewModelId() == ViewModelId;
+		});
+
+	if (Count > 0)
+	{
+		FBlueprintEditorUtils::MarkBlueprintAsStructurallyModified(GetOuterUMVVMWidgetBlueprintExtension_View()->GetWidgetBlueprint());
+		OnViewModelsUpdated.Broadcast();
+	}
+}
+
+void UMVVMBlueprintView::RemoveViewModels(const TArrayView<FGuid> ViewModelIds)
+{
+	int32 Count = 0;
+	for (const FGuid& ViewModelId : ViewModelIds)
+	{
+		Count += AvailableViewModels.RemoveAll([ViewModelId](const FMVVMBlueprintViewModelContext& VM)
+			{
+				return VM.GetViewModelId() == ViewModelId;
+			});
+	}
+
+	if (Count > 0)
+	{
+		FBlueprintEditorUtils::MarkBlueprintAsStructurallyModified(GetOuterUMVVMWidgetBlueprintExtension_View()->GetWidgetBlueprint());
+		OnViewModelsUpdated.Broadcast();
+	}
+}
+
+void UMVVMBlueprintView::SetViewModels(const TArray<FMVVMBlueprintViewModelContext>& ViewModelContexts)
+{
+	AvailableViewModels = ViewModelContexts;
+
+	FBlueprintEditorUtils::MarkBlueprintAsStructurallyModified(GetOuterUMVVMWidgetBlueprintExtension_View()->GetWidgetBlueprint());
+	OnViewModelsUpdated.Broadcast();
+}
+
+const FMVVMBlueprintViewBinding* UMVVMBlueprintView::FindBinding(const UWidget* Widget, const FProperty* Property) const
+{
+	return FindBinding(Widget->GetFName(), FMVVMBindingName(Property->GetFName()));
+}
+
+FMVVMBlueprintViewBinding* UMVVMBlueprintView::FindBinding(const UWidget* Widget, const FProperty* Property)
+{
+	return FindBinding(Widget->GetFName(), FMVVMBindingName(Property->GetFName()));
+}
+
+const FMVVMBlueprintViewBinding* UMVVMBlueprintView::FindBinding(FName WidgetName, FMVVMBindingName BindingName) const
+{
+	return Bindings.FindByPredicate([WidgetName, BindingName](const FMVVMBlueprintViewBinding& Binding)
+		{
+			return Binding.WidgetPath.WidgetName == WidgetName &&
+				Binding.WidgetPath.GetBindingName() == BindingName;
+		});
+}
+
+FMVVMBlueprintViewBinding* UMVVMBlueprintView::FindBinding(FName WidgetName, FMVVMBindingName BindingName)
+{
+	return Bindings.FindByPredicate([WidgetName, BindingName](const FMVVMBlueprintViewBinding& Binding)
+		{
+			return Binding.WidgetPath.WidgetName == WidgetName &&
+				Binding.WidgetPath.GetBindingName() == BindingName;
+		});
+}
+
+void UMVVMBlueprintView::RemoveBindingAt(int32 Index)
+{
+	Bindings.RemoveAt(Index);
+	OnBindingsUpdated.Broadcast();
+}
+
+void UMVVMBlueprintView::RemoveBinding(const FMVVMBlueprintViewBinding* Binding)
+{
+	int32 Index = 0;
+	for (; Index < Bindings.Num(); ++Index)
+	{
+		if (&Bindings[Index] == Binding)
+		{
+			break;
+		}
+	}
+
+	Bindings.RemoveAt(Index);
+	OnBindingsUpdated.Broadcast();
+}
+
+FMVVMBlueprintViewBinding& UMVVMBlueprintView::AddBinding(const UWidget* Widget, const FProperty* Property)
+{
+	FMVVMBlueprintViewBinding& NewBinding = Bindings.AddDefaulted_GetRef();
+	NewBinding.WidgetPath.WidgetName = Widget->GetFName();
+	NewBinding.WidgetPath.SetBindingName(Property->GetFName());
+
+	OnBindingsUpdated.Broadcast();
+	return NewBinding;
+}
+
+FMVVMBlueprintViewBinding& UMVVMBlueprintView::AddDefaultBinding()
+{
+	FMVVMBlueprintViewBinding& NewBinding = Bindings.AddDefaulted_GetRef();
+
+	OnBindingsUpdated.Broadcast();
+	return NewBinding;
+}
+
+FMVVMBlueprintViewBinding* UMVVMBlueprintView::GetBindingAt(int32 Index)
+{
+	if (Bindings.IsValidIndex(Index))
+	{
+		return &Bindings[Index];
+	}
+	return nullptr;
+}
+
+const FMVVMBlueprintViewBinding* UMVVMBlueprintView::GetBindingAt(int32 Index) const
+{
+	if (Bindings.IsValidIndex(Index))
+	{
+		return &Bindings[Index];
+	}
+	return nullptr;
+}
+
+#if WITH_EDITOR
+void UMVVMBlueprintView::PostEditChangeProperty(FPropertyChangedEvent& PropertyChangedEvent)
+{
+	Super::PostEditChangeProperty(PropertyChangedEvent);
+	if (PropertyChangedEvent.Property && PropertyChangedEvent.Property->GetFName() == GET_MEMBER_NAME_CHECKED(UMVVMBlueprintView, Bindings))
+	{
+		OnBindingsUpdated.Broadcast();
+	}
+	if (PropertyChangedEvent.Property && PropertyChangedEvent.Property->GetFName() == GET_MEMBER_NAME_CHECKED(UMVVMBlueprintView, AvailableViewModels))
+	{
+		OnViewModelsUpdated.Broadcast();
+	}
+}
+
+void UMVVMBlueprintView::WidgetRenamed(FName OldObjectName, FName NewObjectName)
+{
+	bool bRenamed = false;
+	for (FMVVMBlueprintViewBinding& Binding : Bindings)
+	{
+		if (Binding.WidgetPath.WidgetName == OldObjectName)
+		{
+			Binding.WidgetPath.WidgetName = NewObjectName;
+			bRenamed = true;
+		}
+	}
+
+	if (bRenamed)
+	{
+		OnBindingsUpdated.Broadcast();
+	}
+}
+#endif
diff --git a/Engine/Plugins/Runtime/ModelViewViewModel/Source/ModelViewViewModelBlueprint/Private/MVVMBlueprintViewModelContext.cpp b/Engine/Plugins/Runtime/ModelViewViewModel/Source/ModelViewViewModelBlueprint/Private/MVVMBlueprintViewModelContext.cpp
new file mode 100644
--- /dev/null
+++ b/Engine/Plugins/Runtime/ModelViewViewModel/Source/ModelViewViewModelBlueprint/Private/MVVMBlueprintViewModelContext.cpp
@@ -0,0 +1,23 @@
+// Copyright Epic Games, Inc. All Rights Reserved.
+
+#include "MVVMBlueprintViewModelContext.h"
+
+FMVVMBlueprintViewModelContext::FMVVMBlueprintViewModelContext(TSubclassOf<UMVVMViewModelBase> InClass, FGuid InGuid) :
+	ViewModelContextId(InGuid),
+	ViewModelClass(InClass)
+{
+
+}
+
+FText FMVVMBlueprintViewModelContext::GetDisplayName() const
+{
+	if (!OverrideDisplayName.IsEmpty())
+	{
+		return OverrideDisplayName;
+	}
+	if (ViewModelClass.Get())
+	{
+		return ViewModelClass->GetDisplayNameText();
+	}
+	return FText::GetEmpty();
+}
diff --git a/Engine/Plugins/Runtime/ModelViewViewModel/Source/ModelViewViewModelBlueprint/Private/MVVMFunctionGraphHelper.h b/Engine/Plugins/Runtime/ModelViewViewModel/Source/ModelViewViewModelBlueprint/Private/MVVMFunctionGraphHelper.h
new file mode 100644
--- /dev/null
+++ b/Engine/Plugins/Runtime/ModelViewViewModel/Source/ModelViewViewModelBlueprint/Private/MVVMFunctionGraphHelper.h
@@ -0,0 +1,27 @@
+// Copyright Epic Games, Inc. All Rights Reserved.
+
+#pragma once
+
+#include "CoreMinimal.h"
+#include "Templates/SubclassOf.h"
+#include "UObject/Object.h"
+
+class FKismetCompilerContext;
+class UEdGraph;
+
+namespace UE::MVVM::FunctionGraphHelper
+{
+
+	/** Generate a function graph that can be viewed in the BP editor. */
+	MODELVIEWVIEWMODELBLUEPRINT_API UEdGraph* CreateFunctionGraph(FKismetCompilerContext& Context, const FStringView FunctionName, EFunctionFlags ExtraFunctionFlag, const FStringView Category, bool bIsEditable);
+
+	/** Generate a function graph that cannot be viewed in the BP editor. */
+	MODELVIEWVIEWMODELBLUEPRINT_API UEdGraph* CreateIntermediateFunctionGraph(FKismetCompilerContext& Context, const FStringView FunctionName, EFunctionFlags ExtraFunctionFlag, const FStringView Category, bool bIsEditable);
+
+	/** Add an input argument to an existing function graph. */
+	MODELVIEWVIEWMODELBLUEPRINT_API bool AddFunctionArgument(UEdGraph* FunctionGraph, TSubclassOf<UObject> Argument, FName ArgumentName);
+
+	/** Generate the nodes needed for to call the MVVMView::SetViewModel function. */
+	MODELVIEWVIEWMODELBLUEPRINT_API bool GenerateViewModelSetter(FKismetCompilerContext& InContext, UEdGraph* InFunctionGraph, FName InViewModelName);
+
+} //namespace
\ No newline at end of file
diff --git a/Engine/Plugins/Runtime/ModelViewViewModel/Source/ModelViewViewModelBlueprint/Private/MVVMViewBlueprintCompiler.cpp b/Engine/Plugins/Runtime/ModelViewViewModel/Source/ModelViewViewModelBlueprint/Private/MVVMViewBlueprintCompiler.cpp
new file mode 100644
--- /dev/null
+++ b/Engine/Plugins/Runtime/ModelViewViewModel/Source/ModelViewViewModelBlueprint/Private/MVVMViewBlueprintCompiler.cpp
@@ -0,0 +1,890 @@
+// Copyright Epic Games, Inc. All Rights Reserved.
+
+#include "MVVMViewBlueprintCompiler.h"
+#include "MVVMWidgetBlueprintExtension_View.h"
+#include "Blueprint/WidgetTree.h"
+#include "Bindings/MVVMBindingHelper.h"
+#include "Bindings/MVVMCompiledBindingLibraryCompiler.h"
+#include "Bindings/MVVMFieldPathHelper.h"
+#include "MVVMBlueprintView.h"
+#include "MVVMFunctionGraphHelper.h"
+#include "MVVMSubsystem.h"
+#include "MVVMViewModelBase.h"
+#include "WidgetBlueprintCompiler.h"
+#include "View/MVVMViewClass.h"
+
+namespace UE::MVVM::Private
+{
+void FMVVMViewBlueprintCompiler::AddErrorForBinding(FMVVMBlueprintViewBinding& Binding, const FString& Message) const
+{
+	WidgetBlueprintCompilerContext.MessageLog.Error(*Message);
+	Binding.Errors.Add(FText::FromString(Message));
+}
+
+void FMVVMViewBlueprintCompiler::AddExtension(UWidgetBlueprintGeneratedClass* Class, UMVVMViewClass* ViewExtension)
+{
+	WidgetBlueprintCompilerContext.AddExtension(Class, ViewExtension);
+}
+
+void FMVVMViewBlueprintCompiler::CleanOldData(UWidgetBlueprintGeneratedClass* ClassToClean, UObject* OldCDO)
+{
+	// Clean old View
+	if (!WidgetBlueprintCompilerContext.Blueprint->bIsRegeneratingOnLoad && WidgetBlueprintCompilerContext.bIsFullCompile)
+	{
+		auto RenameObjectToTransientPackage = [](UObject* ObjectToRename)
+		{
+			const ERenameFlags RenFlags = REN_DontCreateRedirectors | REN_NonTransactional | REN_DoNotDirty;
+
+			ObjectToRename->Rename(nullptr, GetTransientPackage(), RenFlags);
+			ObjectToRename->SetFlags(RF_Transient);
+			ObjectToRename->ClearFlags(RF_Public | RF_Standalone | RF_ArchetypeObject);
+
+			FLinkerLoad::InvalidateExport(ObjectToRename);
+		};
+
+		TArray<UObject*> Children;
+		const bool bIncludeNestedObjects = false;
+		ForEachObjectWithOuter(ClassToClean, [&Children](UObject* Child)
+			{
+				if (Cast<UMVVMViewClass>(Child))
+				{
+					Children.Add(Child);
+				}
+			}, bIncludeNestedObjects);
+
+		for (UObject* Child : Children)
+		{
+			RenameObjectToTransientPackage(Child);
+		}
+	}
+}
+
+
+void FMVVMViewBlueprintCompiler::CreateFunctions(UMVVMBlueprintView* BlueprintView)
+{
+	if (!bAreSourcesCreatorValid || !bIsBindingsValid)
+	{
+		return;
+	}
+
+	for (const FCompilerSourceCreatorContext& SourceCreator : SourceCreatorContexts)
+	{
+		if (SourceCreator.SetterGraph)
+		{
+			if (!UE::MVVM::FunctionGraphHelper::GenerateViewModelSetter(WidgetBlueprintCompilerContext, SourceCreator.SetterGraph, SourceCreator.ViewModelContext.GetViewModelName()))
+			{
+				WidgetBlueprintCompilerContext.MessageLog.Error(TEXT("The setter function for '%s' could not be generated."), *SourceCreator.ViewModelContext.GetDisplayName().ToString());
+				continue;
+			}
+		}
+	}
+}
+
+
+void FMVVMViewBlueprintCompiler::CreateVariables(const FWidgetBlueprintCompilerContext::FCreateVariableContext& Context, UMVVMBlueprintView* BlueprintView)
+{
+	if (!BlueprintView)
+	{
+		return;
+	}
+
+	if (!bAreSourcesCreatorValid || !bIsBindingsValid || !bIsBindingsValid)
+	{
+		return;
+	}
+
+	if (Context.GetCompileType() == EKismetCompileType::SkeletonOnly)
+	{
+		CreateWidgetMap(Context, BlueprintView);
+		CreateSourceLists(Context, BlueprintView);
+		CreateFunctionsDeclaration(Context, BlueprintView);
+	}
+
+	auto CreateVariable = [&Context](const FCompilerSourceContext& SourceContext) -> FProperty*
+	{
+		FEdGraphPinType ViewModelPinType(UEdGraphSchema_K2::PC_Object, NAME_None, SourceContext.Class, EPinContainerType::None, false, FEdGraphTerminalType());
+		FProperty* ViewModelProperty = Context.CreateVariable(SourceContext.PropertyName, ViewModelPinType);
+		if (ViewModelProperty != nullptr)
+		{
+			ViewModelProperty->SetPropertyFlags(CPF_BlueprintVisible | CPF_BlueprintReadOnly | CPF_RepSkip
+				| CPF_Transient | CPF_DuplicateTransient);
+			ViewModelProperty->SetPropertyFlags(SourceContext.bExposeOnSpawn ? CPF_ExposeOnSpawn : CPF_DisableEditOnInstance);
+
+#if WITH_EDITOR
+			if (!SourceContext.BlueprintSetter.IsEmpty())
+			{
+				const FName NAME_MetaDataBlueprintSetter = "BlueprintSetter";
+				ViewModelProperty->SetMetaData(NAME_MetaDataBlueprintSetter, *SourceContext.BlueprintSetter);
+			}
+			if (!SourceContext.DisplayName.IsEmpty())
+			{
+				const FName NAME_MetaDataDisplayName = "DisplayName";
+				ViewModelProperty->SetMetaData(NAME_MetaDataDisplayName, *SourceContext.DisplayName);
+			}
+
+			if (!SourceContext.CategoryName.IsEmpty())
+			{
+				const FName NAME_MetaDataCategoryName = "Category";
+				ViewModelProperty->SetMetaData(NAME_MetaDataCategoryName, *SourceContext.CategoryName);
+			}
+#endif
+		}
+		return ViewModelProperty;
+	};
+
+	for (FCompilerSourceContext& SourceContext : SourceContexts)
+	{
+		SourceContext.Field = BindingHelper::FindFieldByName(Context.GetSkeletonGeneratedClass(), FMVVMBindingName(SourceContext.PropertyName));
+
+		if (!SourceContext.Field.IsEmpty())
+		{
+			if (!BindingHelper::IsValidForSourceBinding(SourceContext.Field))
+			{
+				WidgetBlueprintCompilerContext.MessageLog.Error(*FString::Printf(TEXT("The field for source '%s' exist but will is not accessible at runtime.")
+					, *SourceContext.DisplayName));
+				bAreSourcesCreatorValid = false;
+				continue;
+			}
+
+			const FProperty* Property = SourceContext.Field.IsProperty() ? SourceContext.Field.GetProperty() : BindingHelper::GetReturnProperty(SourceContext.Field.GetFunction());
+			const FObjectPropertyBase* ObjectProperty = CastField<const FObjectPropertyBase>(Property);
+			const bool bIsCompatible = ObjectProperty && SourceContext.Class->IsChildOf(ObjectProperty->PropertyClass);
+			if (!bIsCompatible)
+			{
+				WidgetBlueprintCompilerContext.MessageLog.Error(TEXT("There is already a property named '%s' that is not compatible with source of the same name."), *SourceContext.DisplayName);
+				bAreSourceContextsValid = false;
+				continue;
+			}
+		}
+
+		if (SourceContext.Field.IsEmpty())
+		{
+			SourceContext.Field = FMVVMConstFieldVariant(CreateVariable(SourceContext));
+		}
+
+		if (SourceContext.Field.IsEmpty())
+		{
+			WidgetBlueprintCompilerContext.MessageLog.Error(TEXT("The variable for '%s' could not be created."), *SourceContext.DisplayName);
+			bAreSourceContextsValid = false;
+			continue;
+		}
+	}
+}
+
+
+void FMVVMViewBlueprintCompiler::CreateWidgetMap(const FWidgetBlueprintCompilerContext::FCreateVariableContext& Context, UMVVMBlueprintView* BlueprintView)
+{
+	// The widget tree is not created yet for SKEL class.
+	//Context.GetGeneratedClass()->GetWidgetTreeArchetype()
+	WidgetNameToWidgetPointerMap.Reset();
+
+	TArray<UWidget*> Widgets;
+	UWidgetBlueprint* WidgetBPToScan = Context.GetWidgetBlueprint();
+	while (WidgetBPToScan != nullptr)
+	{
+		Widgets = WidgetBPToScan->GetAllSourceWidgets();
+		if (Widgets.Num() != 0)
+		{
+			break;
+		}
+		WidgetBPToScan = WidgetBPToScan->ParentClass && WidgetBPToScan->ParentClass->ClassGeneratedBy ? Cast<UWidgetBlueprint>(WidgetBPToScan->ParentClass->ClassGeneratedBy) : nullptr;
+	}
+
+	for (UWidget* Widget : Widgets)
+	{
+		WidgetNameToWidgetPointerMap.Add(Widget->GetFName(), Widget);
+	}
+}
+
+
+void FMVVMViewBlueprintCompiler::CreateSourceLists(const FWidgetBlueprintCompilerContext::FCreateVariableContext& Context, UMVVMBlueprintView* BlueprintView)
+{
+	SourceContexts.Reset();
+
+	if (Context.GetCompileType() == EKismetCompileType::SkeletonOnly)
+	{
+		SourceCreatorContexts.Reset();
+	}
+
+	TSet<FGuid> ViewModelsGuid;
+	TSet<FName> WidgetSources;
+	for (const FMVVMBlueprintViewModelContext& ViewModelContext : BlueprintView->GetViewModels())
+	{
+		if (!ViewModelContext.GetViewModelId().IsValid())
+		{
+			WidgetBlueprintCompilerContext.MessageLog.Error(TEXT("The viewmodel context has an invalid Guid"));
+			bAreSourcesCreatorValid = false;
+			continue;
+		}
+
+		if (ViewModelsGuid.Contains(ViewModelContext.GetViewModelId()))
+		{
+			WidgetBlueprintCompilerContext.MessageLog.Error(*FString::Printf(TEXT("The viewmodel '%s' is present twice.")
+				, *ViewModelContext.GetViewModelId().ToString()));
+			bAreSourcesCreatorValid = false;
+			continue;
+		}
+
+		ViewModelsGuid.Add(ViewModelContext.GetViewModelId());
+
+		if (ViewModelContext.GetViewModelClass() == nullptr)
+		{
+			WidgetBlueprintCompilerContext.MessageLog.Error(*FString::Printf(TEXT("The viewmodel '%s' has an invalid class.")
+				, *ViewModelContext.GetViewModelId().ToString()));
+			bAreSourcesCreatorValid = false;
+			continue;
+		}
+
+		int32 FoundSourceContextIndex = INDEX_NONE;
+		if (Context.GetCompileType() == EKismetCompileType::SkeletonOnly)
+		{
+			FCompilerSourceCreatorContext SourceContext;
+			SourceContext.ViewModelContext = ViewModelContext;
+			SourceContext.Type = ECompilerSourceCreatorType::ViewModel;
+			if (ViewModelContext.bCreateSetterFunction)
+			{
+				SourceContext.SetterFunctionName = TEXT("Set") + ViewModelContext.GetViewModelName().ToString();
+			}
+			FoundSourceContextIndex = SourceCreatorContexts.Emplace(MoveTemp(SourceContext));
+		}
+		else
+		{
+			FGuid ViewModelId = ViewModelContext.GetViewModelId();
+			FoundSourceContextIndex = SourceCreatorContexts.IndexOfByPredicate([ViewModelId](const FCompilerSourceCreatorContext& Other)
+				{
+					return Other.ViewModelContext.GetViewModelId() == ViewModelId;
+				});
+		}
+		checkf(FoundSourceContextIndex != INDEX_NONE, TEXT("The viewmodel was added after the skeleton was created?"));
+
+		FCompilerSourceContext SourceVariable;
+		SourceVariable.Class = ViewModelContext.GetViewModelClass().Get();
+		SourceVariable.PropertyName = ViewModelContext.GetViewModelName();
+		SourceVariable.DisplayName = ViewModelContext.GetDisplayName().ToString();
+		SourceVariable.CategoryName = TEXT("Viewmodel");
+		//SourceVariable.bExposeOnSpawn = ViewModelContext.bCreateSetterFunction;
+		SourceVariable.BlueprintSetter = SourceCreatorContexts[FoundSourceContextIndex].SetterFunctionName;
+		SourceContexts.Emplace(MoveTemp(SourceVariable));
+	}
+
+	bAreSourceContextsValid = bAreSourcesCreatorValid;
+
+	// Only find the source first property and destination first property.
+	//The full path will be tested later. We want to build the list of property needed.
+	for (int32 Index = 0; Index < BlueprintView->GetNumBindings(); ++Index)
+	{
+		const FMVVMBlueprintViewBinding* BindingPtr = BlueprintView->GetBindingAt(Index);
+		if (BindingPtr == nullptr)
+		{
+			WidgetBlueprintCompilerContext.MessageLog.Error(*FString::Printf(TEXT("The binding number %d is invalid."), Index));
+			bAreSourceContextsValid = false;
+			continue;
+		}
+		const FMVVMBlueprintViewBinding& Binding = *BindingPtr;
+		if (!Binding.bCompile)
+		{
+			continue;
+		}
+
+		if (IsForwardBinding(Binding.BindingType) || IsBackwardBinding(Binding.BindingType))
+		{
+			// todo support any type of UObject
+			const FMVVMBlueprintViewModelContext* SourceViewModelContext = BlueprintView->FindViewModel(Binding.ViewModelPath.ContextId);
+			if (SourceViewModelContext == nullptr)
+			{
+				WidgetBlueprintCompilerContext.MessageLog.Error(*FString::Printf(TEXT("The binding number %d has an invalid source."), Index));
+				bAreSourceContextsValid = false;
+				continue;
+			}
+
+			bool bFound = ViewModelsGuid.Contains(SourceViewModelContext->GetViewModelId());
+			if (!bFound)
+			{
+				WidgetBlueprintCompilerContext.MessageLog.Error(*FString::Printf(TEXT("The binding number %d has a source that need to be added automaticly. That is not supported yet."), Index));
+				bAreSourceContextsValid = false;
+			}
+
+			if (Binding.WidgetPath.WidgetName == NAME_None)
+			{
+				WidgetBlueprintCompilerContext.MessageLog.Error(*FString::Printf(TEXT("The widget binding is invalid.")));
+				bAreSourceContextsValid = false;
+				continue;
+			}
+
+			if (Binding.WidgetPath.WidgetName != Context.GetWidgetBlueprint()->GetFName()) // is the userwidget itself
+			{
+				if (!WidgetSources.Contains(Binding.WidgetPath.WidgetName))
+				{
+					WidgetSources.Add(Binding.WidgetPath.WidgetName);
+
+					UWidget** WidgetPtr = WidgetNameToWidgetPointerMap.Find(Binding.WidgetPath.WidgetName);
+					if (WidgetPtr == nullptr || *WidgetPtr == nullptr)
+					{
+						WidgetBlueprintCompilerContext.MessageLog.Error(*FString::Printf(TEXT("The widget '%s' could not be found."), *Binding.WidgetPath.WidgetName.ToString()));
+						bAreSourceContextsValid = false;
+						continue;
+					}
+
+					UWidget* Widget = *WidgetPtr;
+					FCompilerSourceContext SourceVariable;
+					SourceVariable.Class = Widget->GetClass();
+					SourceVariable.PropertyName = Binding.WidgetPath.WidgetName;
+					SourceVariable.DisplayName = Widget->GetDisplayLabel();
+					SourceVariable.CategoryName = TEXT("Widget");
+					SourceContexts.Emplace(MoveTemp(SourceVariable));
+				}
+			}
+
+			// todo do the same for the function arguments
+		}
+	}
+}
+
+
+void FMVVMViewBlueprintCompiler::CreateFunctionsDeclaration(const FWidgetBlueprintCompilerContext::FCreateVariableContext& Context, UMVVMBlueprintView* BlueprintView)
+{
+	for (FCompilerSourceCreatorContext& SourceCreator : SourceCreatorContexts)
+	{
+		if (!SourceCreator.SetterFunctionName.IsEmpty())
+		{
+			ensure(SourceCreator.SetterGraph == nullptr);
+			SourceCreator.SetterGraph = UE::MVVM::FunctionGraphHelper::CreateIntermediateFunctionGraph(
+				WidgetBlueprintCompilerContext
+				, SourceCreator.SetterFunctionName
+				, (FUNC_BlueprintCallable | FUNC_Public)
+				, TEXT("Viewmodel")
+				, false);
+
+			UE::MVVM::FunctionGraphHelper::AddFunctionArgument(SourceCreator.SetterGraph, SourceCreator.ViewModelContext.GetViewModelClass().Get(), "ViewModel");
+		}
+	}
+}
+
+
+bool FMVVMViewBlueprintCompiler::PreCompile(UWidgetBlueprintGeneratedClass* Class, UMVVMBlueprintView* BlueprintView)
+{
+	if (!bAreSourcesCreatorValid || !bIsBindingsValid || !bIsBindingsValid)
+	{
+		return false;
+	}
+
+	PreCompileSourceCreators(Class, BlueprintView);
+	PreCompileBindings(Class, BlueprintView);
+
+	return bAreSourcesCreatorValid && bIsBindingsValid && bIsBindingsValid;
+}
+
+
+bool FMVVMViewBlueprintCompiler::Compile(UWidgetBlueprintGeneratedClass* Class, UMVVMBlueprintView* BlueprintView, UMVVMViewClass* ViewExtension)
+{
+	if (!bAreSourcesCreatorValid || !bIsBindingsValid || !bIsBindingsValid)
+	{
+		return false;
+	}
+
+	TValueOrError<FCompiledBindingLibraryCompiler::FCompileResult, FString> CompileResult = BindingLibraryCompiler.Compile();
+	if (CompileResult.HasError())
+	{
+		WidgetBlueprintCompilerContext.MessageLog.Error(TEXT("The binding compilation failed. %s"), *CompileResult.GetError());
+		return false;
+	}
+	CompileSourceCreators(CompileResult.GetValue(), Class, BlueprintView, ViewExtension);
+	CompileBindings(CompileResult.GetValue(), Class, BlueprintView, ViewExtension);
+
+	bool bResult = bAreSourcesCreatorValid && bIsBindingsValid && bIsBindingsValid;
+	if (bResult)
+	{
+		ViewExtension->BindingLibrary = MoveTemp(CompileResult.GetValue().Library);
+	}
+
+	return bResult;
+}
+
+
+bool FMVVMViewBlueprintCompiler::PreCompileSourceCreators(UWidgetBlueprintGeneratedClass* Class, UMVVMBlueprintView* BlueprintView)
+{
+	if (!bAreSourcesCreatorValid)
+	{
+		return false;
+	}
+
+	for (FCompilerSourceCreatorContext& SourceCreatorContext : SourceCreatorContexts)
+	{
+		FMVVMViewClass_SourceCreator CompiledSourceCreator;
+
+		if (SourceCreatorContext.Type == ECompilerSourceCreatorType::ViewModel)
+		{
+			const FMVVMBlueprintViewModelContext& ViewModelContext = SourceCreatorContext.ViewModelContext;
+			checkf(ViewModelContext.GetViewModelClass(), TEXT("The viewmodel class is invalid. It was checked in CreateSourceList"));
+
+			if (ViewModelContext.GetViewModelClass()->HasAllClassFlags(CLASS_Deprecated))
+			{
+				WidgetBlueprintCompilerContext.MessageLog.Warning(*FString::Printf(TEXT("The viewmodel '%s' is deprecated. Please update your view.")
+					, *ViewModelContext.GetViewModelClass()->GetDisplayNameText().ToString()));
+			}
+
+			if (ViewModelContext.CreationType == EMVVMBlueprintViewModelContextCreationType::CreateInstance)
+			{
+				if (ViewModelContext.GetViewModelClass()->HasAllClassFlags(CLASS_Abstract))
+				{
+					WidgetBlueprintCompilerContext.MessageLog.Error(*FString::Printf(TEXT("The viewmodel '%s' is abstract and can't be created.")
+						, *ViewModelContext.GetViewModelClass()->GetDisplayNameText().ToString()));
+					bAreSourcesCreatorValid = false;
+					continue;
+				}
+			}
+			else if (ViewModelContext.CreationType == EMVVMBlueprintViewModelContextCreationType::PropertyPath)
+			{
+				if (ViewModelContext.ViewModelPropertyPath.IsEmpty())
+				{
+					WidgetBlueprintCompilerContext.MessageLog.Error(*FString::Printf(TEXT("The viewmodel '%s' doesn't have a valid field path to be resolved at runtime. Add a Getter or use a global identifier.")
+						, *ViewModelContext.GetViewModelClass()->GetDisplayNameText().ToString()));
+					bAreSourcesCreatorValid = true;
+					continue;
+				}
+
+				// Generate a path to read the value at runtime
+				TValueOrError<FCompiledBindingLibraryCompiler::FFieldPathHandle, FString> ReadFieldPathResult = BindingLibraryCompiler.AddObjectFieldPath(Class, ViewModelContext.ViewModelPropertyPath, ViewModelContext.GetViewModelClass(), true);
+				if (ReadFieldPathResult.HasError())
+				{
+					WidgetBlueprintCompilerContext.MessageLog.Error(*FString::Printf(TEXT("The field path for reading the viewmodel '%s' could not be created. %s")
+						, *ViewModelContext.GetViewModelClass()->GetDisplayNameText().ToString()
+						, *ReadFieldPathResult.GetError()));
+					bAreSourcesCreatorValid = false;
+					continue;
+				}
+
+				SourceCreatorContext.ReadPropertyPath = ReadFieldPathResult.StealValue();
+			}
+			else if (ViewModelContext.CreationType == EMVVMBlueprintViewModelContextCreationType::GlobalViewModelCollection)
+			{
+				if (ViewModelContext.GlobalViewModelIdentifier.IsNone())
+				{
+					WidgetBlueprintCompilerContext.MessageLog.Error(*FString::Printf(TEXT("The viewmodel '%s' doesn't have a valid Global identifier.")
+						, *ViewModelContext.GetViewModelClass()->GetDisplayNameText().ToString()));
+					bAreSourcesCreatorValid = false;
+					continue;
+				}
+			}
+			else
+			{
+				WidgetBlueprintCompilerContext.MessageLog.Error(*FString::Printf(TEXT("The viewmodel '%s' doesn't have a valid creation type.")
+					, *ViewModelContext.GetViewModelClass()->GetDisplayNameText().ToString()));
+				bAreSourcesCreatorValid = true;
+				continue;
+			}
+		}
+	}
+
+	return bAreSourcesCreatorValid;
+}
+
+
+bool FMVVMViewBlueprintCompiler::CompileSourceCreators(const FCompiledBindingLibraryCompiler::FCompileResult& CompileResult, UWidgetBlueprintGeneratedClass* Class, UMVVMBlueprintView* BlueprintView, UMVVMViewClass* ViewExtension)
+{
+	if (!bAreSourcesCreatorValid)
+	{
+		return false;
+	}
+
+	for (const FCompilerSourceCreatorContext& SourceCreatorContext : SourceCreatorContexts)
+	{
+		FMVVMViewClass_SourceCreator CompiledSourceCreator;
+
+		if (SourceCreatorContext.Type == ECompilerSourceCreatorType::ViewModel)
+		{
+			const FMVVMBlueprintViewModelContext& ViewModelContext = SourceCreatorContext.ViewModelContext;
+			if (ViewModelContext.CreationType == EMVVMBlueprintViewModelContextCreationType::CreateInstance)
+			{
+				CompiledSourceCreator = FMVVMViewClass_SourceCreator::MakeInstance(ViewModelContext.GetViewModelName(), ViewModelContext.GetViewModelClass());
+			}
+			else if (ViewModelContext.CreationType == EMVVMBlueprintViewModelContextCreationType::PropertyPath)
+			{
+				const FMVVMVCompiledFieldPath* CompiledFieldPath = CompileResult.FieldPaths.Find(SourceCreatorContext.ReadPropertyPath);
+				if (CompiledFieldPath == nullptr)
+				{
+					WidgetBlueprintCompilerContext.MessageLog.Error(*FString::Printf(TEXT("The viewmodel '%s' initialization binding was not generated.")
+						, *ViewModelContext.GetViewModelClass()->GetDisplayNameText().ToString()));
+					bAreSourcesCreatorValid = false;
+					continue;
+				}
+				CompiledSourceCreator = FMVVMViewClass_SourceCreator::MakeFieldPath(ViewModelContext.GetViewModelName(), ViewModelContext.GetViewModelClass(), *CompiledFieldPath);
+			}
+			else if (ViewModelContext.CreationType == EMVVMBlueprintViewModelContextCreationType::GlobalViewModelCollection)
+			{
+				if (ViewModelContext.GlobalViewModelIdentifier.IsNone())
+				{
+					WidgetBlueprintCompilerContext.MessageLog.Error(*FString::Printf(TEXT("The viewmodel '%s' doesn't have a valid Global identifier.")
+						, *ViewModelContext.GetViewModelClass()->GetDisplayNameText().ToString()));
+					bAreSourcesCreatorValid = false;
+					continue;
+				}
+
+				FMVVMViewModelContext GlobalViewModelInstance;
+				GlobalViewModelInstance.ContextClass = ViewModelContext.GetViewModelClass();
+				GlobalViewModelInstance.ContextName = ViewModelContext.GlobalViewModelIdentifier;
+				if (!GlobalViewModelInstance.IsValid())
+				{
+					WidgetBlueprintCompilerContext.MessageLog.Error(*FString::Printf(TEXT("The context for viewmodel '%s' could not be created.")
+						, *ViewModelContext.GetViewModelClass()->GetDisplayNameText().ToString()));
+					bAreSourcesCreatorValid = false;
+					continue;
+				}
+
+				CompiledSourceCreator = FMVVMViewClass_SourceCreator::MakeGlobalContext(ViewModelContext.GetViewModelName(), MoveTemp(GlobalViewModelInstance));
+			}
+			else
+			{
+				WidgetBlueprintCompilerContext.MessageLog.Error(*FString::Printf(TEXT("The viewmodel '%s' doesn't have a valid creation type.")
+					, *ViewModelContext.GetViewModelClass()->GetDisplayNameText().ToString()));
+				bAreSourcesCreatorValid = false;
+				continue;
+			}
+
+			ViewExtension->SourceCreators.Add(MoveTemp(CompiledSourceCreator));
+		}
+	}
+
+	return bAreSourcesCreatorValid;
+}
+
+
+bool FMVVMViewBlueprintCompiler::PreCompileBindings(UWidgetBlueprintGeneratedClass* Class, UMVVMBlueprintView* BlueprintView)
+{
+	if (!bAreSourceContextsValid)
+	{
+		return false;
+	}
+
+	const int32 NumBindings = BlueprintView->GetNumBindings();
+	Bindings.Reset(NumBindings);
+
+	for (int32 Index = 0; Index < BlueprintView->GetNumBindings(); ++Index)
+	{
+		const FMVVMBlueprintViewBinding* BindingPtr = BlueprintView->GetBindingAt(Index);
+		if (BindingPtr == nullptr)
+		{
+			WidgetBlueprintCompilerContext.MessageLog.Error(*FString::Printf(TEXT("The binding number %d is invalid."), Index));
+			bAreSourceContextsValid = false;
+			continue;
+		}
+		const FMVVMBlueprintViewBinding& Binding = *BindingPtr;
+		if (!Binding.bCompile)
+		{
+			continue;
+		}
+
+		FMVVMViewBlueprintCompiler* Self = this;
+		auto AddBinding = [Self](UWidgetBlueprintGeneratedClass* Class, const TArrayView<FString> Getters, const FString& Setter, const FString& ConversionFunction) -> TValueOrError<FCompilerBinding, FString>
+		{
+			FCompilerBinding Result;
+
+			for (const FString& Getter : Getters)
+			{
+				// Generate a path to read the value at runtime
+				TValueOrError<FCompiledBindingLibraryCompiler::FFieldPathHandle, FString> FieldPathResult = Self->BindingLibraryCompiler.AddFieldPath(Class, Getter, true);
+				if (FieldPathResult.HasError())
+				{
+					return MakeError(FString::Printf(TEXT("Couldn't create the source field path '%s'. %s")
+						, *Getter
+						, *FieldPathResult.GetError()));
+				}
+				Result.SourceRead.Add(FieldPathResult.GetValue());
+			}
+
+			{
+				TValueOrError<FCompiledBindingLibraryCompiler::FFieldPathHandle, FString> FieldPathResult = Self->BindingLibraryCompiler.AddFieldPath(Class, Setter, false);
+				if (FieldPathResult.HasError())
+				{
+					return MakeError(FString::Printf(TEXT("Couldn't create the destination field path '%s'. %s")
+						, *Setter
+						, *FieldPathResult.GetError()));
+				}
+				Result.DestinationWrite = FieldPathResult.GetValue();
+			}
+
+			if (!ConversionFunction.IsEmpty())
+			{
+				TValueOrError<FCompiledBindingLibraryCompiler::FFieldPathHandle, FString> FieldPathResult = Self->BindingLibraryCompiler.AddConversionFunctionFieldPath(Class, ConversionFunction);
+				if (FieldPathResult.HasError())
+				{
+					return MakeError(FString::Printf(TEXT("Couldn't create the conversion function field path '%s'. %s")
+						, *ConversionFunction
+						, *FieldPathResult.GetError()));
+				}
+				Result.ConversionFunction = FieldPathResult.GetValue();
+			}
+
+			// Generate the binding
+			TValueOrError<FCompiledBindingLibraryCompiler::FBindingHandle, FString> BindingResult = Self->BindingLibraryCompiler.AddBinding(Result.SourceRead, Result.DestinationWrite, Result.ConversionFunction);
+			if (BindingResult.HasError())
+			{
+				return MakeError(BindingResult.StealError());
+			}
+			Result.BindingHandle = BindingResult.StealValue();
+
+			return MakeValue(Result);
+		};
+
+		auto AddFieldId = [Self](UClass* SourceContextClass, bool bNotifyFieldValueChangedRequired, EMVVMBindingMode BindingMode, FName FieldToListenTo) -> TValueOrError<FCompiledBindingLibraryCompiler::FFieldIdHandle, FString>
+		{
+			UE::FieldNotification::FFieldId FoundFieldId;
+
+			if (!IsOneTimeBinding(BindingMode) && bNotifyFieldValueChangedRequired)
+			{
+				return Self->BindingLibraryCompiler.AddFieldId(SourceContextClass, FieldToListenTo);
+			}
+			return MakeValue(FCompiledBindingLibraryCompiler::FFieldIdHandle());
+		};
+
+		if (IsForwardBinding(Binding.BindingType))
+		{
+			// Todo change that for any type of object
+			const FMVVMBlueprintViewModelContext* SourceViewModelContext = BlueprintView->FindViewModel(Binding.ViewModelPath.ContextId);
+			check(SourceViewModelContext);
+			FName SourceName = SourceViewModelContext->GetViewModelName();
+			int32 SourceContextIndex = Self->SourceContexts.IndexOfByPredicate([SourceName](const FCompilerSourceContext& Other) { return Other.PropertyName == SourceName; });
+			check(SourceContextIndex != INDEX_NONE);
+			TValueOrError<FCompiledBindingLibraryCompiler::FFieldIdHandle, FString> AddFieldResult = AddFieldId(SourceContexts[SourceContextIndex].Class, true, Binding.BindingType, Binding.ViewModelPath.GetBindingName().ToName());
+			if (AddFieldResult.HasError())
+			{
+				WidgetBlueprintCompilerContext.MessageLog.Error(*FString::Printf(TEXT("The binding '%d' could not create it's source. %s")
+					, Index
+					, *AddFieldResult.GetError()));
+				bIsBindingsValid = false;
+				continue;
+			}
+
+			TArray<FString> Getters;
+			//if (Binding.Conversion.HasConversionFunction && has read values defined (that are not the default one))
+			//{
+			//  Add the read values for the conversion function
+			//}
+			//else
+			{
+				FString Getter = SourceContexts[SourceContextIndex].PropertyName.ToString();
+				if (!Binding.ViewModelPath.GetGetterPropertyPath().IsEmpty())
+				{
+					Getter.AppendChar(TEXT('.'));
+					Getter.Append(Binding.ViewModelPath.GetGetterPropertyPath());
+				}
+				Getters.Add(MoveTemp(Getter));
+			}
+
+			FString Setter;
+			{
+				FName DestinationName = Binding.WidgetPath.WidgetName;
+				checkf(!DestinationName.IsNone(), TEXT("The destination should have been checked and set bAreSourceContextsValid"));
+				if (DestinationName == Class->ClassGeneratedBy->GetFName()) // is the userwidget itself
+				{
+					Setter = Binding.WidgetPath.GetSetterPropertyPath();
+				}
+				else
+				{
+					const int32 DestinationVariableContextIndex = SourceContexts.IndexOfByPredicate([DestinationName](const FCompilerSourceContext& Other) { return Other.PropertyName == DestinationName; });
+					check(DestinationVariableContextIndex != INDEX_NONE);
+
+					Setter = SourceContexts[DestinationVariableContextIndex].PropertyName.ToString();
+					if (!Binding.WidgetPath.GetSetterPropertyPath().IsEmpty())
+					{
+						Setter.AppendChar(TEXT('.'));
+						Setter.Append(Binding.WidgetPath.GetSetterPropertyPath());
+					}
+				}
+			}
+
+			TValueOrError<FCompilerBinding, FString> AddBindingResult = AddBinding(Class, Getters, Setter, Binding.Conversion.SetConversionFunctionPath);
+			if (AddBindingResult.HasError())
+			{
+				WidgetBlueprintCompilerContext.MessageLog.Error(*FString::Printf(TEXT("The binding '%d' could not be created. %s")
+					, Index
+					, *AddBindingResult.GetError()));
+				bIsBindingsValid = false;
+				continue;
+			}
+
+			FCompilerBinding NewBinding = AddBindingResult.StealValue();
+			NewBinding.BindingIndex = Index;
+			NewBinding.SourceContextIndex = SourceContextIndex;
+			NewBinding.FieldIdHandle = AddFieldResult.StealValue();
+			Bindings.Emplace(NewBinding);
+		}
+
+		if (IsBackwardBinding(Binding.BindingType))
+		{
+			// Todo change that for any type of object
+
+			FCompiledBindingLibraryCompiler::FFieldIdHandle FieldIdHandle;
+			int32 SourceContextIndex = INDEX_NONE;
+
+			const FName SourceName = Binding.WidgetPath.WidgetName;
+			checkf(!SourceName.IsNone(), TEXT("The source should have been checked and set bAreSourceContextsValid"));
+			const bool bSourceIsUserWidget = SourceName == Class->ClassGeneratedBy->GetFName();
+			if (bSourceIsUserWidget)
+			{
+				TValueOrError<FCompiledBindingLibraryCompiler::FFieldIdHandle, FString> AddFieldIdResult = AddFieldId(Class->ClassGeneratedBy->GetClass(), true, Binding.BindingType, Binding.WidgetPath.GetBindingName().ToName());
+				if (AddFieldIdResult.HasError())
+				{
+					WidgetBlueprintCompilerContext.MessageLog.Error(*FString::Printf(TEXT("The binding '%d' could not create it's source. %s")
+						, Index
+						, *AddFieldIdResult.GetError()));
+					bIsBindingsValid = false;
+					continue;
+				}
+				FieldIdHandle = AddFieldIdResult.StealValue();
+			}
+			else
+			{
+				SourceContextIndex = Self->SourceContexts.IndexOfByPredicate([SourceName](const FCompilerSourceContext& Other) { return Other.PropertyName == SourceName; });
+				check(SourceContextIndex != INDEX_NONE);
+				TValueOrError<FCompiledBindingLibraryCompiler::FFieldIdHandle, FString> AddFieldIdResult = AddFieldId(SourceContexts[SourceContextIndex].Class, true, Binding.BindingType, Binding.WidgetPath.GetBindingName().ToName());
+				if (AddFieldIdResult.HasError())
+				{
+					WidgetBlueprintCompilerContext.MessageLog.Error(*FString::Printf(TEXT("The binding '%d' could not create it's source. %s")
+						, Index
+						, *AddFieldIdResult.GetError()));
+					bIsBindingsValid = false;
+					continue;
+				}
+				FieldIdHandle = AddFieldIdResult.StealValue();
+			}
+
+			TArray<FString> Getters;
+			//if (Binding.Conversion.HasConversionFunction && has read values defined (that are not the default one))
+			//{
+			//  Add the read values for the conversion function
+			//}
+			//else
+			{
+				if (bSourceIsUserWidget)
+				{
+					Getters.Add(Binding.WidgetPath.GetGetterPropertyPath());
+				}
+				else
+				{
+					FString Getter = SourceContexts[SourceContextIndex].PropertyName.ToString();
+					if (!Binding.WidgetPath.GetGetterPropertyPath().IsEmpty())
+					{
+						Getter.AppendChar(TEXT('.'));
+						Getter.Append(Binding.WidgetPath.GetGetterPropertyPath());
+					}
+					Getters.Add(MoveTemp(Getter));
+				}
+			}
+
+			FString Setter;
+			{
+				const FMVVMBlueprintViewModelContext* SourceViewModelContext = BlueprintView->FindViewModel(Binding.ViewModelPath.ContextId);
+				check(SourceViewModelContext);
+				FName DestinationName = SourceViewModelContext->GetViewModelName();
+				const int32 DestinationVariableContextIndex = SourceContexts.IndexOfByPredicate([DestinationName](const FCompilerSourceContext& Other) { return Other.PropertyName == DestinationName; });
+				check(DestinationVariableContextIndex != INDEX_NONE);
+
+				Setter = SourceContexts[DestinationVariableContextIndex].PropertyName.ToString();
+				if (!Binding.ViewModelPath.GetSetterPropertyPath().IsEmpty())
+				{
+					Setter.AppendChar(TEXT('.'));
+					Setter.Append(Binding.WidgetPath.GetSetterPropertyPath());
+				}
+			}
+
+			TValueOrError<FCompilerBinding, FString> AddBindingResult = AddBinding(Class, Getters, Setter, Binding.Conversion.GetConversionFunctionPath);
+			if (AddBindingResult.HasError())
+			{
+				WidgetBlueprintCompilerContext.MessageLog.Error(*FString::Printf(TEXT("The binding '%d' could not be created. %s")
+					, Index
+					, *AddBindingResult.GetError()));
+				bIsBindingsValid = false;
+				continue;
+			}
+
+			FCompilerBinding NewBinding = AddBindingResult.StealValue();
+			NewBinding.BindingIndex = Index;
+			NewBinding.SourceContextIndex = SourceContextIndex;
+			NewBinding.bSourceIsUserWidget = bSourceIsUserWidget;
+			NewBinding.FieldIdHandle = FieldIdHandle;
+			Bindings.Emplace(NewBinding);
+		}
+	}
+
+	return bIsBindingsValid;
+}
+
+
+bool FMVVMViewBlueprintCompiler::CompileBindings(const FCompiledBindingLibraryCompiler::FCompileResult& CompileResult, UWidgetBlueprintGeneratedClass* Class, UMVVMBlueprintView* BlueprintView, UMVVMViewClass* ViewExtension)
+{
+	if (!bIsBindingsValid)
+	{
+		return false;
+	}
+
+	for (const FCompilerBinding& CompileBinding : Bindings)
+	{
+		const FMVVMBlueprintViewBinding& ViewBinding = *BlueprintView->GetBindingAt(CompileBinding.BindingIndex);
+
+		FMVVMViewClass_CompiledBinding NewBinding;
+
+		check(CompileBinding.SourceContextIndex != INDEX_NONE);
+		NewBinding.SourcePropertyName = SourceContexts[CompileBinding.SourceContextIndex].PropertyName;
+
+		const FMVVMVCompiledFieldId* CompiledFieldId = CompileResult.FieldIds.Find(CompileBinding.FieldIdHandle);
+		if (CompiledFieldId == nullptr)
+		{
+			WidgetBlueprintCompilerContext.MessageLog.Error(*FString::Printf(TEXT("The field id for binding '%d' was not generated.")
+				, CompileBinding.BindingIndex));
+			bIsBindingsValid = false;
+			continue;
+		}
+
+		const FMVVMVCompiledBinding* CompiledBinding = CompileResult.Bindings.Find(CompileBinding.BindingHandle);
+		if (CompiledBinding == nullptr)
+		{
+			WidgetBlueprintCompilerContext.MessageLog.Error(*FString::Printf(TEXT("The binding '%d' was not generated.")
+				, CompileBinding.BindingIndex));
+			bIsBindingsValid = false;
+			continue;
+		}
+
+		NewBinding.FieldId = *CompiledFieldId;
+		NewBinding.Binding = *CompiledBinding;
+		NewBinding.UpdateMode = ViewBinding.UpdateMode;
+
+		NewBinding.Flags = 0;
+		NewBinding.Flags |= (ViewBinding.bEnabled) ? FMVVMViewClass_CompiledBinding::EBindingFlags::EnabledByDefault : 0;
+		NewBinding.Flags |= (IsForwardBinding(ViewBinding.BindingType)) ? FMVVMViewClass_CompiledBinding::EBindingFlags::ForwardBinding : 0;
+		NewBinding.Flags |= (ViewBinding.BindingType == EMVVMBindingMode::TwoWay) ? FMVVMViewClass_CompiledBinding::EBindingFlags::TwoWayBinding : 0;
+		NewBinding.Flags |= (IsOneTimeBinding(ViewBinding.BindingType)) ? FMVVMViewClass_CompiledBinding::EBindingFlags::OneTime : 0;
+
+		ViewExtension->CompiledBindings.Emplace(MoveTemp(NewBinding));
+	}
+
+	return bIsBindingsValid;
+}
+
+
+const FMVVMViewBlueprintCompiler::FCompilerSourceCreatorContext* FMVVMViewBlueprintCompiler::FindViewModelSource(FGuid Id) const
+{
+	return SourceCreatorContexts.FindByPredicate([Id](const FCompilerSourceCreatorContext& Other) { return Other.Type == ECompilerSourceCreatorType::ViewModel ? Other.ViewModelContext.GetViewModelId() == Id : false; });
+}
+
+//void FMVVMViewBlueprintCompiler::TransformToFunction(FMVVMCachedPropertyPath& PathToTransform, const UFunction* Function) const
+//{
+//	FString StrToTransform = PathToTransform.ToString();
+//
+//	const TCHAR Delim = TEXT('.');
+//	int32 Index = INDEX_NONE;
+//	StrToTransform.FindLastChar(Delim, Index);
+//	if (Index != INDEX_NONE)
+//	{
+//		StrToTransform.RemoveAt(Index + 1, StrToTransform.Len(), false);
+//		StrToTransform.Append(Function->GetName());
+//	}
+//	else
+//	{
+//		StrToTransform = Function->GetName();
+//	}
+//
+//	PathToTransform = StrToTransform;
+//}
+
+} //namespace
diff --git a/Engine/Plugins/Runtime/ModelViewViewModel/Source/ModelViewViewModelBlueprint/Private/MVVMViewBlueprintCompiler.h b/Engine/Plugins/Runtime/ModelViewViewModel/Source/ModelViewViewModelBlueprint/Private/MVVMViewBlueprintCompiler.h
new file mode 100644
--- /dev/null
+++ b/Engine/Plugins/Runtime/ModelViewViewModel/Source/ModelViewViewModelBlueprint/Private/MVVMViewBlueprintCompiler.h
@@ -0,0 +1,131 @@
+// Copyright Epic Games, Inc. All Rights Reserved.
+
+#pragma once
+
+#include "CoreMinimal.h"
+#include "Bindings/MVVMCompiledBindingLibraryCompiler.h"
+#include "MVVMBlueprintViewBinding.h"
+#include "MVVMBlueprintViewModelContext.h"
+#include "Templates/ValueOrError.h"
+#include "Types/MVVMFieldVariant.h"
+#include "UObject/Object.h"
+#include "WidgetBlueprintCompiler.h"
+
+struct FMVVMCachedPropertyPath;
+struct FMVVMBlueprintViewBinding;
+class FWidgetBlueprintCompilerContext;
+class UEdGraph;
+class UMVVMBlueprintView;
+class UMVVMViewClass;
+class UWidgetBlueprintGeneratedClass;
+
+namespace UE::MVVM::Private
+{
+
+struct FMVVMViewBlueprintCompiler
+{
+private:
+	struct FCompilerSourceContext;
+	struct FCompilerSourceCreatorContext;
+	struct FCompilerBinding;
+
+public:
+	FMVVMViewBlueprintCompiler(FWidgetBlueprintCompilerContext& InCreationContext)
+		: WidgetBlueprintCompilerContext(InCreationContext)
+	{}
+
+
+	FWidgetBlueprintCompilerContext& GetCompilerContext()
+	{
+		return WidgetBlueprintCompilerContext;
+	}
+
+	void AddExtension(UWidgetBlueprintGeneratedClass* Class, UMVVMViewClass* ViewExtension);
+	void CleanOldData(UWidgetBlueprintGeneratedClass* ClassToClean, UObject* OldCDO);
+
+	/** Generate function that are hidden from the user (not on the Skeleton class). */
+	void CreateFunctions(UMVVMBlueprintView* BlueprintView);
+	/** Generate variable in the Skeleton class */
+	void CreateVariables(const FWidgetBlueprintCompilerContext::FCreateVariableContext& Context, UMVVMBlueprintView* BlueprintView);
+
+	/** Add all the field path and the bindings to the library compiler. */
+	bool PreCompile(UWidgetBlueprintGeneratedClass* Class, UMVVMBlueprintView* BlueprintView);
+	/** Compile the library and fill the view and viewclass */
+	bool Compile(UWidgetBlueprintGeneratedClass* Class, UMVVMBlueprintView* BlueprintView, UMVVMViewClass* ViewExtension);
+
+private:
+	void CreateWidgetMap(const FWidgetBlueprintCompilerContext::FCreateVariableContext& Context, UMVVMBlueprintView* BlueprintView);
+	void CreateSourceLists(const FWidgetBlueprintCompilerContext::FCreateVariableContext& Context, UMVVMBlueprintView* BlueprintView);
+	void CreateFunctionsDeclaration(const FWidgetBlueprintCompilerContext::FCreateVariableContext& Context, UMVVMBlueprintView* BlueprintView);
+
+	bool PreCompileSourceCreators(UWidgetBlueprintGeneratedClass* Class, UMVVMBlueprintView* BlueprintView);
+	bool CompileSourceCreators(const FCompiledBindingLibraryCompiler::FCompileResult& CompileResult, UWidgetBlueprintGeneratedClass* Class, UMVVMBlueprintView* BlueprintView, UMVVMViewClass* ViewExtension);
+
+	bool PreCompileBindings(UWidgetBlueprintGeneratedClass* Class, UMVVMBlueprintView* BlueprintView);
+	bool CompileBindings(const FCompiledBindingLibraryCompiler::FCompileResult& CompileResult, UWidgetBlueprintGeneratedClass* Class, UMVVMBlueprintView* BlueprintView, UMVVMViewClass* ViewExtension);
+
+	const FCompilerSourceCreatorContext* FindViewModelSource(FGuid Id) const;
+
+	void AddErrorForBinding(FMVVMBlueprintViewBinding& Binding, const FString& Message) const;
+
+private:
+	struct FCompilerSourceContext
+	{
+		UClass* Class = nullptr;
+		FName PropertyName;
+		FString DisplayName;
+		FString CategoryName;
+		FString BlueprintSetter;
+		FMVVMConstFieldVariant Field;
+		bool bExposeOnSpawn = false;
+	};
+	TArray<FCompilerSourceContext> SourceContexts;
+
+	enum class ECompilerSourceCreatorType
+	{
+		ViewModel,
+	};
+	struct FCompilerSourceCreatorContext
+	{
+		FMVVMBlueprintViewModelContext ViewModelContext;
+		FCompiledBindingLibraryCompiler::FFieldPathHandle ReadPropertyPath;
+		ECompilerSourceCreatorType Type = ECompilerSourceCreatorType::ViewModel;
+		FString SetterFunctionName;
+		UEdGraph* SetterGraph = nullptr;
+	};
+	TArray<FCompilerSourceCreatorContext> SourceCreatorContexts;
+
+	enum class ECompilerBindingType
+	{
+		PropertyBinding, // normal binding
+	};
+	struct FCompilerBinding
+	{
+		int32 BindingIndex = INDEX_NONE;
+		int32 SourceContextIndex = INDEX_NONE;
+		bool bSourceIsUserWidget = false;
+
+		FCompiledBindingLibraryCompiler::FBindingHandle BindingHandle;
+
+		// At runtime, when TriggerFields changes
+		//if the ConversionFunction exist
+		// then we need to read all the inputs of the ConversionFunction
+		// then we copy the result of the ConversionFunction to DestinationWrite
+		//else we read the same property that TriggerField changed
+		// then we copy TriggerField to DestiantionWrite.
+		FCompiledBindingLibraryCompiler::FFieldIdHandle FieldIdHandle;
+		TArray<FCompiledBindingLibraryCompiler::FFieldPathHandle> SourceRead;
+		FCompiledBindingLibraryCompiler::FFieldPathHandle DestinationWrite;
+		FCompiledBindingLibraryCompiler::FFieldPathHandle ConversionFunction;
+	};
+	TArray<FCompilerBinding> Bindings;
+
+	TMap<FName, UWidget*> WidgetNameToWidgetPointerMap;
+	FWidgetBlueprintCompilerContext& WidgetBlueprintCompilerContext;
+	FCompiledBindingLibraryCompiler BindingLibraryCompiler;
+	bool bAreSourcesCreatorValid = true;
+	bool bAreSourceContextsValid = true;
+	bool bIsBindingsValid = true;
+};
+
+} //namespace
\ No newline at end of file
diff --git a/Engine/Plugins/Runtime/ModelViewViewModel/Source/ModelViewViewModelBlueprint/Private/MVVMWidgetBlueprintExtension_View.cpp b/Engine/Plugins/Runtime/ModelViewViewModel/Source/ModelViewViewModelBlueprint/Private/MVVMWidgetBlueprintExtension_View.cpp
new file mode 100644
--- /dev/null
+++ b/Engine/Plugins/Runtime/ModelViewViewModel/Source/ModelViewViewModelBlueprint/Private/MVVMWidgetBlueprintExtension_View.cpp
@@ -0,0 +1,133 @@
+// Copyright Epic Games, Inc. All Rights Reserved.
+
+#include "MVVMWidgetBlueprintExtension_View.h"
+#include "Blueprint/WidgetTree.h"
+#include "MVVMBlueprintView.h"
+#include "MVVMViewBlueprintCompiler.h"
+#include "MVVMViewModelBase.h"
+#include "WidgetBlueprintCompiler.h"
+#include "View/MVVMViewClass.h"
+
+
+void UMVVMWidgetBlueprintExtension_View::CreateBlueprintViewInstance()
+{
+	BlueprintView = NewObject<UMVVMBlueprintView>(this);
+	BlueprintViewChangedDelegate.Broadcast();
+}
+
+
+void UMVVMWidgetBlueprintExtension_View::DestroyBlueprintViewInstance()
+{
+	BlueprintView = nullptr;
+	BlueprintViewChangedDelegate.Broadcast();
+}
+
+
+void UMVVMWidgetBlueprintExtension_View::PostLoad()
+{
+	Super::PostLoad();
+	//for (FMVVMViewBindingEditorData& Data : ViewBindings)
+	//{
+	//	Data.Resolve();
+	//}
+}
+
+
+void UMVVMWidgetBlueprintExtension_View::HandlePreloadObjectsForCompilation(UBlueprint* OwningBlueprint)
+{
+	if (BlueprintView)
+	{
+		BlueprintView->ConditionalPostLoad();
+
+		for (const FMVVMBlueprintViewBinding& Bindings : BlueprintView->GetBindings())
+		{
+			//UBlueprint::ForceLoad(Bindings.ViewModelPath);
+			//UBlueprint::ForceLoad(Bindings.WidgetPath);
+		}
+		for (const FMVVMBlueprintViewModelContext& AvailableViewModel : BlueprintView->GetViewModels())
+		{
+			if (AvailableViewModel.GetViewModelClass())
+			{
+				UBlueprint::ForceLoad(AvailableViewModel.GetViewModelClass().Get());
+			}
+			//UBlueprint::ForceLoad(ViewModelPropertyPath);
+		}
+	}
+}
+
+
+void UMVVMWidgetBlueprintExtension_View::HandleBeginCompilation(FWidgetBlueprintCompilerContext& InCreationContext)
+{
+	CurrentCompilerContext.Reset();
+	if (BlueprintView)
+	{
+		for (FMVVMBlueprintViewBinding& Binding : BlueprintView->GetBindings())
+		{
+			Binding.Errors.Reset();
+		}
+
+		CurrentCompilerContext = MakePimpl<UE::MVVM::Private::FMVVMViewBlueprintCompiler>(InCreationContext);
+	}
+}
+
+
+void UMVVMWidgetBlueprintExtension_View::HandleEndCompilation()
+{
+	CurrentCompilerContext.Reset();
+}
+
+
+void UMVVMWidgetBlueprintExtension_View::HandleCleanAndSanitizeClass(UWidgetBlueprintGeneratedClass* ClassToClean, UObject* OldCDO)
+{
+	Super::HandleCleanAndSanitizeClass(ClassToClean, OldCDO);
+
+	if (CurrentCompilerContext)
+	{
+		CurrentCompilerContext->CleanOldData(ClassToClean, OldCDO);
+	}
+}
+
+
+void UMVVMWidgetBlueprintExtension_View::HandleCreateClassVariablesFromBlueprint(const FWidgetBlueprintCompilerContext::FCreateVariableContext& Context)
+{
+	Super::HandleCreateClassVariablesFromBlueprint(Context);
+
+	CurrentCompilerContext->CreateVariables(Context, GetBlueprintView());
+}
+
+
+void UMVVMWidgetBlueprintExtension_View::HandleCreateFunctionList()
+{
+	Super::HandleCreateFunctionList();
+
+	if (CurrentCompilerContext)
+	{
+		CurrentCompilerContext->CreateFunctions(BlueprintView);
+	}
+}
+
+
+void UMVVMWidgetBlueprintExtension_View::HandleFinishCompilingClass(UWidgetBlueprintGeneratedClass* Class)
+{
+	Super::HandleFinishCompilingClass(Class);
+
+	check(CurrentCompilerContext);
+
+	if (CurrentCompilerContext->GetCompilerContext().bAssignDelegateSignatureFunction)
+	{
+		return;
+	}
+
+	if (!CurrentCompilerContext->PreCompile(Class, BlueprintView))
+	{
+		return;
+	}
+
+	UMVVMViewClass* ViewExtension = NewObject<UMVVMViewClass>(Class);
+	if (!CurrentCompilerContext->Compile(Class, BlueprintView, ViewExtension))
+	{
+		return;
+	}
+
+	CurrentCompilerContext->AddExtension(Class, ViewExtension);
+}
