MVVM: Move the ModelViewViewModel plugin to //UE5/Main/

diff --git a/Engine/Plugins/Runtime/ModelViewViewModel/Source/ModelViewViewModelBlueprint/Private/ModelViewViewModelBlueprintModule.cpp b/Engine/Plugins/Runtime/ModelViewViewModel/Source/ModelViewViewModelBlueprint/Private/ModelViewViewModelBlueprintModule.cpp
new file mode 100644
--- /dev/null
+++ b/Engine/Plugins/Runtime/ModelViewViewModel/Source/ModelViewViewModelBlueprint/Private/ModelViewViewModelBlueprintModule.cpp
@@ -0,0 +1,8 @@
+// Copyright Epic Games, Inc. All Rights Reserved.
+
+#include "ModelViewViewModelBlueprintModule.h"
+#include "Modules/ModuleInterface.h"
+#include "Modules/ModuleManager.h"
+
+
+IMPLEMENT_MODULE(FDefaultModuleImpl, ModelViewViewModelBlueprintModule);
diff --git a/Engine/Plugins/Runtime/ModelViewViewModel/Source/ModelViewViewModelBlueprint/Public/Bindings/MVVMCompiledBindingLibraryCompiler.h b/Engine/Plugins/Runtime/ModelViewViewModel/Source/ModelViewViewModelBlueprint/Public/Bindings/MVVMCompiledBindingLibraryCompiler.h
new file mode 100644
--- /dev/null
+++ b/Engine/Plugins/Runtime/ModelViewViewModel/Source/ModelViewViewModelBlueprint/Public/Bindings/MVVMCompiledBindingLibraryCompiler.h
@@ -0,0 +1,196 @@
+// Copyright Epic Games, Inc. All Rights Reserved.
+
+#pragma once
+
+#include "CoreMinimal.h"
+#include "Bindings/MVVMCompiledBindingLibrary.h"
+#include "MVVMSubsystem.h"
+#include "Templates/PimplPtr.h"
+#include "Templates/ValueOrError.h"
+#include "Types/MVVMFieldVariant.h"
+
+
+namespace UE::MVVM::Private
+{
+	class FCompiledBindingLibraryCompilerImpl;
+} //namespace UE::MVVM::Private
+
+
+namespace UE::MVVM
+{
+
+/**  */
+class MODELVIEWVIEWMODELBLUEPRINT_API FCompiledBindingLibraryCompiler
+{
+public:
+	/** */
+	struct FFieldPathHandle
+	{
+	public:
+		explicit FFieldPathHandle()
+			: Id(0)
+		{
+		}
+
+		static FFieldPathHandle MakeHandle()
+		{
+			FFieldPathHandle Handle;
+			++IdGenerator;
+			Handle.Id = IdGenerator;
+			return Handle;
+		}
+
+		bool IsValid() const
+		{
+			return Id != 0;
+		}
+
+		bool operator==(const FFieldPathHandle& Other) const
+		{
+			return Id == Other.Id;
+		}
+
+		bool operator!=(const FFieldPathHandle& Other) const
+		{
+			return Id != Other.Id;
+		}
+
+		friend uint32 GetTypeHash(const FFieldPathHandle Handle)
+		{
+			return GetTypeHash(Handle.Id);
+		}
+
+	private:
+		static int32 IdGenerator;
+		int32 Id;
+	};
+
+
+	/** */
+	struct FBindingHandle
+	{
+	public:
+		explicit FBindingHandle()
+			: Id(0)
+		{
+		}
+
+		static FBindingHandle MakeHandle()
+		{
+			FBindingHandle Handle;
+			++IdGenerator;
+			Handle.Id = IdGenerator;
+			return Handle;
+		}
+
+		bool IsValid() const
+		{
+			return Id != 0;
+		}
+
+		bool operator==(const FBindingHandle& Other) const
+		{
+			return Id == Other.Id;
+		}
+
+		bool operator!=(const FBindingHandle& Other) const
+		{
+			return Id != Other.Id;
+		}
+
+		friend uint32 GetTypeHash(const FBindingHandle Handle)
+		{
+			return GetTypeHash(Handle.Id);
+		}
+
+	private:
+		static int32 IdGenerator;
+		int32 Id;
+	};
+
+	/** */
+	struct FFieldIdHandle
+	{
+	public:
+		explicit FFieldIdHandle()
+			: Id(0)
+		{
+		}
+
+		static FFieldIdHandle MakeHandle()
+		{
+			FFieldIdHandle Handle;
+			++IdGenerator;
+			Handle.Id = IdGenerator;
+			return Handle;
+		}
+
+		bool IsValid() const
+		{
+			return Id != 0;
+		}
+
+		bool operator==(const FFieldIdHandle& Other) const
+		{
+			return Id == Other.Id;
+		}
+
+		bool operator!=(const FFieldIdHandle& Other) const
+		{
+			return Id != Other.Id;
+		}
+
+		friend uint32 GetTypeHash(const FFieldIdHandle Handle)
+		{
+			return GetTypeHash(Handle.Id);
+		}
+
+	private:
+		static int32 IdGenerator;
+		int32 Id;
+	};
+
+public:
+	FCompiledBindingLibraryCompiler();
+
+public:
+	/** */
+	TValueOrError<FFieldIdHandle, FString> AddFieldId(TSubclassOf<UObject> SourceClass, FName FieldName);
+
+	/** */
+	TValueOrError<FFieldPathHandle, FString> AddFieldPath(TSubclassOf<UObject> SourceClass, FStringView FieldPath, bool bRead);
+	
+	/** */
+	TValueOrError<FFieldPathHandle, FString> AddFieldPath(TArrayView<UE::MVVM::FMVVMFieldVariant> FieldPath, bool bRead);
+
+	/** */
+	TValueOrError<FFieldPathHandle, FString> AddObjectFieldPath(TSubclassOf<UObject> SourceClass, FStringView FieldPath, UClass* ExpectedType, bool bRead);
+
+	/** */
+	TValueOrError<FFieldPathHandle, FString> AddConversionFunctionFieldPath(TSubclassOf<UObject> SourceClass, FStringView FieldPath);
+
+	/** */
+	TValueOrError<FBindingHandle, FString> AddBinding(FFieldPathHandle Source, FFieldPathHandle Destination);
+
+	/** */
+	TValueOrError<FBindingHandle, FString> AddBinding(FFieldPathHandle Source, FFieldPathHandle Destination, FFieldPathHandle ConversionFunction);
+
+	/** */
+	TValueOrError<FBindingHandle, FString> AddBinding(TArrayView<FFieldPathHandle> Sources, FFieldPathHandle Destination, FFieldPathHandle ConversionFunction);
+
+	struct FCompileResult
+	{
+		FMVVMCompiledBindingLibrary Library;
+		TMap<FFieldPathHandle, FMVVMVCompiledFieldPath> FieldPaths;
+		TMap<FBindingHandle, FMVVMVCompiledBinding> Bindings;
+		TMap<FFieldIdHandle, FMVVMVCompiledFieldId> FieldIds;
+	};
+
+	/** */
+	TValueOrError<FCompileResult, FString> Compile();
+
+private:
+	TPimplPtr<Private::FCompiledBindingLibraryCompilerImpl> Impl;
+};
+
+} //namespace UE::MVVM
\ No newline at end of file
diff --git a/Engine/Plugins/Runtime/ModelViewViewModel/Source/ModelViewViewModelBlueprint/Public/MVVMBlueprintView.h b/Engine/Plugins/Runtime/ModelViewViewModel/Source/ModelViewViewModelBlueprint/Public/MVVMBlueprintView.h
new file mode 100644
--- /dev/null
+++ b/Engine/Plugins/Runtime/ModelViewViewModel/Source/ModelViewViewModelBlueprint/Public/MVVMBlueprintView.h
@@ -0,0 +1,83 @@
+// Copyright Epic Games, Inc. All Rights Reserved.
+
+#pragma once
+
+#include "CoreMinimal.h"
+#include "View/MVVMView.h"
+#include "MVVMBlueprintViewBinding.h"
+#include "MVVMBlueprintViewModelContext.h"
+
+#include "MVVMBlueprintView.generated.h"
+
+class UWidget;
+class UWidgetBlueprint;
+
+/**
+ * 
+ */
+UCLASS(Within=MVVMWidgetBlueprintExtension_View)
+class MODELVIEWVIEWMODELBLUEPRINT_API UMVVMBlueprintView : public UObject
+{
+	GENERATED_BODY()
+
+public:
+	FMVVMBlueprintViewModelContext* FindViewModel(FGuid ViewModelId);
+	const FMVVMBlueprintViewModelContext* FindViewModel(FGuid ViewModelId) const;
+
+	void AddViewModel(const FMVVMBlueprintViewModelContext& NewContext);
+	void RemoveViewModel(FGuid ViewModelId);
+	void RemoveViewModels(const TArrayView<FGuid> ViewModelIds);
+	void SetViewModels(const TArray<FMVVMBlueprintViewModelContext>& ViewModelContexts);
+
+	const TArrayView<const FMVVMBlueprintViewModelContext> GetViewModels() const
+	{
+		return AvailableViewModels; 
+	}
+
+	const FMVVMBlueprintViewBinding* FindBinding(const UWidget* Widget, const FProperty* Property) const;
+	FMVVMBlueprintViewBinding* FindBinding(const UWidget* Widget, const FProperty* Property);
+	const FMVVMBlueprintViewBinding* FindBinding(FName WidgetName, FMVVMBindingName BindingName) const;
+	FMVVMBlueprintViewBinding* FindBinding(FName WidgetName, FMVVMBindingName BindingName);
+
+	void RemoveBinding(const FMVVMBlueprintViewBinding* Binding);
+	void RemoveBindingAt(int32 Index);
+
+	FMVVMBlueprintViewBinding& AddBinding(const UWidget* Widget, const FProperty* Property);
+	FMVVMBlueprintViewBinding& AddDefaultBinding();
+
+	int32 GetNumBindings() const
+	{
+		return Bindings.Num();
+	}
+
+	FMVVMBlueprintViewBinding* GetBindingAt(int32 Index);
+	const FMVVMBlueprintViewBinding* GetBindingAt(int32 Index) const;
+
+	TArrayView<FMVVMBlueprintViewBinding> GetBindings()
+	{
+		return Bindings;
+	}
+
+	const TArrayView<const FMVVMBlueprintViewBinding> GetBindings() const
+	{
+		return Bindings;
+	}
+
+#if WITH_EDITOR
+	virtual void PostEditChangeProperty(FPropertyChangedEvent& PropertyChangedEvent) override;
+	void WidgetRenamed(FName OldObjectName, FName NewObjectName);
+#endif
+
+	DECLARE_EVENT(UMVVMBlueprintView, FOnBindingsUpdated);
+	FOnBindingsUpdated OnBindingsUpdated;
+
+	DECLARE_EVENT(UMVVMBlueprintView, FOnViewModelsUpdated);
+	FOnViewModelsUpdated OnViewModelsUpdated;
+
+private:
+	UPROPERTY(EditAnywhere, Category = "MVVM")
+	TArray<FMVVMBlueprintViewBinding> Bindings;
+
+	UPROPERTY(EditAnywhere, Category = "MVVM")
+	TArray<FMVVMBlueprintViewModelContext> AvailableViewModels;
+};
diff --git a/Engine/Plugins/Runtime/ModelViewViewModel/Source/ModelViewViewModelBlueprint/Public/MVVMBlueprintViewBinding.h b/Engine/Plugins/Runtime/ModelViewViewModel/Source/ModelViewViewModelBlueprint/Public/MVVMBlueprintViewBinding.h
new file mode 100644
--- /dev/null
+++ b/Engine/Plugins/Runtime/ModelViewViewModel/Source/ModelViewViewModelBlueprint/Public/MVVMBlueprintViewBinding.h
@@ -0,0 +1,63 @@
+// Copyright Epic Games, Inc. All Rights Reserved.
+
+#pragma once
+
+#include "MVVMPropertyPath.h"
+#include "Types/MVVMBindingMode.h"
+
+#include "MVVMBlueprintViewBinding.generated.h"
+
+/**
+*
+*/
+USTRUCT(BlueprintType)
+struct FMVVMBlueprintViewConversionPath
+{
+	GENERATED_BODY()
+
+	/** The Conversion function when setting the binding. */
+	UPROPERTY(EditAnywhere, Category = "MVVM", AdvancedDisplay)
+	FString SetConversionFunctionPath;
+
+	/** The Conversion function when getting the binding (in 2 way). */
+	UPROPERTY(EditAnywhere, Category = "MVVM", AdvancedDisplay)
+	FString GetConversionFunctionPath;
+};
+
+/**
+*
+*/
+USTRUCT(BlueprintType)
+struct FMVVMBlueprintViewBinding
+{
+	GENERATED_BODY()
+
+	UPROPERTY(EditAnywhere, Category = "MVVM")
+	FMVVMViewModelPropertyPath ViewModelPath;
+
+	UPROPERTY(EditAnywhere, Category = "MVVM")
+	FMVVMWidgetPropertyPath WidgetPath;
+
+	/** */
+	UPROPERTY(EditAnywhere, Category = "MVVM")
+	EMVVMBindingMode BindingType = EMVVMBindingMode::OneWayToDestination;
+
+	UPROPERTY(EditAnywhere, Category = "MVVM")
+	EMVVMViewBindingUpdateMode UpdateMode = EMVVMViewBindingUpdateMode::Immediate;
+
+	/** */
+	UPROPERTY(EditAnywhere, Category = "MVVM")
+	FMVVMBlueprintViewConversionPath Conversion;
+
+	/** */
+	UPROPERTY(VisibleAnywhere, Category = "MVVM", Transient)
+	TArray<FText> Errors;
+
+	/** Whether the binding is enabled or disabled by default. The instance may enable the binding at runtime. */
+	UPROPERTY(EditAnywhere, Category = "MVVM")
+	bool bEnabled = true;
+
+	/** The binding is visible in the editor, but is not compiled and cannot be used at runtime. */
+	UPROPERTY(EditAnywhere, Category = "MVVM")
+	bool bCompile = true;
+};
diff --git a/Engine/Plugins/Runtime/ModelViewViewModel/Source/ModelViewViewModelBlueprint/Public/MVVMBlueprintViewModel.h b/Engine/Plugins/Runtime/ModelViewViewModel/Source/ModelViewViewModelBlueprint/Public/MVVMBlueprintViewModel.h
new file mode 100644
--- /dev/null
+++ b/Engine/Plugins/Runtime/ModelViewViewModel/Source/ModelViewViewModelBlueprint/Public/MVVMBlueprintViewModel.h
@@ -0,0 +1,18 @@
+// Copyright Epic Games, Inc. All Rights Reserved.
+
+#pragma once
+
+#include "CoreMinimal.h"
+#include "UObject/Object.h"
+#include "Templates/SubclassOf.h"
+
+#include "MVVMBlueprintViewModel.generated.h"
+
+
+UCLASS(DisplayName="ViewModel Blueprint")
+class MODELVIEWVIEWMODELBLUEPRINT_API UMVVMBlueprintViewModel : public UObject
+{
+	GENERATED_BODY()
+
+public:
+};
diff --git a/Engine/Plugins/Runtime/ModelViewViewModel/Source/ModelViewViewModelBlueprint/Public/MVVMBlueprintViewModelContext.h b/Engine/Plugins/Runtime/ModelViewViewModel/Source/ModelViewViewModelBlueprint/Public/MVVMBlueprintViewModelContext.h
new file mode 100644
--- /dev/null
+++ b/Engine/Plugins/Runtime/ModelViewViewModel/Source/ModelViewViewModelBlueprint/Public/MVVMBlueprintViewModelContext.h
@@ -0,0 +1,80 @@
+// Copyright Epic Games, Inc. All Rights Reserved.
+
+#pragma once
+
+#include "CoreMinimal.h"
+#include "UObject/Object.h"
+#include "Templates/SubclassOf.h"
+#include "MVVMViewModelBase.h"
+
+#include "MVVMBlueprintViewModelContext.generated.h"
+
+/**
+ *
+ */
+UENUM()
+enum class EMVVMBlueprintViewModelContextCreationType : uint8
+{
+	CreateInstance,
+	GlobalViewModelCollection,
+	PropertyPath,
+};
+
+/**
+ *
+ */
+USTRUCT(BlueprintType)
+struct MODELVIEWVIEWMODELBLUEPRINT_API FMVVMBlueprintViewModelContext
+{
+	GENERATED_BODY()
+
+public:
+	FMVVMBlueprintViewModelContext() {}
+	FMVVMBlueprintViewModelContext(TSubclassOf<UMVVMViewModelBase> InClass, FGuid InGuid);
+
+	FGuid GetViewModelId() const
+	{
+		return ViewModelContextId;
+	}
+
+	FName GetViewModelName() const
+	{
+		//todo change that to the fname
+		return *GetDisplayName().ToString();
+	}
+
+	FText GetDisplayName() const;
+
+	TSubclassOf<UMVVMViewModelBase> GetViewModelClass() const
+	{
+		return ViewModelClass;
+	}
+
+private:
+	/** When the view is spawn, create an instance of the viewmodel. */
+	UPROPERTY(EditAnywhere, Category = "MVVM")
+	FGuid ViewModelContextId;
+
+	UPROPERTY(EditAnywhere, Category = "MVVM")
+	TSubclassOf<UMVVMViewModelBase> ViewModelClass;
+
+public:
+	UPROPERTY(EditAnywhere, Category = "MVVM")
+	FText OverrideDisplayName;
+
+	/** When the view is spawn, create an instance of the viewmodel. */
+	UPROPERTY(EditAnywhere, Category = "MVVM")
+	EMVVMBlueprintViewModelContextCreationType CreationType = EMVVMBlueprintViewModelContextCreationType::CreateInstance;
+
+	/** Identifier of an already registered viewmodel. */
+	UPROPERTY(EditAnywhere, Category = "MVVM", AdvancedDisplay, meta = (EditCondition = "CreationType == EMVVMBlueprintViewModelContextCreationType::GlobalViewModelCollection"))
+	FName GlobalViewModelIdentifier;
+
+	/** The Path to get the viewmodel instance. */
+	UPROPERTY(EditAnywhere, Category = "MVVM", AdvancedDisplay, meta = (EditCondition = "CreationType == EMVVMBlueprintViewModelContextCreationType::PropertyPath"))
+	FString ViewModelPropertyPath;
+
+	/** Generate a setter function for this viewmodel. */
+	UPROPERTY(EditAnywhere, Category = "MVVM", AdvancedDisplay)
+	bool bCreateSetterFunction = false;
+};
diff --git a/Engine/Plugins/Runtime/ModelViewViewModel/Source/ModelViewViewModelBlueprint/Public/MVVMPropertyPath.h b/Engine/Plugins/Runtime/ModelViewViewModel/Source/ModelViewViewModelBlueprint/Public/MVVMPropertyPath.h
new file mode 100644
--- /dev/null
+++ b/Engine/Plugins/Runtime/ModelViewViewModel/Source/ModelViewViewModelBlueprint/Public/MVVMPropertyPath.h
@@ -0,0 +1,105 @@
+// Copyright Epic Games, Inc. All Rights Reserved.
+
+#pragma once
+
+#include "Types/MVVMBindingName.h"
+
+#include "MVVMPropertyPath.generated.h"
+
+/**
+ *
+ */
+USTRUCT()
+struct FMVVMPropertyPathBase
+{
+	GENERATED_BODY()
+
+private:
+	/** Name we are looking for. */
+	UPROPERTY(EditAnywhere, Category = "MVVM")
+	FMVVMBindingName BindingName;
+
+	/** The path of the property (either the getter or the set action). */
+	UPROPERTY(EditAnywhere, Category = "MVVM")
+	FString GetterPropertyPath;
+
+	/** The path of the property (either the setter or the set action). */
+	UPROPERTY(EditAnywhere, Category = "MVVM")
+	FString SetterPropertyPath;
+
+public:
+	FMVVMBindingName GetBindingName() const
+	{
+		return BindingName;
+	}
+
+	void SetBindingName(FName InBindingName)
+	{
+		BindingName = FMVVMBindingName(InBindingName);
+	}
+
+	FString GetGetterPropertyPath() const
+	{
+		return !GetterPropertyPath.IsEmpty() ? GetterPropertyPath : BindingName.ToString();
+	}
+
+	FString GetSetterPropertyPath() const
+	{
+		return !SetterPropertyPath.IsEmpty() ? SetterPropertyPath : BindingName.ToString();
+	}
+
+	bool operator==(const FMVVMPropertyPathBase& Other) const
+	{
+		return BindingName == Other.BindingName &&
+			GetterPropertyPath == Other.GetterPropertyPath &&
+			SetterPropertyPath == Other.SetterPropertyPath;
+	}
+	bool operator!=(const FMVVMPropertyPathBase& Other) const
+	{
+		return !operator==(Other);
+	}
+};
+
+USTRUCT(BlueprintType)
+struct FMVVMViewModelPropertyPath : public FMVVMPropertyPathBase
+{
+	GENERATED_BODY();
+
+	UPROPERTY(EditAnywhere, Category = "MVVM")
+	FGuid ContextId;
+
+	bool operator==(const FMVVMViewModelPropertyPath& Other) const
+	{
+		return ContextId == Other.ContextId &&
+			FMVVMPropertyPathBase::operator==(Other);
+	}
+
+	bool operator!=(const FMVVMViewModelPropertyPath& Other) const
+	{
+		return ContextId != Other.ContextId &&
+			FMVVMPropertyPathBase::operator!=(Other);
+	}
+};
+
+USTRUCT(BlueprintType)
+struct FMVVMWidgetPropertyPath : public FMVVMPropertyPathBase
+{
+	GENERATED_BODY();
+
+	/** The context from which the destination path will be evaluated from. */
+	UPROPERTY(EditAnywhere, Category = "MVVM")
+	FName WidgetName;
+
+	bool operator==(const FMVVMWidgetPropertyPath& Other) const
+	{
+		return WidgetName == Other.WidgetName &&
+			FMVVMPropertyPathBase::operator==(Other);
+	}
+
+
+	bool operator!=(const FMVVMWidgetPropertyPath& Other) const
+	{
+		return WidgetName != Other.WidgetName ||
+			FMVVMPropertyPathBase::operator!=(Other);
+	}
+};
diff --git a/Engine/Plugins/Runtime/ModelViewViewModel/Source/ModelViewViewModelBlueprint/Public/MVVMWidgetBlueprintExtension_View.h b/Engine/Plugins/Runtime/ModelViewViewModel/Source/ModelViewViewModelBlueprint/Public/MVVMWidgetBlueprintExtension_View.h
new file mode 100644
--- /dev/null
+++ b/Engine/Plugins/Runtime/ModelViewViewModel/Source/ModelViewViewModelBlueprint/Public/MVVMWidgetBlueprintExtension_View.h
@@ -0,0 +1,67 @@
+// Copyright Epic Games, Inc. All Rights Reserved.
+
+#pragma once
+
+#include "CoreMinimal.h"
+#include "WidgetBlueprintExtension.h"
+
+#include "MVVMWidgetBlueprintExtension_View.generated.h"
+
+class UMVVMBlueprintView;
+class FWidgetBlueprintCompilerContext;
+class UWidgetBlueprintGeneratedClass;
+
+namespace UE::MVVM::Private
+{
+	struct FMVVMViewBlueprintCompiler;
+} //namespace
+
+/**
+ *
+ */
+UCLASS()
+class MODELVIEWVIEWMODELBLUEPRINT_API UMVVMWidgetBlueprintExtension_View : public UWidgetBlueprintExtension
+{
+	GENERATED_BODY()
+
+public:
+	void CreateBlueprintViewInstance();
+	void DestroyBlueprintViewInstance();
+
+	UMVVMBlueprintView* GetBlueprintView()
+	{
+		return BlueprintView;
+	}
+
+	const UMVVMBlueprintView* GetBlueprintView() const
+	{
+		return BlueprintView;
+	}
+
+	FSimpleMulticastDelegate& OnBlueprintViewChangedDelegate()
+	{
+		return BlueprintViewChangedDelegate;
+	}
+
+public:
+	//~ Begin UObject interface
+	virtual void PostLoad() override;
+	//~ End UObject interface
+
+	//~ Begin UWidgetBlueprintExtension interface
+	virtual void HandlePreloadObjectsForCompilation(UBlueprint* OwningBlueprint) override;
+	virtual void HandleBeginCompilation(FWidgetBlueprintCompilerContext& InCreationContext) override;
+	virtual void HandleCleanAndSanitizeClass(UWidgetBlueprintGeneratedClass* ClassToClean, UObject* OldCDO) override;
+	virtual void HandleCreateClassVariablesFromBlueprint(const FWidgetBlueprintCompilerContext::FCreateVariableContext& Context) override;
+	virtual void HandleCreateFunctionList() override;
+	virtual void HandleFinishCompilingClass(UWidgetBlueprintGeneratedClass* Class) override;
+	virtual void HandleEndCompilation() override;
+	//~ End UWidgetBlueprintExtension interface
+
+private:
+	UPROPERTY(Instanced)
+	TObjectPtr<UMVVMBlueprintView> BlueprintView;
+
+	FSimpleMulticastDelegate BlueprintViewChangedDelegate;
+	TPimplPtr<UE::MVVM::Private::FMVVMViewBlueprintCompiler> CurrentCompilerContext;
+};
diff --git a/Engine/Plugins/Runtime/ModelViewViewModel/Source/ModelViewViewModelEditor/Private/Customizations/MVVMBindPropertiesDetailView.cpp b/Engine/Plugins/Runtime/ModelViewViewModel/Source/ModelViewViewModelEditor/Private/Customizations/MVVMBindPropertiesDetailView.cpp
new file mode 100644
--- /dev/null
+++ b/Engine/Plugins/Runtime/ModelViewViewModel/Source/ModelViewViewModelEditor/Private/Customizations/MVVMBindPropertiesDetailView.cpp
@@ -0,0 +1,182 @@
+// Copyright Epic Games, Inc. All Rights Reserved.
+
+#include "Customizations/MVVMBindPropertiesDetailView.h"
+
+#include "Components/Widget.h"
+#include "WidgetBlueprint.h"
+
+#include "MVVMWidgetBlueprintExtension_View.h"
+#include "MVVMPropertyAccess.h"
+
+#include "EditorStyleSet.h"
+#include "PropertyEditorDelegates.h"
+#include "PropertyHandle.h"
+
+#include "Widgets/SNullWidget.h"
+
+#define LOCTEXT_NAMESPACE "BindPropertiesDetailView"
+
+bool FMVVMBindPropertiesDetailView::IsSupported(const FOnGenerateGlobalRowExtensionArgs& InArgs)
+{
+	if (!InArgs.PropertyHandle)
+	{
+		return false;
+	}
+
+	if (!InArgs.PropertyHandle->GetOuterBaseClass()->IsChildOf(UWidget::StaticClass()))
+	{
+		return false;
+	}
+
+	static const FName BlueprintSetterMetaData = TEXT("BlueprintSetter");
+	if (!InArgs.PropertyHandle->GetProperty()->HasMetaData(BlueprintSetterMetaData))
+	{
+		return false;
+	}
+
+	return true;
+}
+
+bool FMVVMBindPropertiesDetailView::GetBindingInfo(const TSharedPtr<IPropertyHandle>& PropertyHandle, FBinding& OutBinding)
+{
+	TArray<UObject*> OuterObjects;
+	PropertyHandle->GetOuterObjects(OuterObjects);
+	if (OuterObjects.Num() == 0)
+	{
+		return false;
+	}
+
+	UWidgetBlueprint* WidgetBlueprint = nullptr;
+	for (UObject* OuterObject : OuterObjects)
+	{
+		if (UWidget* Widget = Cast<UWidget>(OuterObject))
+		{
+			if (!Widget->IsDesignTime())
+			{
+				return false;
+			}
+
+			UWidgetBlueprint* CurrentWidgetBlueprint = Cast<UWidgetBlueprint>(Widget->WidgetGeneratedBy);
+			if (!CurrentWidgetBlueprint)
+			{
+				return false;
+			}
+			if (WidgetBlueprint && CurrentWidgetBlueprint != WidgetBlueprint)
+			{
+				return false;
+			}
+
+			OutBinding.Widgets.Add(Widget);
+			WidgetBlueprint = CurrentWidgetBlueprint;
+		}
+		else
+		{
+			return false;
+		}
+	}
+
+
+	OutBinding.WidgetBlueprint = WidgetBlueprint;
+	OutBinding.PropertyHandle = PropertyHandle;
+	return true;
+}
+
+void FMVVMBindPropertiesDetailView::CreatePropertyRowExtension(const FOnGenerateGlobalRowExtensionArgs& InArgs, TArray<FPropertyRowExtensionButton>& OutExtensions)
+{
+	if (IsSupported(InArgs))
+	{
+		FBinding BindingInfo;
+		if (GetBindingInfo(InArgs.PropertyHandle, BindingInfo))
+		{
+			FPropertyRowExtensionButton& ExposeButton = OutExtensions.AddDefaulted_GetRef();
+			ExposeButton.Icon = FSlateIcon(FEditorStyle::Get().GetStyleSetName(), "Level.ScriptIcon16x");
+			ExposeButton.Label = LOCTEXT("ExposeProperty", "MVVM Property");
+			ExposeButton.ToolTip = LOCTEXT("ExposePropertyToolTip", "Bind a View property to a ViewModel.");
+			ExposeButton.UIAction = FUIAction(
+				FExecuteAction(),
+				//FExecuteAction::CreateRaw(this, &FMVVMBindPropertiesDetailView::HandleExecutePropertyExposed, BindingInfo),
+				FCanExecuteAction(),
+				FGetActionCheckState::CreateRaw(this, &FMVVMBindPropertiesDetailView::GetPropertyExposedCheckState, BindingInfo)
+			);
+			//ExposeButton.MenuContentGenerator = FOnGetContent::CreateRaw(this, &FMVVMBindPropertiesDetailView::HandleExecutePropertyExposed, BindingInfo);
+		}
+	}
+}
+//
+//#include "Widgets/Input/SButton.h"
+//TSharedRef<SWidget> FMVVMBindPropertiesDetailView::HandleExecutePropertyExposed(FBinding BindingInfo)
+//{
+//	if (UWidgetBlueprint* WidgetBlueprint = BindingInfo.WidgetBlueprint.Get())
+//	{
+//			if (TSharedPtr<IPropertyHandle> PropertyHandle = BindingInfo.PropertyHandle.Pin())
+//			{
+//				Poppup = MakeShared<FMVVMPropertyAccess>(WidgetBlueprint, PropertyHandle.ToSharedRef());
+//				TSharedPtr<SWidget> MenuWidget = Poppup->MakeViewModelBindingMenu();
+//				if (MenuWidget)
+//				{
+//					return MenuWidget.ToSharedRef();
+//				}
+//			}
+// 
+//			UAssetEditorSubsystem* AssetEditorSubsystem = GEditor->GetEditorSubsystem<UAssetEditorSubsystem>();
+//			if (IAssetEditorInstance* AssetEditorInstance = FindEditorForAsset(UObject * Asset, bool bFocusIfOpen))
+//			{
+//				if (AssetEditorInstance->GetEditorName() == TEXT("WidgetBlueprintEditor"))
+//				{
+//					FWidgetBlueprintEditor* WidgetBlueprintEditor = static_cast<FWidgetBlueprintEditor>(AssetEditorInstance);
+//					check(WidgetBlueprintEditor->GetBlueprintObj() == WidgetBlueprint);
+//				}
+//			}
+//	}
+//	return SNullWidget::NullWidget;
+//}
+
+
+ECheckBoxState FMVVMBindPropertiesDetailView::GetPropertyExposedCheckState(FBinding BindingInfo) const
+{
+	//if (UWidgetBlueprint* WidgetBlueprint = BindingInfo.WidgetBlueprint.Get())
+	//{
+	//	if (UMVVMViewModelBlueprintExtension* Extension = UWidgetBlueprintExtension::GetExtension<UMVVMViewModelBlueprintExtension>(WidgetBlueprint))
+	//	{
+	//		if (TSharedPtr<IPropertyHandle> PropertyHandle = BindingInfo.PropertyHandle.Pin())
+	//		{
+	//			int32 ValidCount = 0;
+	//			int32 InvalidValidCount = 0;
+	//			for (TWeakObjectPtr<UWidget>& WeakObject : BindingInfo.Widgets)
+	//			{
+	//				if (UObject* Object = WeakObject.Get())
+	//				{
+	//					if (Extension->FindViewBinding(Object, PropertyHandle.Get()))
+	//					{
+	//						++ValidCount;
+	//					}
+	//				}
+	//				else
+	//				{
+	//					return ECheckBoxState::Undetermined;
+	//				}
+	//			}
+
+	//			if (ValidCount == BindingInfo.Widgets.Num())
+	//			{
+	//				return ECheckBoxState::Checked;
+	//			}
+	//			else if (ValidCount == 0)
+	//			{
+	//				return ECheckBoxState::Unchecked;
+	//			}
+	//			else
+	//			{
+	//				return ECheckBoxState::Undetermined;
+	//			}
+	//		}
+	//		else
+	//		{
+	//			return ECheckBoxState::Undetermined;
+	//		}
+	//	}
+	//}
+	return ECheckBoxState::Unchecked;
+}
+
+#undef LOCTEXT_NAMESPACE
diff --git a/Engine/Plugins/Runtime/ModelViewViewModel/Source/ModelViewViewModelEditor/Private/Customizations/MVVMBindPropertiesDetailView.h b/Engine/Plugins/Runtime/ModelViewViewModel/Source/ModelViewViewModelEditor/Private/Customizations/MVVMBindPropertiesDetailView.h
new file mode 100644
--- /dev/null
+++ b/Engine/Plugins/Runtime/ModelViewViewModel/Source/ModelViewViewModelEditor/Private/Customizations/MVVMBindPropertiesDetailView.h
@@ -0,0 +1,38 @@
+// Copyright Epic Games, Inc. All Rights Reserved.
+
+#pragma once
+
+#include "CoreMinimal.h"
+#include "UObject/WeakObjectPtr.h"
+#include "UObject/WeakObjectPtrTemplates.h"
+
+enum class ECheckBoxState : uint8;
+class IPropertyHandle;
+class UWidget;
+class UWidgetBlueprint;
+struct FOnGenerateGlobalRowExtensionArgs;
+struct FPropertyRowExtensionButton;
+
+class FMVVMPropertyAccess;
+
+class FMVVMBindPropertiesDetailView
+{
+public:
+	static bool IsSupported(const FOnGenerateGlobalRowExtensionArgs& InArgs);
+	void CreatePropertyRowExtension(const FOnGenerateGlobalRowExtensionArgs& InArgs, TArray<FPropertyRowExtensionButton>& OutExtensions);
+
+private:
+	struct FBinding
+	{
+		TArray<TWeakObjectPtr<UWidget>> Widgets;
+		TWeakObjectPtr<UWidgetBlueprint> WidgetBlueprint;
+		TWeakPtr<IPropertyHandle> PropertyHandle;
+	};
+
+	bool GetBindingInfo(const TSharedPtr<IPropertyHandle>& Handle, FBinding& Out);
+	//TSharedRef<SWidget> HandleExecutePropertyExposed(FBinding);
+	ECheckBoxState GetPropertyExposedCheckState(FBinding) const;
+
+private:
+	TSharedPtr<FMVVMPropertyAccess> Poppup;
+};
\ No newline at end of file
diff --git a/Engine/Plugins/Runtime/ModelViewViewModel/Source/ModelViewViewModelEditor/Private/Customizations/MVVMConversionPathCustomization.cpp b/Engine/Plugins/Runtime/ModelViewViewModel/Source/ModelViewViewModelEditor/Private/Customizations/MVVMConversionPathCustomization.cpp
new file mode 100644
--- /dev/null
+++ b/Engine/Plugins/Runtime/ModelViewViewModel/Source/ModelViewViewModelEditor/Private/Customizations/MVVMConversionPathCustomization.cpp
@@ -0,0 +1,152 @@
+// Copyright Epic Games, Inc. All Rights Reserved.
+
+#include "MVVMConversionPathCustomization.h"
+#include "DetailWidgetRow.h"
+#include "IDetailChildrenBuilder.h"
+#include "MVVMBlueprintViewBinding.h"
+#include "PropertyHandle.h"
+#include "WidgetBlueprint.h"
+#include "Widgets/SMVVMConversionPath.h"
+#include "Widgets/Input/SEditableTextBox.h"
+#include "Widgets/Text/STextBlock.h"
+
+#define LOCTEXT_NAMESPACE "MVVMConversionPath"
+
+UE::MVVM::FConversionPathCustomization::FConversionPathCustomization(UWidgetBlueprint* InWidgetBlueprint)
+{
+	WidgetBlueprint = InWidgetBlueprint;
+}
+
+void UE::MVVM::FConversionPathCustomization::CustomizeHeader(TSharedRef<IPropertyHandle> InPropertyHandle, FDetailWidgetRow& HeaderRow, IPropertyTypeCustomizationUtils& CustomizationUtils)
+{
+	SetterProperty = InPropertyHandle->GetChildHandle(GET_MEMBER_NAME_CHECKED(FMVVMBlueprintViewConversionPath, SetConversionFunctionPath));
+	GetterProperty = InPropertyHandle->GetChildHandle(GET_MEMBER_NAME_CHECKED(FMVVMBlueprintViewConversionPath, GetConversionFunctionPath));
+
+	HeaderRow.NameContent()
+	[
+		InPropertyHandle->CreatePropertyNameWidget()
+	]
+	.ValueContent()
+	[
+		InPropertyHandle->CreatePropertyValueWidget()
+	];
+}
+
+FText UE::MVVM::FConversionPathCustomization::GetGetterPath() const
+{
+	FString Value;
+	FPropertyAccess::Result Result = GetterProperty->GetValue(Value);
+
+	if (Result == FPropertyAccess::Success)
+	{
+		return FText::FromString(Value);
+	}
+
+	if (Result == FPropertyAccess::MultipleValues)
+	{
+		return LOCTEXT("MultipleValues", "Multiple Values");
+	}
+
+	return FText::GetEmpty();
+}
+
+FText UE::MVVM::FConversionPathCustomization::GetSetterPath() const
+{
+	FString Value;
+	FPropertyAccess::Result Result = SetterProperty->GetValue(Value);
+
+	if (Result == FPropertyAccess::Success)
+	{
+		return FText::FromString(Value);
+	}
+
+	if (Result == FPropertyAccess::MultipleValues)
+	{
+		return LOCTEXT("MultipleValues", "Multiple Values");
+	}
+
+	return FText::GetEmpty();
+}
+
+void UE::MVVM::FConversionPathCustomization::OnTextCommitted(const FText& NewValue, ETextCommit::Type CommitType, bool bIsGetter)
+{
+	FString NewString = NewValue.ToString();
+	OnFunctionPathChanged(NewString, bIsGetter);
+}
+
+void UE::MVVM::FConversionPathCustomization::OnFunctionPathChanged(const FString& NewPath, bool bIsGetter)
+{
+	TSharedPtr<IPropertyHandle> Handle = bIsGetter ? GetterProperty : SetterProperty;
+	Handle->SetValue(NewPath);
+}
+
+void UE::MVVM::FConversionPathCustomization::CustomizeChildren(TSharedRef<IPropertyHandle> InPropertyHandle, IDetailChildrenBuilder& ChildBuilder, IPropertyTypeCustomizationUtils& CustomizationUtils)
+{
+	TSharedPtr<IPropertyHandle> ParentHandle = InPropertyHandle->GetParentHandle();
+	TArray<void*> RawData;
+	ParentHandle->AccessRawData(RawData);
+
+	TArray<FMVVMBlueprintViewBinding*> ViewBindings;
+	for (void* Data : RawData)
+	{
+		ViewBindings.Add((FMVVMBlueprintViewBinding*) Data);
+	}
+
+	// setter
+	ChildBuilder.AddProperty(SetterProperty.ToSharedRef())
+		.CustomWidget()
+		.NameContent()
+		[
+			SNew(STextBlock)
+			.Text(LOCTEXT("SetterLabel", "Setter"))
+		]
+		.ValueContent()
+		[
+			SNew(SHorizontalBox)
+			+ SHorizontalBox::Slot()
+			.VAlign(VAlign_Center)
+			[
+				SNew(SEditableTextBox)
+				.Text(this, &FConversionPathCustomization::GetSetterPath)
+				.OnTextCommitted(this, &FConversionPathCustomization::OnTextCommitted, false)
+			]
+			+ SHorizontalBox::Slot()
+			.VAlign(VAlign_Center)
+			.AutoWidth()
+			[
+				SNew(SMVVMConversionPath, WidgetBlueprint, false)
+				.Bindings(ViewBindings)
+				.OnFunctionChanged(this, &FConversionPathCustomization::OnFunctionPathChanged, false)
+			]
+		];
+
+	// getter
+	ChildBuilder.AddProperty(GetterProperty.ToSharedRef())
+		.CustomWidget()
+		.NameContent()
+		[
+			SNew(STextBlock)
+			.Text(LOCTEXT("GetterLabel", "Getter"))
+		]
+		.ValueContent()
+		[
+			SNew(SHorizontalBox)
+			+ SHorizontalBox::Slot()
+			.VAlign(VAlign_Center)
+			[
+				SNew(SEditableTextBox)
+				.Text(this, &FConversionPathCustomization::GetGetterPath)
+				.OnTextCommitted(this, &FConversionPathCustomization::OnTextCommitted, true)
+			]
+			+ SHorizontalBox::Slot()
+			.VAlign(VAlign_Center)
+			.AutoWidth()
+			[
+				SNew(SMVVMConversionPath, WidgetBlueprint, true)
+				.Bindings(ViewBindings)
+				.OnFunctionChanged(this, &FConversionPathCustomization::OnFunctionPathChanged, true)
+			]
+		];
+}
+
+#undef LOCTEXT_NAMESPACE
\ No newline at end of file
diff --git a/Engine/Plugins/Runtime/ModelViewViewModel/Source/ModelViewViewModelEditor/Private/Customizations/MVVMConversionPathCustomization.h b/Engine/Plugins/Runtime/ModelViewViewModel/Source/ModelViewViewModelEditor/Private/Customizations/MVVMConversionPathCustomization.h
new file mode 100644
--- /dev/null
+++ b/Engine/Plugins/Runtime/ModelViewViewModel/Source/ModelViewViewModelEditor/Private/Customizations/MVVMConversionPathCustomization.h
@@ -0,0 +1,36 @@
+// Copyright Epic Games, Inc. All Rights Reserved.
+
+#pragma once
+
+#include "IPropertyTypeCustomization.h"
+
+class UWidgetBlueprint;
+
+namespace UE::MVVM
+{
+	class FConversionPathCustomization : public IPropertyTypeCustomization
+	{
+	public:
+		static TSharedRef<IPropertyTypeCustomization> MakeInstance(UWidgetBlueprint* InWidgetBlueprint)
+		{
+			return MakeShared<FConversionPathCustomization>(InWidgetBlueprint);
+		}
+
+		FConversionPathCustomization(UWidgetBlueprint* InWidgetBlueprint);
+
+		virtual void CustomizeHeader(TSharedRef<IPropertyHandle> InPropertyHandle, FDetailWidgetRow& HeaderRow, IPropertyTypeCustomizationUtils& CustomizationUtils) override;
+		virtual void CustomizeChildren(TSharedRef<IPropertyHandle> InPropertyHandle, IDetailChildrenBuilder& ChildBuilder, IPropertyTypeCustomizationUtils& CustomizationUtils) override;
+
+	private:
+		FText GetGetterPath() const;
+		FText GetSetterPath() const;
+
+		void OnTextCommitted(const FText& NewValue, ETextCommit::Type CommitType, bool bIsGetter);
+		void OnFunctionPathChanged(const FString& NewPath, bool bIsGetter);
+
+	private:
+		UWidgetBlueprint* WidgetBlueprint = nullptr;
+		TSharedPtr<IPropertyHandle> GetterProperty;
+		TSharedPtr<IPropertyHandle> SetterProperty;
+	};
+}
\ No newline at end of file
diff --git a/Engine/Plugins/Runtime/ModelViewViewModel/Source/ModelViewViewModelEditor/Private/Customizations/MVVMPropertyBindingExtension.cpp b/Engine/Plugins/Runtime/ModelViewViewModel/Source/ModelViewViewModelEditor/Private/Customizations/MVVMPropertyBindingExtension.cpp
new file mode 100644
--- /dev/null
+++ b/Engine/Plugins/Runtime/ModelViewViewModel/Source/ModelViewViewModelEditor/Private/Customizations/MVVMPropertyBindingExtension.cpp
@@ -0,0 +1,199 @@
+// Copyright Epic Games, Inc. All Rights Reserved.
+
+#include "Customizations/MVVMPropertyBindingExtension.h"
+
+#include "Components/Widget.h"
+#include "Framework/MultiBox/MultiBoxBuilder.h"
+#include "Kismet2/BlueprintEditorUtils.h"
+#include "MVVMBlueprintView.h"
+#include "MVVMWidgetBlueprintExtension_View.h"
+#include "Styling/StyleColors.h"
+#include "UObject/UnrealType.h"
+#include "WidgetBlueprint.h"
+#include "Widgets/Images/SImage.h"
+#include "Widgets/SBoxPanel.h"
+#include "Widgets/Text/STextBlock.h"
+
+#define LOCTEXT_NAMESPACE "MVVMPropertyBindingExtension"
+
+static void ExtendBindingsMenu(FMenuBuilder& MenuBuilder, const UWidgetBlueprint* WidgetBlueprint, const UWidget* Widget, const FProperty* Property)
+{
+	MenuBuilder.BeginSection("ViewModels", LOCTEXT("ViewModels", "View Models"));
+
+	UMVVMWidgetBlueprintExtension_View* MVVMExtensionPtr = UMVVMWidgetBlueprintExtension_View::GetExtension<UMVVMWidgetBlueprintExtension_View>(WidgetBlueprint);
+	if (MVVMExtensionPtr == nullptr)
+	{
+		return;
+	}
+	UMVVMBlueprintView* MVVMBlueprintView = MVVMExtensionPtr->GetBlueprintView();
+	if (MVVMBlueprintView == nullptr)
+	{
+		return;
+	}
+
+	const UEdGraphSchema_K2* Schema = GetDefault<UEdGraphSchema_K2>();
+
+	auto CreatePropertyWidget = [Schema](const FProperty* Property, bool bRequiresConversion)
+		-> TSharedRef<SWidget>
+	{
+		FEdGraphPinType PinType;
+		Schema->ConvertPropertyToPinType(Property, PinType);
+		const FSlateBrush* SlateBrush = FBlueprintEditorUtils::GetIconFromPin(PinType, true);
+
+		TSharedRef<SHorizontalBox> HorizontalBox =
+			SNew(SHorizontalBox)
+			.ToolTipText(Property->GetDisplayNameText())
+			+ SHorizontalBox::Slot()
+			.VAlign(VAlign_Center)
+			.Padding(0, 0, 5, 0)
+			.AutoWidth()
+			[
+				SNew(SImage)
+				.Image(SlateBrush)
+				.ColorAndOpacity(Schema->GetPinTypeColor(PinType))
+			]
+			+ SHorizontalBox::Slot()
+			.Padding(0, 0, 5, 0)
+			.VAlign(VAlign_Center)
+			[
+				SNew(STextBlock)
+				.Text(Property->GetDisplayNameText())
+				.OverflowPolicy(ETextOverflowPolicy::Ellipsis)
+			];
+
+		if (bRequiresConversion)
+		{
+			HorizontalBox->AddSlot()
+				.VAlign(VAlign_Center)
+				.AutoWidth()
+				[
+					SNew(SImage)
+					.Image(FAppStyle::Get().GetBrush("Sequencer.CreateEventBinding"))
+					.ColorAndOpacity(FSlateColor(EStyleColor::AccentGreen))
+				];
+		}
+
+		return HorizontalBox;
+	};
+
+	auto CreateBinding = [MVVMBlueprintView](const UWidget* Widget, const FProperty* Property, FGuid ViewModelId, const FProperty* ViewModelProperty)
+	{
+		FMVVMBlueprintViewBinding& NewBinding = MVVMBlueprintView->AddDefaultBinding();
+
+		NewBinding.ViewModelPath.ContextId = ViewModelId;
+		NewBinding.ViewModelPath.SetBindingName(ViewModelProperty->GetFName());
+
+		NewBinding.WidgetPath.WidgetName = Widget->GetFName();
+		NewBinding.WidgetPath.SetBindingName(Property->GetFName());
+
+		NewBinding.BindingType = EMVVMBindingMode::OneWayToDestination;
+	};
+
+	for (const FMVVMBlueprintViewModelContext& ViewModel : MVVMBlueprintView->GetViewModels())
+	{
+		if (ViewModel.GetViewModelClass() == nullptr)
+		{
+			// invalid viewmodel, possibly just created by the user but not filled in, skip it for now
+			continue;
+		}
+
+		MenuBuilder.AddSubMenu(ViewModel.GetDisplayName(), ViewModel.GetDisplayName(),
+			FNewMenuDelegate::CreateLambda([ViewModel, Widget, Property, Schema, CreatePropertyWidget, CreateBinding](FMenuBuilder& MenuBuilder)
+				{
+					const UClass* ViewModelClass = ViewModel.GetViewModelClass().Get();
+
+					MenuBuilder.BeginSection("ValidDataTypes", LOCTEXT("ValidDataTypes", "Valid Data Types"));
+
+					for (TFieldIterator<FProperty> It(ViewModelClass); It; ++It)
+					{
+						const FProperty* VMProperty = *It;
+						if (VMProperty->GetClass() != Property->GetClass())
+						{
+							continue;
+						}
+
+						FUIAction UIAction;
+						UIAction.ExecuteAction = FExecuteAction::CreateLambda(CreateBinding, Widget, Property, ViewModel.GetViewModelId(), VMProperty);
+						MenuBuilder.AddMenuEntry(
+							UIAction,
+							CreatePropertyWidget(VMProperty, false)
+						);
+					}
+
+					MenuBuilder.EndSection();
+
+					MenuBuilder.BeginSection("InvalidDataTypes", LOCTEXT("InvalidDataTypes", "Invalid Data Types"));
+
+					for (TFieldIterator<FProperty> It(ViewModelClass); It; ++It)
+					{
+						const FProperty* VMProperty = *It;
+						if (VMProperty->GetClass() == Property->GetClass())
+						{
+							continue;
+						}
+
+						FUIAction UIAction;
+						UIAction.ExecuteAction = FExecuteAction::CreateLambda(CreateBinding, Widget, Property, ViewModel.GetViewModelId(), VMProperty);
+						MenuBuilder.AddMenuEntry(
+							UIAction,
+							CreatePropertyWidget(VMProperty, true)
+						);
+					}
+				}));
+	}
+
+	MenuBuilder.EndSection();
+}
+
+TOptional<FName> FMVVMPropertyBindingExtension::GetCurrentValue(const UWidgetBlueprint* WidgetBlueprint, const UWidget* Widget, const FProperty* Property) const
+{
+	const UMVVMWidgetBlueprintExtension_View* MVVMExtensionPtr = UMVVMWidgetBlueprintExtension_View::GetExtension<UMVVMWidgetBlueprintExtension_View>(WidgetBlueprint);
+	if (MVVMExtensionPtr == nullptr)
+	{
+		return TOptional<FName>();
+	}
+	const UMVVMBlueprintView* MVVMBlueprintView = MVVMExtensionPtr->GetBlueprintView();
+
+	const FMVVMBlueprintViewBinding* Binding = MVVMBlueprintView ? MVVMBlueprintView->FindBinding(Widget, Property) : nullptr;
+	if (Binding == nullptr)
+	{
+		return TOptional<FName>();
+	}
+
+	return Binding->ViewModelPath.GetBindingName().ToName();
+}
+
+
+void FMVVMPropertyBindingExtension::ClearCurrentValue(const UWidgetBlueprint* WidgetBlueprint, const UWidget* Widget, const FProperty* Property)
+{
+	if (UMVVMWidgetBlueprintExtension_View* MVVMExtensionPtr = UMVVMWidgetBlueprintExtension_View::GetExtension<UMVVMWidgetBlueprintExtension_View>(WidgetBlueprint))
+	{
+		if (UMVVMBlueprintView* MVVMBlueprintView = MVVMExtensionPtr->GetBlueprintView())
+		{
+			if (FMVVMBlueprintViewBinding* Binding = MVVMBlueprintView->FindBinding(Widget, Property))
+			{
+				Binding->ViewModelPath.SetBindingName(FName());
+			}
+		}
+	}
+}
+
+TSharedPtr<FExtender> FMVVMPropertyBindingExtension::CreateMenuExtender(const UWidgetBlueprint* WidgetBlueprint, const UWidget* Widget, const FProperty* Property)
+{
+	TSharedPtr<FExtender> Extender = MakeShared<FExtender>();
+	Extender->AddMenuExtension("BindingActions", EExtensionHook::After, nullptr, FMenuExtensionDelegate::CreateStatic(&ExtendBindingsMenu, WidgetBlueprint, Widget, Property));
+	return Extender;
+}
+
+bool FMVVMPropertyBindingExtension::CanExtend(const UWidgetBlueprint* WidgetBlueprint, const UWidget* Widget, const FProperty* Property) const
+{
+	const UMVVMWidgetBlueprintExtension_View* MVVMExtensionPtr = UMVVMWidgetBlueprintExtension_View::GetExtension<UMVVMWidgetBlueprintExtension_View>(WidgetBlueprint);
+	if (MVVMExtensionPtr == nullptr)
+	{
+		return false;
+	}
+	const UMVVMBlueprintView* MVVMBlueprintView = MVVMExtensionPtr->GetBlueprintView();
+	return MVVMBlueprintView ? MVVMBlueprintView->GetViewModels().Num() > 0 : false;
+}
+
+#undef LOCTEXT_NAMESPACE
diff --git a/Engine/Plugins/Runtime/ModelViewViewModel/Source/ModelViewViewModelEditor/Private/Customizations/MVVMPropertyBindingExtension.h b/Engine/Plugins/Runtime/ModelViewViewModel/Source/ModelViewViewModelEditor/Private/Customizations/MVVMPropertyBindingExtension.h
new file mode 100644
--- /dev/null
+++ b/Engine/Plugins/Runtime/ModelViewViewModel/Source/ModelViewViewModelEditor/Private/Customizations/MVVMPropertyBindingExtension.h
@@ -0,0 +1,17 @@
+// Copyright Epic Games, Inc. All Rights Reserved.
+
+#pragma once
+
+#include "IHasPropertyBindingExtensibility.h"
+
+class UWidgetBlueprint;
+class UWidget;
+
+class FMVVMPropertyBindingExtension
+	: public IPropertyBindingExtension
+{
+	virtual TOptional<FName> GetCurrentValue(const UWidgetBlueprint* WidgetBlueprint, const UWidget* Widget, const FProperty* Property) const override;
+	virtual void ClearCurrentValue(const UWidgetBlueprint* WidgetBlueprint, const UWidget* Widget, const FProperty* Property) override;
+	virtual TSharedPtr<FExtender> CreateMenuExtender(const UWidgetBlueprint* WidgetBlueprint, const UWidget* Widget, const FProperty* Property) override;
+	virtual bool CanExtend(const UWidgetBlueprint* WidgetBlueprint, const UWidget* Widget, const FProperty* Property) const;
+};
diff --git a/Engine/Plugins/Runtime/ModelViewViewModel/Source/ModelViewViewModelEditor/Private/Customizations/MVVMPropertyPathCustomization.cpp b/Engine/Plugins/Runtime/ModelViewViewModel/Source/ModelViewViewModelEditor/Private/Customizations/MVVMPropertyPathCustomization.cpp
new file mode 100644
--- /dev/null
+++ b/Engine/Plugins/Runtime/ModelViewViewModel/Source/ModelViewViewModelEditor/Private/Customizations/MVVMPropertyPathCustomization.cpp
@@ -0,0 +1,225 @@
+// Copyright Epic Games, Inc. All Rights Reserved.
+
+#include "MVVMPropertyPathCustomization.h"
+
+#include "Algo/RemoveIf.h"
+#include "Algo/Transform.h"
+#include "BlueprintEditor.h"
+#include "Components/Widget.h"
+#include "DetailWidgetRow.h"
+#include "MVVMWidgetBlueprintExtension_View.h"
+#include "PropertyHandle.h"
+#include "ScopedTransaction.h"
+#include "SSimpleButton.h"
+#include "Templates/UnrealTemplate.h"
+#include "WidgetBlueprint.h"
+#include "Widgets/Images/SLayeredImage.h"
+#include "Widgets/Input/SComboBox.h" 
+#include "Widgets/SMVVMFieldIcon.h"
+#include "Widgets/SMVVMFieldSelector.h"
+#include "Widgets/SMVVMSourceSelector.h"
+#include "Widgets/SOverlay.h"
+#include "Widgets/Text/STextBlock.h"
+
+#define LOCTEXT_NAMESPACE "MVVMPropertyPathCustomization"
+
+namespace UE::MVVM
+{
+	namespace Private
+	{
+		// figure out what the selection is - if any is different in the case of a multiselection, then return null
+		// this is required for getting the selection of the other pair
+		TArray<TSharedPtr<IFieldPathHelper>> GetPathHelpersForProperty(const TSharedPtr<IPropertyHandle>& PropertyHandle, UWidgetBlueprint* WidgetBlueprint, EPropertyPathType PathType)
+		{
+			TArray<TSharedPtr<IFieldPathHelper>> Result;
+
+			TArray<void*> RawData;
+			PropertyHandle->AccessRawData(RawData);
+
+			for (void* Datum : RawData)
+			{
+				const FProperty* ThisSelected = nullptr;
+				if (PathType == EPropertyPathType::ViewModel)
+				{
+					Result.Add(MakeShared<FViewModelFieldPathHelper>(reinterpret_cast<FMVVMViewModelPropertyPath*>(Datum), WidgetBlueprint));
+				}
+				else
+				{
+					Result.Add(MakeShared<FWidgetFieldPathHelper>(reinterpret_cast<FMVVMWidgetPropertyPath*>(Datum), WidgetBlueprint));
+				}
+			}
+			return Result;
+		}
+
+		TArray<IFieldPathHelper*> GetRawPathHelperPointers(const TArray<TSharedPtr<IFieldPathHelper>>& SharedPtrs)
+		{
+			TArray<IFieldPathHelper*> RawPointers;
+			RawPointers.Reserve(SharedPtrs.Num());
+
+			Algo::Transform(SharedPtrs, RawPointers, [](const TSharedPtr<IFieldPathHelper>& X) { return X.Get(); });
+
+			return RawPointers;
+		}
+	}
+
+	FPropertyPathCustomizationBase::FPropertyPathCustomizationBase(UWidgetBlueprint* InWidgetBlueprint, EPropertyPathType InPathType) :
+		PathType(InPathType),
+		WidgetBlueprint(InWidgetBlueprint)
+	{
+		check(WidgetBlueprint != nullptr);
+
+		OnBlueprintChangedHandle = WidgetBlueprint->OnChanged().AddRaw(this, &FPropertyPathCustomizationBase::HandleBlueprintChanged);
+	}
+
+	FPropertyPathCustomizationBase::~FPropertyPathCustomizationBase()
+	{
+		WidgetBlueprint->OnChanged().Remove(OnBlueprintChangedHandle);
+	}
+
+	void FPropertyPathCustomizationBase::CustomizeHeader(TSharedRef<IPropertyHandle> InPropertyHandle, FDetailWidgetRow& HeaderRow, IPropertyTypeCustomizationUtils& CustomizationUtils)
+	{
+		PropertyHandle = InPropertyHandle;
+		PathHelpers = Private::GetPathHelpersForProperty(PropertyHandle, WidgetBlueprint, PathType);
+
+		TSharedPtr<IPropertyHandle> ParentHandle = PropertyHandle->GetParentHandle();
+		TSharedPtr<IPropertyHandle> OtherHandle;
+
+		if (PathType == EPropertyPathType::ViewModel)
+		{
+			OtherHandle = ParentHandle->GetChildHandle(GET_MEMBER_NAME_CHECKED(FMVVMBlueprintViewBinding, WidgetPath));
+		}
+		else
+		{
+			OtherHandle = ParentHandle->GetChildHandle(GET_MEMBER_NAME_CHECKED(FMVVMBlueprintViewBinding, ViewModelPath));
+		}
+
+		if (OtherHandle.IsValid() && OtherHandle->IsValidHandle())
+		{
+			OtherHandle->SetOnPropertyValueChanged(FSimpleDelegate::CreateSP(this, &FPropertyPathCustomizationBase::OnOtherPropertyChanged));
+			OtherHelpers = Private::GetPathHelpersForProperty(OtherHandle, WidgetBlueprint, PathType == EPropertyPathType::Widget ? EPropertyPathType::ViewModel : EPropertyPathType::Widget);
+		}
+		else
+		{
+			ensureMsgf(false, TEXT("MVVMPropertyPathCustomization used in unknown context."));
+		}
+
+		TSharedPtr<IPropertyHandle> BindingTypeHandle = ParentHandle->GetChildHandle(GET_MEMBER_NAME_CHECKED(FMVVMBlueprintViewBinding, BindingType));
+		if (BindingTypeHandle.IsValid() && BindingTypeHandle->IsValidHandle())
+		{
+			BindingTypeHandle->SetOnPropertyValueChanged(FSimpleDelegate::CreateSP(this, &FPropertyPathCustomizationBase::OnOtherPropertyChanged));
+		}
+		else
+		{
+			ensureMsgf(false, TEXT("MVVMPropertyPathCustomization used in unknown context."));
+		}
+
+		uint8 Value;
+		BindingTypeHandle->GetValue(Value);
+
+		HeaderRow
+			.NameWidget
+			.VAlign(VAlign_Center)
+			[
+				PropertyHandle->CreatePropertyNameWidget()
+			]
+			.ValueWidget
+			.VAlign(VAlign_Center)
+			.HAlign(HAlign_Fill)
+			[
+				SNew(SHorizontalBox)
+				+ SHorizontalBox::Slot()
+				.VAlign(VAlign_Center)
+				[
+					SAssignNew(SourceSelector, SMVVMSourceSelector)
+					.TextStyle(FAppStyle::Get(), "SmallText")
+					.PathHelpers(this, &FPropertyPathCustomizationBase::GetRawPathHelpers)
+					.OnSelectionChanged(this, &FPropertyPathCustomizationBase::OnSourceSelectionChanged)
+				]
+			+ SHorizontalBox::Slot()
+				.Padding(5, 0, 0, 0)
+				.VAlign(VAlign_Center)
+				[
+					SAssignNew(FieldSelector, SMVVMFieldSelector)
+					.TextStyle(FAppStyle::Get(), "SmallText")
+					.PathHelpers(GetRawPathHelpers())
+					.CounterpartHelpers(GetRawOtherHelpers())
+					.BindingMode(this, &FPropertyPathCustomizationBase::GetCurrentBindingMode)
+					.OnSelectionChanged(this, &FPropertyPathCustomizationBase::OnPropertySelectionChanged)
+				]				
+			];
+	}
+
+	TArray<IFieldPathHelper*> FPropertyPathCustomizationBase::GetRawPathHelpers() const
+	{
+		return Private::GetRawPathHelperPointers(PathHelpers);
+	}
+
+	TArray<IFieldPathHelper*> FPropertyPathCustomizationBase::GetRawOtherHelpers() const
+	{
+		return Private::GetRawPathHelperPointers(OtherHelpers);
+	}
+
+	void FPropertyPathCustomizationBase::OnOtherPropertyChanged()
+	{
+		FieldSelector->Refresh();
+	}
+
+	EMVVMBindingMode FPropertyPathCustomizationBase::GetCurrentBindingMode() const
+	{
+		TSharedPtr<IPropertyHandle> BindingTypeHandle = PropertyHandle->GetParentHandle()->GetChildHandle(GET_MEMBER_NAME_CHECKED(FMVVMBlueprintViewBinding, BindingType));
+		uint8 BindingMode;
+		if (BindingTypeHandle->GetValue(BindingMode) == FPropertyAccess::Success)
+		{
+			return (EMVVMBindingMode) BindingMode;
+		}
+
+		return EMVVMBindingMode::OneWayToDestination;
+	}
+
+	void FPropertyPathCustomizationBase::OnSourceSelectionChanged(TOptional<FBindingSource> Selected)
+	{
+		FScopedTransaction Transaction(LOCTEXT("ChangeSource", "Change binding source."));
+
+		PropertyHandle->NotifyPreChange();
+
+		for (const TSharedPtr<IFieldPathHelper>& Path : PathHelpers)
+		{
+			Path->SetSelectedSource(Selected);
+			Path->SetBindingName(FName()); // Might make sense to keep this around in case we retarget to a compatible widget or switch back.
+		}
+
+		PropertyHandle->NotifyPostChange(EPropertyChangeType::ValueSet);
+
+		FieldSelector->Refresh();
+	}
+
+	void FPropertyPathCustomizationBase::OnPropertySelectionChanged(FMVVMConstFieldVariant Selected)
+	{
+		if (bPropertySelectionChanging)
+		{
+			return;
+		}
+		TGuardValue<bool> ReentrantGuard(bPropertySelectionChanging, true);
+
+		FScopedTransaction Transaction(LOCTEXT("ChangeBindingProperty", "Change binding property."));
+
+		PropertyHandle->NotifyPreChange();
+
+		// This is currently handled in SMVVMFieldSelector, but I wonder if it should be here...
+		TArray<IFieldPathHelper*> Helpers = GetRawPathHelpers();
+		for (const IFieldPathHelper* Helper : Helpers)
+		{
+			Helper->SetBindingName(Selected.GetName());
+		}
+
+		PropertyHandle->NotifyPostChange(EPropertyChangeType::ValueSet);
+	}
+
+	void FPropertyPathCustomizationBase::HandleBlueprintChanged(UBlueprint* Blueprint)
+	{
+		SourceSelector->Refresh();
+		FieldSelector->Refresh();
+	}
+}
+
+#undef LOCTEXT_NAMESPACE
\ No newline at end of file
diff --git a/Engine/Plugins/Runtime/ModelViewViewModel/Source/ModelViewViewModelEditor/Private/Customizations/MVVMPropertyPathCustomization.h b/Engine/Plugins/Runtime/ModelViewViewModel/Source/ModelViewViewModelEditor/Private/Customizations/MVVMPropertyPathCustomization.h
new file mode 100644
--- /dev/null
+++ b/Engine/Plugins/Runtime/ModelViewViewModel/Source/ModelViewViewModelEditor/Private/Customizations/MVVMPropertyPathCustomization.h
@@ -0,0 +1,94 @@
+// Copyright Epic Games, Inc. All Rights Reserved.
+
+#pragma once
+
+#include "IPropertyTypeCustomization.h"
+#include "MVVMBlueprintView.h"
+#include "MVVMPropertyPathHelpers.h"
+#include "PropertyHandle.h"
+#include "Types/MVVMFieldVariant.h"
+
+class SMVVMSourceSelector;
+class SMVVMFieldSelector;
+
+namespace UE::MVVM
+{
+	struct FBindingSource;
+
+	enum class EPropertyPathType
+	{
+		Widget,
+		ViewModel
+	};
+
+	class FPropertyPathCustomizationBase : public IPropertyTypeCustomization
+	{
+	protected:
+		FPropertyPathCustomizationBase(UWidgetBlueprint* WidgetBlueprint, EPropertyPathType PathType);
+		~FPropertyPathCustomizationBase();
+
+	public:
+		virtual void CustomizeHeader(TSharedRef<IPropertyHandle> InPropertyHandle, FDetailWidgetRow& HeaderRow, IPropertyTypeCustomizationUtils& CustomizationUtils) override;
+		virtual void CustomizeChildren(TSharedRef<IPropertyHandle> InPropertyHandle, IDetailChildrenBuilder& ChildBuilder, IPropertyTypeCustomizationUtils& CustomizationUtils) override {}
+
+	private:
+		void OnSourceSelectionChanged(TOptional<FBindingSource> Selected);
+		void OnPropertySelectionChanged(FMVVMConstFieldVariant Selected);
+
+		void OnOtherPropertyChanged();
+
+		TArray<IFieldPathHelper*> GetRawPathHelpers() const;
+		TArray<IFieldPathHelper*> GetRawOtherHelpers() const;
+		
+		EMVVMBindingMode GetCurrentBindingMode() const;
+
+		void HandleBlueprintChanged(UBlueprint* Blueprint);
+
+	private:
+		TSharedPtr<IPropertyHandle> PropertyHandle;
+		EPropertyPathType PathType;
+		UWidgetBlueprint* WidgetBlueprint = nullptr;
+
+		TSharedPtr<SMVVMSourceSelector> SourceSelector;
+		TSharedPtr<SMVVMFieldSelector> FieldSelector;
+
+		TArray<TSharedPtr<IFieldPathHelper>> PathHelpers;
+		TArray<TSharedPtr<IFieldPathHelper>> OtherHelpers;
+
+		bool bPropertySelectionChanging = false;
+
+		FDelegateHandle OnBlueprintChangedHandle;
+	};
+
+	class FWidgetPropertyPathCustomization : public FPropertyPathCustomizationBase
+	{
+	public:
+		static TSharedRef<IPropertyTypeCustomization> MakeInstance(UWidgetBlueprint* InWidgetBlueprint)
+		{
+			check(InWidgetBlueprint != nullptr);
+			TSharedRef<FWidgetPropertyPathCustomization> Customization = MakeShared<FWidgetPropertyPathCustomization>(InWidgetBlueprint);
+			return Customization;
+		}
+	
+		FWidgetPropertyPathCustomization(UWidgetBlueprint* WidgetBlueprint) : 
+			FPropertyPathCustomizationBase(WidgetBlueprint, EPropertyPathType::Widget)
+		{
+		}
+	};
+
+	class FViewModelPropertyPathCustomization : public FPropertyPathCustomizationBase
+	{
+	public:
+		static TSharedRef<IPropertyTypeCustomization> MakeInstance(UWidgetBlueprint* InWidgetBlueprint)
+		{
+			check(InWidgetBlueprint != nullptr);
+			TSharedRef<FViewModelPropertyPathCustomization> Customization = MakeShared<FViewModelPropertyPathCustomization>(InWidgetBlueprint);
+			return Customization;
+		}
+
+		FViewModelPropertyPathCustomization(UWidgetBlueprint* WidgetBlueprint) :
+			FPropertyPathCustomizationBase(WidgetBlueprint, EPropertyPathType::ViewModel)
+		{
+		}
+	};
+}
\ No newline at end of file
diff --git a/Engine/Plugins/Runtime/ModelViewViewModel/Source/ModelViewViewModelEditor/Private/MVVMEditorSubsystem.cpp b/Engine/Plugins/Runtime/ModelViewViewModel/Source/ModelViewViewModelEditor/Private/MVVMEditorSubsystem.cpp
new file mode 100644
--- /dev/null
+++ b/Engine/Plugins/Runtime/ModelViewViewModel/Source/ModelViewViewModelEditor/Private/MVVMEditorSubsystem.cpp
@@ -0,0 +1,133 @@
+// Copyright Epic Games, Inc. All Rights Reserved.
+
+#include "MVVMEditorSubsystem.h"
+
+#include "Bindings/MVVMBindingHelper.h"
+#include "BlueprintActionDatabase.h"
+#include "BlueprintNodeSpawner.h"
+#include "Engine/Engine.h"
+#include "MVVMSubsystem.h"
+#include "MVVMWidgetBlueprintExtension_View.h"
+#include "WidgetBlueprint.h"
+
+UMVVMBlueprintView* UMVVMEditorSubsystem::RequestView(UWidgetBlueprint* WidgetBlueprint) const
+{
+	return UMVVMWidgetBlueprintExtension_View::RequestExtension<UMVVMWidgetBlueprintExtension_View>(WidgetBlueprint)->GetBlueprintView();
+}
+
+UMVVMBlueprintView* UMVVMEditorSubsystem::GetView(UWidgetBlueprint* WidgetBlueprint) const
+{
+	if (UMVVMWidgetBlueprintExtension_View* ExtensionView = UMVVMWidgetBlueprintExtension_View::GetExtension<UMVVMWidgetBlueprintExtension_View>(WidgetBlueprint))
+	{
+		return ExtensionView->GetBlueprintView();
+	}
+	return nullptr;
+}
+
+TArray<UE::MVVM::FMVVMConstFieldVariant> UMVVMEditorSubsystem::GetChildViewModels(TSubclassOf<UMVVMViewModelBase> Class)
+{
+	TArray<UE::MVVM::FMVVMConstFieldVariant> Result;
+
+	for (TFieldIterator<const FProperty> PropertyIt(Class, EFieldIteratorFlags::IncludeSuper); PropertyIt; ++PropertyIt)
+	{
+		const FProperty* Property = *PropertyIt;
+		if (const FObjectPropertyBase* ObjectProperty = CastField<FObjectPropertyBase>(Property))
+		{
+			if (ObjectProperty->PropertyClass->IsChildOf(UMVVMViewModelBase::StaticClass()))
+			{
+				if (UE::MVVM::BindingHelper::IsValidForSourceBinding(Property))
+				{
+					Result.Add(UE::MVVM::FMVVMConstFieldVariant(Property));
+				}
+			}
+		}
+	}
+
+	return Result;
+}
+
+void UMVVMEditorSubsystem::RemoveViewModel(UWidgetBlueprint* WidgetBlueprint, const FMVVMBlueprintViewModelContext& ViewModel)
+{
+	if (UMVVMBlueprintView* View = GetView(WidgetBlueprint))
+	{
+		View->RemoveViewModel(ViewModel.GetViewModelId());
+	}
+}
+
+FMVVMBlueprintViewBinding& UMVVMEditorSubsystem::AddBinding(UWidgetBlueprint* WidgetBlueprint)
+{
+	UMVVMBlueprintView* View = RequestView(WidgetBlueprint);
+	return View->AddDefaultBinding();
+}
+
+void UMVVMEditorSubsystem::RemoveBinding(UWidgetBlueprint* WidgetBlueprint, const FMVVMBlueprintViewBinding& Binding)
+{
+	if (UMVVMBlueprintView* View = GetView(WidgetBlueprint))
+	{
+		View->RemoveBinding(&Binding);
+	}
+}
+
+TArray<const UFunction*> UMVVMEditorSubsystem::GetAvailableConversionFunctions(const UE::MVVM::FMVVMConstFieldVariant& Source, const UE::MVVM::FMVVMConstFieldVariant& Destination) const
+{
+	TArray<const UFunction*> ConversionFunctions;
+	TArray<const UFunction*> DEBUG_StaticTwoArgFunctions;
+
+	if (Source.IsEmpty() || Destination.IsEmpty())
+	{
+		return ConversionFunctions;
+	}
+
+	const UMVVMSubsystem* Subsystem = GEngine->GetEngineSubsystem<UMVVMSubsystem>();
+
+	FBlueprintActionDatabase& ActionDatabase = FBlueprintActionDatabase::Get();
+	const FBlueprintActionDatabase::FActionRegistry& ActionRegistry = ActionDatabase.GetAllActions();
+
+	for (auto It = ActionRegistry.CreateConstIterator(); It; ++It)
+	{
+		if (UObject* ActionObject = It->Key.ResolveObjectPtr())
+		{
+			for (const UBlueprintNodeSpawner* NodeSpawner : It->Value)
+			{
+				FBlueprintActionInfo BlueprintAction(ActionObject, NodeSpawner);
+				const UFunction* Function = BlueprintAction.GetAssociatedFunction();
+				if (Function != nullptr)
+				{
+					TValueOrError<UE::MVVM::BindingHelper::FConversionFunctionArguments, FString> FunctionProperties = UE::MVVM::BindingHelper::TryGetPropertyTypeForConversionFunction(Function);
+					if (FunctionProperties.HasError())
+					{
+						continue;
+					}
+
+					const FProperty* ArgumentProperty = FunctionProperties.GetValue().ArgumentProperty;
+					const FProperty* ReturnProperty = FunctionProperties.GetValue().ReturnProperty;
+
+					DEBUG_StaticTwoArgFunctions.Add(Function);
+
+					const FProperty* SourceProperty = Source.IsProperty() ? Source.GetProperty() : UE::MVVM::BindingHelper::GetReturnProperty(Source.GetFunction());
+					const FProperty* DestinationProperty = Destination.IsProperty() ? Destination.GetProperty() : UE::MVVM::BindingHelper::GetFirstArgumentProperty(Destination.GetFunction());
+
+					// check that the source -> argument is valid
+					if (!UE::MVVM::BindingHelper::ArePropertiesCompatible(SourceProperty, ArgumentProperty))
+					{
+						continue;
+					}
+					
+					// check that the return -> dest is valid
+					if (!UE::MVVM::BindingHelper::ArePropertiesCompatible(ReturnProperty, DestinationProperty))
+					{
+						continue;
+					}
+
+					ConversionFunctions.Add(Function);
+				}
+			}
+		}
+	}
+
+	ConversionFunctions.Sort([](const UFunction& A, const UFunction& B) -> bool
+		{
+			return A.GetFName().LexicalLess(B.GetFName());
+		});
+	return ConversionFunctions;
+}
diff --git a/Engine/Plugins/Runtime/ModelViewViewModel/Source/ModelViewViewModelEditor/Private/MVVMPropertyAccess.cpp b/Engine/Plugins/Runtime/ModelViewViewModel/Source/ModelViewViewModelEditor/Private/MVVMPropertyAccess.cpp
new file mode 100644
--- /dev/null
+++ b/Engine/Plugins/Runtime/ModelViewViewModel/Source/ModelViewViewModelEditor/Private/MVVMPropertyAccess.cpp
@@ -0,0 +1,534 @@
+// Copyright Epic Games, Inc. All Rights Reserved.
+
+#include "MVVMPropertyAccess.h"
+//
+//#include "WidgetBlueprint.h"
+//#include "WidgetBlueprintExtension.h"
+//#include "MVVMViewModel.h"
+//
+//#include "Features/IModularFeatures.h"
+//#include "Framework/MultiBox/MultiBoxBuilder.h"
+//#include "Kismet2/BlueprintEditorUtils.h"
+//#include "IPropertyAccessEditor.h"
+//#include "PropertyHandle.h"
+//
+//#include "Widgets/SBoxPanel.h"
+//#include "Widgets/Images/SImage.h"
+//#include "Widgets/Layout/SSpacer.h"
+//#include "Widgets/Text\STextBlock.h"
+//
+//#define LOCTEXT_NAMESPACE "MVVMPropertyAccess"
+//
+//FMVVMPropertyAccess::FMVVMPropertyAccess(UWidgetBlueprint* InWidgetBlueprint, const TSharedPtr<IPropertyHandle>& InPropertyHandle)
+//	: WidgetBlueprint(InWidgetBlueprint)
+//	, PropertyHandle(InPropertyHandle)
+//{
+//
+//}
+//
+//
+//TSharedPtr<SWidget> FMVVMPropertyAccess::MakeViewModelBindingMenu() const
+//{
+//	const bool bInShouldCloseWindowAfterMenuSelection = false;
+//	FMenuBuilder MenuBuilder(bInShouldCloseWindowAfterMenuSelection, nullptr);
+//
+//	{
+//		MenuBuilder.BeginSection("BindingActions", LOCTEXT("Bindings", "Bindings"));
+//		MenuBuilder.AddMenuEntry(
+//			LOCTEXT("RemoveBinding", "Remove Binding"),
+//			LOCTEXT("RemoveBindingToolTip", "Removes the current binding"),
+//			FSlateIcon(FEditorStyle::GetStyleSetName(), "Cross"),
+//			FUIAction()
+//		);
+//		MenuBuilder.AddMenuEntry(
+//			LOCTEXT("CreateBinding", "Create Binding"),
+//			LOCTEXT("CreateBindingToolTip", "Creates a new function on the widget blueprint that will return the binding data for this property."),
+//			FSlateIcon(FEditorStyle::GetStyleSetName(), "Plus"),
+//			FUIAction()
+//		);
+//		MenuBuilder.EndSection();
+//	}
+//
+//	{
+//		UClass* PlayerPawnViewModel = FindObject<UClass>(nullptr, TEXT("/Game/PlayerPawnViewModel.PlayerPawnViewModel_C"));
+//		TSubclassOf<UMVVMViewModelBase> ViewModelDefinition = PlayerPawnViewModel;
+//
+//		MenuBuilder.BeginSection("ViewModels", LOCTEXT("ViewModels", "ViewModels"));
+//		// for all ViewModel
+//		MenuBuilder.AddSubMenu(
+//			MakeViewModelSubMenuWidget(ViewModelDefinition),
+//			FNewMenuDelegate::CreateRaw(this, &FMVVMPropertyAccess::FillViewModelSubMenu, ViewModelDefinition));
+//		MenuBuilder.EndSection();
+//	}
+//
+//	return MenuBuilder.MakeWidget();
+//}
+//
+//TSharedRef<SWidget> FMVVMPropertyAccess::MakeViewModelSubMenuWidget(TSubclassOf<UMVVMViewModelBase> ViewModelDefinition) const
+//{
+//	UMVVMViewModelBase* ViewModel_CDO = ViewModelDefinition.GetDefaultObject();
+//	if (ViewModel_CDO == nullptr)
+//	{
+//		return SNullWidget::NullWidget;
+//	}
+//
+//	return SNew(SHorizontalBox)
+//		.ToolTipText(ViewModelDefinition->GetToolTipText())
+//		+ SHorizontalBox::Slot()
+//		.AutoWidth()
+//		[
+//			SNew(SSpacer)
+//			.Size(FVector2D(18.0f, 0.0f))
+//		]
+//		//+ SHorizontalBox::Slot()
+//		//.AutoWidth()
+//		//.VAlign(VAlign_Center)
+//		//.Padding(1.0f, 0.0f)
+//		//[
+//		//	SNew(SImage)
+//		//	.Image(ContextStruct.Icon)
+//		//]
+//		+ SHorizontalBox::Slot()
+//		.AutoWidth()
+//		.VAlign(VAlign_Center)
+//		.Padding(4.0f, 0.0f)
+//		[
+//			SNew(STextBlock)
+//			.Text(ViewModelDefinition->GetDisplayNameText())
+//		];
+//}
+//
+//
+//TSharedRef<SWidget> FMVVMPropertyAccess::MakeFunctionWidget(UFunction* Function)
+//{
+//	static const FName FunctionIcon(TEXT("GraphEditor.Function_16x"));
+//
+//	const UEdGraphSchema_K2* Schema = GetDefault<UEdGraphSchema_K2>();
+//
+//	FEdGraphPinType PinType;
+//	if (FProperty* ReturnProperty = Function->GetReturnProperty())
+//	{
+//		Schema->ConvertPropertyToPinType(ReturnProperty, PinType);
+//	}
+//
+//	return SNew(SHorizontalBox)
+//		.ToolTipText(Function->GetToolTipText())
+//		+ SHorizontalBox::Slot()
+//		.AutoWidth()
+//		[
+//			SNew(SSpacer)
+//			.Size(FVector2D(18.0f, 0.0f))
+//		]
+//		+ SHorizontalBox::Slot()
+//		.AutoWidth()
+//		.VAlign(VAlign_Center)
+//		.Padding(1.0f, 0.0f)
+//		[
+//			SNew(SImage)
+//			.Image(FEditorStyle::Get().GetBrush(FunctionIcon))
+//			.ColorAndOpacity(Schema->GetPinTypeColor(PinType))
+//		]
+//		+ SHorizontalBox::Slot()
+//		.AutoWidth()
+//		.VAlign(VAlign_Center)
+//		.Padding(4.0f, 0.0f)
+//		[
+//			SNew(STextBlock)
+//			.Text(Function->GetDisplayNameText())
+//		];
+//}
+//
+//
+//TSharedRef<SWidget> FMVVMPropertyAccess::MakePropertyWidget(FProperty* Property)
+//{
+//	static FName PropertyIcon(TEXT("Kismet.VariableList.TypeIcon"));
+//
+//	const UEdGraphSchema_K2* Schema = GetDefault<UEdGraphSchema_K2>();
+//
+//	FEdGraphPinType PinType;
+//	Schema->ConvertPropertyToPinType(Property, PinType);
+//
+//	return SNew(SHorizontalBox)
+//		.ToolTipText(Property->GetToolTipText())
+//		+ SHorizontalBox::Slot()
+//		.AutoWidth()
+//		[
+//			SNew(SSpacer)
+//			.Size(FVector2D(18.0f, 0.0f))
+//		]
+//		+ SHorizontalBox::Slot()
+//		.AutoWidth()
+//		.VAlign(VAlign_Center)
+//		.Padding(1.0f, 0.0f)
+//		[
+//			SNew(SImage)
+//			.Image(FBlueprintEditorUtils::GetIconFromPin(PinType, true))
+//			.ColorAndOpacity(Schema->GetPinTypeColor(PinType))
+//		]
+//		+ SHorizontalBox::Slot()
+//		.AutoWidth()
+//		.VAlign(VAlign_Center)
+//		.Padding(4.0f, 0.0f)
+//		[
+//			SNew(STextBlock)
+//			.Text(FText::FromName(Property->GetFName()))
+//		];
+//}
+//
+//void FMVVMPropertyAccess::FillViewModelSubMenu(FMenuBuilder& MenuBuilder, TSubclassOf<UMVVMViewModelBase> ViewModelDefinition) const
+//{
+//	//for (const FMVVMViewModelOutput& Output : ViewModelDefinition.GetDefaultObject()->Outputs)
+//	//{
+//	//	UFunction* Function = nullptr;
+//	//	FProperty* Property = nullptr;
+//	//	if (!Output.ConversionFunction.IsNone())
+//	//	{
+//	//		Function = ViewModelDefinition->FindFunctionByName(Output.ConversionFunction);
+//	//	}
+//	//	else
+//	//	{
+//	//		if (Output.InputIds.Num() == 1)
+//	//		{
+//	//			if (const FMVVMViewModelInput* Input = ViewModelDefinition.GetDefaultObject()->FindInput(Output.InputIds[0]))
+//	//			{
+//	//				FFieldVariant FieldVariant = ViewModelDefinition.GetDefaultObject()->GetInputField(*Input);
+//	//				Function = Cast<UFunction>(FieldVariant.ToUObject());
+//	//				Property = CastField<FProperty>(FieldVariant.ToField());
+//	//			}
+//	//		}
+//	//	}
+//
+//	//	if (Function)
+//	//	{
+//	//		MenuBuilder.AddMenuEntry(
+//	//			FUIAction(FExecuteAction::CreateRaw(this, &FMVVMPropertyAccess::HandleAddBinding, Output.OutputId)),
+//	//			MakeFunctionWidget(Function));
+//	//	}
+//	//	else if (Property)
+//	//	{
+//	//		MenuBuilder.AddMenuEntry(
+//	//			FUIAction(FExecuteAction::CreateRaw(this, &FMVVMPropertyAccess::HandleAddBinding, Output.OutputId)),
+//	//			MakePropertyWidget(Property));
+//	//	}
+//	//}
+//}
+//
+//void FMVVMPropertyAccess::HandleAddBinding(FGuid OutputId) const
+//{
+//
+//}
+//
+//
+//bool FMVVMPropertyAccess::IsBindableFunction(UFunction* InFunction)
+//{
+//	int32 ParamCount = 0;
+//	FProperty* ReturnProperty = nullptr;
+//	for (TFieldIterator<FProperty> It(InFunction); It && (It->PropertyFlags & CPF_Parm); ++It)
+//	{
+//		if (It->PropertyFlags & CPF_ReturnParm)
+//		{
+//			ReturnProperty = *It;
+//		}
+//		if (It->PropertyFlags & CPF_Parm)
+//		{
+//			++ParamCount;
+//		}
+//	}
+//
+//	return ReturnProperty != nullptr
+//		&& ReturnProperty->IsA(FFloatProperty::StaticClass())
+//		&& ParamCount == 0
+//		&& InFunction->HasAllFunctionFlags(FUNC_Public | FUNC_BlueprintPure)
+//		&& !InFunction->HasAnyFunctionFlags(FUNC_EditorOnly);
+//}
+//
+//
+//bool FMVVMPropertyAccess::IsBindableProperty(FProperty* InProperty)
+//{
+//	return InProperty->IsA(FFloatProperty::StaticClass())
+//		&& InProperty->HasAllPropertyFlags(CPF_BlueprintVisible)
+//		&& !InProperty->HasAnyPropertyFlags(CPF_EditorOnly | CPF_Deprecated);
+//}
+//
+////TSharedPtr<SWidget> FMVVMPropertyAccess::MakeViewModelBindingMenu(UWidgetBlueprint* WidgetBlueprint, const TSharedPtr<IPropertyHandle>& PropertyHandle)
+////{
+////	static const FName NAME_PropertyAccessEditor = "PropertyAccessEditor";
+////	if (!IModularFeatures::Get().IsModularFeatureAvailable("PropertyAccessEditor"))
+////	{
+////		return TSharedPtr<SWidget>();
+////	}
+////
+////	//if (UMVVMViewModelBlueprintExtension* Extension = UWidgetBlueprintExtension::RequestExtension<UMVVMViewModelBlueprintExtension>(WidgetBlueprint))
+////	//{
+////
+////	FProperty* Property = PropertyHandle->GetProperty();
+////	const FName PropertyName = Property->GetFName();
+////
+////	TArray<TWeakObjectPtr<UObject>> Objects;
+////	{
+////		TArray<UObject*> RawObjects;
+////		PropertyHandle->GetOuterObjects(RawObjects);
+////
+////		Objects.Reserve(RawObjects.Num());
+////		for (UObject* RawObject : RawObjects)
+////		{
+////			Objects.Add(RawObject);
+////		}
+////	}
+////
+////	UWidget* Widget = Objects.Num() ? Cast<UWidget>(Objects[0]) : nullptr;
+////
+////	FPropertyBindingMenuWidgetArgs Args;
+////	Args.Property = Property;
+////	Args.bAllowNewBindings = false;
+////	Args.bGeneratePureBindings = true;
+////	//Args.BindableSignature = InSetterSignature;
+////	//Args.OnGenerateBindingName = FOnGenerateBindingName::CreateLambda([WidgetName]()
+////	//	{
+////	//		return WidgetName;
+////	//	});
+////
+////	//Args.OnGotoBinding = FOnGotoBinding::CreateLambda([InEditor, Objects](FName InPropertyName)
+////	//	{
+////	//		UWidgetBlueprint* ThisBlueprint = InEditor.Pin()->GetWidgetBlueprintObj();
+////
+////	//		//TODO UMG O(N) Isn't good for this, needs to be map, but map isn't serialized, need cached runtime map for fast lookups.
+////
+////	//		for (const TWeakObjectPtr<UObject>& ObjectPtr : Objects)
+////	//		{
+////	//			UObject* Object = ObjectPtr.Get();
+////
+////	//			// Ignore null outer objects
+////	//			if (Object == nullptr)
+////	//			{
+////	//				continue;
+////	//			}
+////
+////	//			for (const FDelegateEditorBinding& Binding : ThisBlueprint->Bindings)
+////	//			{
+////	//				if (Binding.ObjectName == Object->GetName() && Binding.PropertyName == InPropertyName)
+////	//				{
+////	//					if (Binding.Kind == EBindingKind::Function)
+////	//					{
+////	//						TArray<UEdGraph*> AllGraphs;
+////	//						ThisBlueprint->GetAllGraphs(AllGraphs);
+////
+////	//						FGuid SearchForGuid = Binding.MemberGuid;
+////	//						if (!Binding.SourcePath.IsEmpty())
+////	//						{
+////	//							SearchForGuid = Binding.SourcePath.Segments.Last().GetMemberGuid();
+////	//						}
+////
+////	//						for (UEdGraph* Graph : AllGraphs)
+////	//						{
+////	//							if (Graph->GraphGuid == SearchForGuid)
+////	//							{
+////	//								InEditor.Pin()->SetCurrentMode(FWidgetBlueprintApplicationModes::GraphMode);
+////	//								InEditor.Pin()->OpenDocument(Graph, FDocumentTracker::OpenNewDocument);
+////	//							}
+////	//						}
+////
+////	//						// Either way return
+////	//						return true;
+////	//					}
+////	//				}
+////	//			}
+////	//		}
+////
+////	//		return false;
+////	//	});
+////
+////		//Args.OnCanGotoBinding = FOnCanGotoBinding::CreateLambda([InEditor, Objects](FName InPropertyName)
+////		//	{
+////		//		//UWidgetBlueprint* ThisBlueprint = InEditor.Pin()->GetWidgetBlueprintObj();
+////
+////		//		//for (const TWeakObjectPtr<UObject>& ObjectPtr : Objects)
+////		//		//{
+////		//		//	UObject* Object = ObjectPtr.Get();
+////
+////		//		//	// Ignore null outer objects
+////		//		//	if (Object == nullptr)
+////		//		//	{
+////		//		//		continue;
+////		//		//	}
+////
+////		//		//	for (const FDelegateEditorBinding& Binding : ThisBlueprint->Bindings)
+////		//		//	{
+////		//		//		if (Binding.ObjectName == Object->GetName() && Binding.PropertyName == InPropertyName)
+////		//		//		{
+////		//		//			if (Binding.Kind == EBindingKind::Function)
+////		//		//			{
+////		//		//				return true;
+////		//		//			}
+////		//		//		}
+////		//		//	}
+////		//		//}
+////
+////		//		return false;
+////		//	});
+////
+////		Args.OnCanBindProperty = FOnCanBindProperty::CreateLambda([](FProperty* InProperty)
+////			{
+////				//static const FName NAME_OnChangedEvent = "OnChangedEvent";
+////				//if (InProperty->HasMetaData(NAME_OnChangedEvent))
+////				//{
+////				//	const FName EventName = FName(InProperty->GetMetaData(NAME_OnChangedEvent));
+////				//	if (FMulticastDelegateProperty* Event = FindFProperty<FMulticastDelegateProperty>(InProperty->GetOwnerStruct(), EventName))
+////				//	{
+////				//		check(Event->SignatureFunction);
+////				//		const bool bHasValidReturnValue = Event->SignatureFunction->GetReturnProperty() == nullptr;
+////				//		bool bHasValidParams = true;
+////				//		for (TFieldIterator<FProperty> ParamIt(Event->SignatureFunction); ParamIt; ++ParamIt)
+////				//		{
+////				//			if ((*ParamIt)->HasAnyPropertyFlags(CPF_Parm))
+////				//			{
+////				//				bHasValidParams = false;
+////				//			}
+////				//		}
+////
+////				//		return bHasValidReturnValue && bHasValidParams;
+////				//	}
+////				//}
+////
+////				//return false;
+////				return InProperty->IsA(FFloatProperty::StaticClass())
+////					&& InProperty->HasAllPropertyFlags(CPF_BlueprintVisible)
+////					&& !InProperty->HasAnyPropertyFlags(CPF_EditorOnly| CPF_Deprecated);
+////			});
+////
+////		Args.OnCanBindFunction = FOnCanBindFunction::CreateLambda([](UFunction* InFunction)
+////			{
+////				int32 ParamCount = 0;
+////				FProperty* ReturnProperty = nullptr;
+////				for (TFieldIterator<FProperty> It(InFunction); It && (It->PropertyFlags & CPF_Parm); ++It)
+////				{
+////					if (It->PropertyFlags & CPF_ReturnParm)
+////					{
+////						ReturnProperty = *It;
+////					}
+////					if (It->PropertyFlags & CPF_Parm)
+////					{
+////						++ParamCount;
+////					}
+////				}
+////
+////				return ReturnProperty != nullptr
+////					&& ReturnProperty->IsA(FFloatProperty::StaticClass())
+////					&& ParamCount == 0
+////					&& InFunction->HasAllFunctionFlags(FUNC_Public| FUNC_BlueprintPure)
+////					&& !InFunction->HasAnyFunctionFlags(FUNC_EditorOnly);
+////			});
+////
+////		//Args.OnCanBindFunction = FOnCanBindFunction::CreateLambda([InDelegateProperty](UFunction* InFunction)
+////		//	{
+////		//		auto HasFunctionBinder = [InFunction](UFunction* InBindableSignature)
+////		//		{
+////		//			if (InFunction->NumParms == 1 && InBindableSignature->NumParms == 1)
+////		//			{
+////		//				if (FProperty* FunctionReturn = InFunction->GetReturnProperty())
+////		//				{
+////		//					if (FProperty* DelegateReturn = InBindableSignature->GetReturnProperty())
+////		//					{
+////		//						// Find the binder that can handle the delegate return type.
+////		//						TSubclassOf<UPropertyBinding> Binder = UWidget::FindBinderClassForDestination(DelegateReturn);
+////		//						if (Binder != nullptr)
+////		//						{
+////		//							// Ensure that the binder also can handle binding from the property we care about.
+////		//							if (Binder->GetDefaultObject<UPropertyBinding>()->IsSupportedSource(FunctionReturn))
+////		//							{
+////		//								return true;
+////		//							}
+////		//						}
+////		//					}
+////		//				}
+////		//			}
+////
+////		//			return false;
+////		//		};
+////
+////		//		// We ignore CPF_ReturnParm because all that matters for binding to script functions is that the number of out parameters match.
+////		//		return (InFunction->IsSignatureCompatibleWith(InDelegateProperty->SignatureFunction, UFunction::GetDefaultIgnoredSignatureCompatibilityFlags() | CPF_ReturnParm) ||
+////		//			HasFunctionBinder(InDelegateProperty->SignatureFunction));
+////		//	});
+////
+////		Args.OnCanBindToClass = FOnCanBindToClass::CreateLambda([](UClass* InClass)
+////			{
+////				return true;
+////			});
+////
+////		Args.OnAddBinding = FOnAddBinding::CreateLambda([WidgetBlueprint, Objects](FName InPropertyName, const TArray<FBindingChainElement>& InBindingChain)
+////			{
+////				WidgetBlueprint->Modify();
+////
+////				UMVVMViewModelBlueprintExtension* Extension = UWidgetBlueprintExtension::GetExtension<UMVVMViewModelBlueprintExtension>(WidgetBlueprint);
+////				if (Extension == nullptr)
+////				{
+////					Extension = UWidgetBlueprintExtension::RequestExtension<UMVVMViewModelBlueprintExtension>(WidgetBlueprint);
+////				}
+////				UMVVMViewModelEditorData* ViewModel = Extension->GetViewModelEditorData();
+////				if (ViewModel == nullptr)
+////				{
+////					Extension->CreateViewModelInstance();
+////					ViewModel = Extension->GetViewModelEditorData();
+////				}
+////
+////
+////				FGuid ViewModelBindingId;
+////				UClass* PlayerPawnViewModel = FindObject<UClass>(nullptr, TEXT("/Game/PlayerPawnViewModel.PlayerPawnViewModel_C"));
+////				if (UMVVMViewModelBase* Base = PlayerPawnViewModel->GetDefaultObject<UMVVMViewModelBase>())
+////				{
+////					const FMVVMViewModelBinding* ViewModelBinding = Base->FindBindingByViewModelPropertyName(InBindingChain.Last().Field.GetFName());
+////					if (ViewModelBinding)
+////					{
+////						ViewModelBindingId = ViewModelBinding->BindingId;
+////					}
+////				}
+////
+////				// Find the ViewModel binding or add it in the ViewModel
+////				//FGuid ViewModelBindingId;
+////				//{
+////				//	FMVVMViewModelBindingEditorData ViewModelEditorData;
+////				//	FMVVMModelContext Context;
+////				//	Context.ContextClass = FindObject<UClass>(nullptr, TEXT("/Script/MVVMProto.ProtoCharacter"));
+////				//	Context.ContextName = TEXT("PlayerCharacter");
+////				//	ViewModel->AddUniqueViewModelBinding(Context, InBindingChain.Last().Field.GetFName().ToString());
+////				//}
+////
+////				// Add or modify the current Model binding
+////				{
+////					for (TWeakObjectPtr<UObject> WeakObject : Objects)
+////					{
+////						if (UObject* Object = WeakObject.Get())
+////						{
+////							FMVVMViewBindingEditorData ViewData;
+////							ViewData.DestinationObject = Object;
+////							ViewData.DestinationPropertyPath = InPropertyName.ToString();
+////							ViewData.ViewModelBindingId = ViewModelBindingId;
+////							if (ViewData.Resolve())
+////							{
+////								Extension->AddUniqueViewBinding(MoveTemp(ViewData));
+////							}
+////						}
+////					}
+////				}
+////			});
+////
+////	IPropertyAccessEditor& PropertyAccessEditor = IModularFeatures::Get().GetModularFeature<IPropertyAccessEditor>("PropertyAccessEditor");
+////	TArray<FBindingContextStruct> BindingContextStructs;
+////
+////	//if (UMVVMViewModelBlueprintExtension* Extension = UWidgetBlueprintExtension::RequestExtension<UMVVMViewModelBlueprintExtension>(WidgetBlueprint))
+////	//{
+////		UClass* PlayerPawnViewModel = FindObject<UClass>(nullptr, TEXT("/Game/PlayerPawnViewModel.PlayerPawnViewModel_C"));
+////		FText DisplayName = PlayerPawnViewModel->GetDisplayNameText();
+////		FText Tooltip = PlayerPawnViewModel->GetToolTipText();
+////		UObject* CDO = PlayerPawnViewModel->GetDefaultObject();
+////		BindingContextStructs.Emplace(PlayerPawnViewModel, nullptr, DisplayName, Tooltip);
+////	//}
+////
+////	//UStruct* FoundStruct = FindObject<UStruct>(nullptr, TEXT("/Script/MVVMProto.ProtoCharacter"));
+////	//BindingContextStructs.Emplace(FoundStruct, nullptr, LOCTEXT("cONTEXT", "fOUNDdOUNCONT"), LOCTEXT("cONTSSEXT", "fOUNDdOUSSNCONT"));
+////
+////	return PropertyAccessEditor.MakePropertyBindingMenuWidget(BindingContextStructs, Args);
+////}
+//
+//#undef LOCTEXT_NAMESPACE
\ No newline at end of file
diff --git a/Engine/Plugins/Runtime/ModelViewViewModel/Source/ModelViewViewModelEditor/Private/MVVMPropertyAccess.h b/Engine/Plugins/Runtime/ModelViewViewModel/Source/ModelViewViewModelEditor/Private/MVVMPropertyAccess.h
new file mode 100644
--- /dev/null
+++ b/Engine/Plugins/Runtime/ModelViewViewModel/Source/ModelViewViewModelEditor/Private/MVVMPropertyAccess.h
@@ -0,0 +1,35 @@
+// Copyright Epic Games, Inc. All Rights Reserved.
+
+#pragma once
+
+#include "CoreMinimal.h"
+
+class IPropertyHandle;
+class SWidget;
+class UMVVMViewModelBase;
+class UWidgetBlueprint;
+
+/**
+ *
+ */
+//class FMVVMPropertyAccess
+//{
+//public:
+//	FMVVMPropertyAccess(UWidgetBlueprint* WidgetBlueprint, const TSharedPtr<IPropertyHandle>& PropertyHandle);
+//
+//public:
+//	TSharedPtr<SWidget> MakeViewModelBindingMenu() const;
+//	static TSharedRef<SWidget> MakeFunctionWidget(UFunction* Function);
+//	static TSharedRef<SWidget> MakePropertyWidget(FProperty* Property);
+//	static bool IsBindableFunction(UFunction* Function);
+//	static bool IsBindableProperty(FProperty* Property);
+//
+//private:
+//	TSharedRef<SWidget> MakeViewModelSubMenuWidget(TSubclassOf<UMVVMViewModelBase> ViewModelDefinition) const;
+//	void FillViewModelSubMenu(FMenuBuilder& MenuBuilder, TSubclassOf<UMVVMViewModelBase> ViewModelDefinition) const;
+//	void HandleAddBinding(FGuid OutputId) const;
+//
+//private:
+//	UWidgetBlueprint* WidgetBlueprint;
+//	const TSharedPtr<IPropertyHandle>& PropertyHandle;
+//};
diff --git a/Engine/Plugins/Runtime/ModelViewViewModel/Source/ModelViewViewModelEditor/Private/MVVMPropertyPathHelpers.cpp b/Engine/Plugins/Runtime/ModelViewViewModel/Source/ModelViewViewModelEditor/Private/MVVMPropertyPathHelpers.cpp
new file mode 100644
--- /dev/null
+++ b/Engine/Plugins/Runtime/ModelViewViewModel/Source/ModelViewViewModelEditor/Private/MVVMPropertyPathHelpers.cpp
@@ -0,0 +1,294 @@
+// Copyright Epic Games, Inc. All Rights Reserved.
+
+#include "MVVMPropertyPathHelpers.h"
+
+#include "Algo/Transform.h"
+#include "Bindings/MVVMBindingHelper.h"
+#include "Blueprint/WidgetTree.h"
+#include "Components/Widget.h"
+#include "Engine/Engine.h"
+#include "MVVMBlueprintView.h"
+#include "MVVMSubsystem.h"
+#include "MVVMWidgetBlueprintExtension_View.h"
+
+namespace UE::MVVM
+{
+	void FWidgetFieldPathHelper::GetAvailableSources(TSet<FBindingSource>& OutSources) const
+	{
+		FMVVMWidgetPropertyPath* Path = PathAttr.Get(nullptr);
+		if (Path == nullptr)
+		{
+			return;
+		}
+
+		const UWidgetTree* WidgetTree = WidgetBlueprint->WidgetTree;
+		if (WidgetTree == nullptr)
+		{
+			return;
+		}
+
+		TArray<UWidget*> AllWidgets;
+		WidgetTree->GetAllWidgets(AllWidgets);
+
+		OutSources.Reserve(AllWidgets.Num() + 1);
+
+		// Add current widget as a possible binding source
+		if (UClass* BPClass = WidgetBlueprint->GeneratedClass)
+		{
+			TArray<FMVVMAvailableBinding> Bindings = GEngine->GetEngineSubsystem<UMVVMSubsystem>()->GetWidgetAvailableBindings(BPClass);
+
+			if (Bindings.Num() > 0)
+			{
+				// at least one valid property, add it to our list
+				FBindingSource Source;
+				Source.Name = WidgetBlueprint->GetFName();
+				Source.DisplayName = FText::FromName(WidgetBlueprint->GetFName());
+				Source.Class = BPClass;
+				Source.IsSelected = Source.Name == Path->WidgetName;
+				OutSources.Add(Source);
+			}
+		}
+
+		for (const UWidget* Widget : AllWidgets)
+		{
+			TArray<FMVVMAvailableBinding> Bindings = GEngine->GetEngineSubsystem<UMVVMSubsystem>()->GetWidgetAvailableBindings(Widget->GetClass());
+
+			if (Bindings.Num() > 0)
+			{
+				// at least one valid property, add it to our list
+				FBindingSource Source;
+				Source.Name = Widget->GetFName();
+				Source.DisplayName = Widget->GetLabelText();
+				Source.Class = Widget->GetClass();
+				Source.IsSelected = Source.Name == Path->WidgetName;
+				OutSources.Add(Source);
+			}
+		}
+	}
+
+	TOptional<FBindingSource> FWidgetFieldPathHelper::GetSelectedSource() const
+	{
+		TSet<FBindingSource> Sources;
+		GetAvailableSources(Sources);
+
+		for (const FBindingSource& Source : Sources)
+		{
+			if (Source.IsSelected)
+			{
+				return Source;
+			}
+		}
+
+		return TOptional<FBindingSource>();
+	}
+
+	void FWidgetFieldPathHelper::GetAvailableFields(TSet<FMVVMConstFieldVariant>& OutFields) const
+	{
+		FMVVMWidgetPropertyPath* Path = PathAttr.Get(nullptr);
+		if (Path == nullptr)
+		{
+			return;
+		}
+
+		const UWidgetTree* WidgetTree = WidgetBlueprint->WidgetTree;
+		if (WidgetTree == nullptr)
+		{
+			return;
+		}
+
+		TSubclassOf<UWidget> WidgetClass = nullptr;
+		if (const UWidget* Widget = WidgetTree->FindWidget(Path->WidgetName))
+		{
+			WidgetClass = Widget->GetClass();
+		}
+		else if (WidgetBlueprint->GetFName() == Path->WidgetName)
+		{
+			WidgetClass = WidgetBlueprint->GeneratedClass;
+		}
+
+		if (!WidgetClass)
+		{
+			return;
+		}
+
+		TArray<FMVVMAvailableBinding> Bindings = GEngine->GetEngineSubsystem<UMVVMSubsystem>()->GetWidgetAvailableBindings(WidgetClass);
+		Algo::Transform(Bindings, OutFields, [WidgetClass](const FMVVMAvailableBinding& Binding)
+			{
+				return UE::MVVM::BindingHelper::FindFieldByName(WidgetClass, Binding.GetBindingName());
+			});
+	}
+
+	FMVVMConstFieldVariant FWidgetFieldPathHelper::GetSelectedField() const
+	{
+		FMVVMBindingName BindingName = GetBindingName();
+
+		TSet<FMVVMConstFieldVariant> Fields;
+		GetAvailableFields(Fields);
+
+		for (const FMVVMConstFieldVariant& Field : Fields)
+		{
+			if (Field.GetName() == BindingName.ToName())
+			{
+				return Field;
+			}
+		}
+
+		return FMVVMConstFieldVariant();
+	}
+
+	void FWidgetFieldPathHelper::SetSelectedSource(const TOptional<FBindingSource>& Source) const
+	{
+		FMVVMWidgetPropertyPath* Path = PathAttr.Get(nullptr);
+		if (Path == nullptr)
+		{
+			return;
+		}
+
+		if (Source.IsSet())
+		{
+			Path->WidgetName = Source.GetValue().Name;
+		}
+		else
+		{
+			Path->WidgetName = FName();
+		}
+	}
+
+	FMVVMBindingName FWidgetFieldPathHelper::GetBindingName() const
+	{
+		if (FMVVMWidgetPropertyPath* Path = PathAttr.Get(nullptr))
+		{
+			return Path->GetBindingName();
+		}
+
+		return FMVVMBindingName();
+	}
+
+	void FWidgetFieldPathHelper::SetBindingName(FName BindingName) const
+	{
+		if (FMVVMWidgetPropertyPath* Path = PathAttr.Get(nullptr))
+		{
+			Path->SetBindingName(BindingName);
+		}
+	}
+
+	void FViewModelFieldPathHelper::GetAvailableSources(TSet<FBindingSource>& OutSources) const
+	{
+		FMVVMViewModelPropertyPath* Path = PathAttr.Get(nullptr);
+		if (Path == nullptr)
+		{
+			return;
+		}
+
+		const UMVVMWidgetBlueprintExtension_View* MVVMExtensionPtr = UMVVMWidgetBlueprintExtension_View::GetExtension<UMVVMWidgetBlueprintExtension_View>(WidgetBlueprint);
+		const UMVVMBlueprintView* MVVMBlueprintView = MVVMExtensionPtr->GetBlueprintView();
+		const TArrayView<const FMVVMBlueprintViewModelContext>& ViewModels = MVVMBlueprintView->GetViewModels();
+
+		OutSources.Reserve(ViewModels.Num());
+
+		for (const FMVVMBlueprintViewModelContext& ViewModel : ViewModels)
+		{
+			FBindingSource Source;
+			Source.SourceGuid = ViewModel.GetViewModelId();
+			Source.DisplayName = ViewModel.GetDisplayName();
+			Source.Class = ViewModel.GetViewModelClass().Get();
+			Source.IsSelected = (Source.SourceGuid == Path->ContextId);
+
+			OutSources.Add(Source);
+		}
+	}
+
+	TOptional<FBindingSource> FViewModelFieldPathHelper::GetSelectedSource() const
+	{
+		TSet<FBindingSource> Sources;
+		GetAvailableSources(Sources);
+
+		for (const FBindingSource& Source : Sources)
+		{
+			if (Source.IsSelected)
+			{
+				return Source;
+			}
+		}
+
+		return TOptional<FBindingSource>();
+	}
+
+	void FViewModelFieldPathHelper::GetAvailableFields(TSet<FMVVMConstFieldVariant>& OutFields) const
+	{
+		FMVVMViewModelPropertyPath* Path = PathAttr.Get(nullptr);
+		if (Path == nullptr)
+		{
+			return;
+		}
+
+		const UMVVMWidgetBlueprintExtension_View* MVVMExtensionPtr = UMVVMWidgetBlueprintExtension_View::GetExtension<UMVVMWidgetBlueprintExtension_View>(WidgetBlueprint);
+		const UMVVMBlueprintView* MVVMBlueprintView = MVVMExtensionPtr->GetBlueprintView();
+
+		const FMVVMBlueprintViewModelContext* ViewModelContext = MVVMBlueprintView->FindViewModel(Path->ContextId);
+		if (ViewModelContext == nullptr)
+		{
+			return;
+		}
+
+		TArray<FMVVMAvailableBinding> Bindings = GEngine->GetEngineSubsystem<UMVVMSubsystem>()->GetViewModelAvailableBindings(ViewModelContext->GetViewModelClass());
+		Algo::Transform(Bindings, OutFields, [ViewModelContext](const FMVVMAvailableBinding& Binding)
+			{
+				return UE::MVVM::BindingHelper::FindFieldByName(ViewModelContext->GetViewModelClass(), Binding.GetBindingName());
+			});
+	}
+
+	FMVVMConstFieldVariant FViewModelFieldPathHelper::GetSelectedField() const
+	{
+		FMVVMBindingName BindingName = GetBindingName();
+
+		TSet<FMVVMConstFieldVariant> Fields;
+		GetAvailableFields(Fields);
+
+		for (const FMVVMConstFieldVariant& Field : Fields)
+		{
+			if (Field.GetName() == BindingName.ToName())
+			{
+				return Field;
+			}
+		}
+
+		return FMVVMConstFieldVariant();
+	}
+
+	void FViewModelFieldPathHelper::SetSelectedSource(const TOptional<FBindingSource>& Source) const
+	{
+		FMVVMViewModelPropertyPath* Path = PathAttr.Get(nullptr);
+		if (Path == nullptr)
+		{
+			return;
+		}
+
+		if (Source.IsSet())
+		{
+			Path->ContextId = Source.GetValue().SourceGuid;
+		}
+		else
+		{
+			Path->ContextId = FGuid();
+		}
+	}
+
+	FMVVMBindingName FViewModelFieldPathHelper::GetBindingName() const
+	{
+		if (FMVVMViewModelPropertyPath* Path = PathAttr.Get(nullptr))
+		{
+			return Path->GetBindingName();
+		}
+
+		return FMVVMBindingName();
+	}
+
+	void FViewModelFieldPathHelper::SetBindingName(FName BindingName) const
+	{
+		if (FMVVMViewModelPropertyPath* Path = PathAttr.Get(nullptr))
+		{
+			Path->SetBindingName(BindingName);
+		}
+	}
+}
diff --git a/Engine/Plugins/Runtime/ModelViewViewModel/Source/ModelViewViewModelEditor/Private/MVVMPropertyPathHelpers.h b/Engine/Plugins/Runtime/ModelViewViewModel/Source/ModelViewViewModelEditor/Private/MVVMPropertyPathHelpers.h
new file mode 100644
--- /dev/null
+++ b/Engine/Plugins/Runtime/ModelViewViewModel/Source/ModelViewViewModelEditor/Private/MVVMPropertyPathHelpers.h
@@ -0,0 +1,83 @@
+// Copyright Epic Games, Inc. All Rights Reserved.
+
+#pragma once
+
+#include "Internationalization/Text.h"
+#include "Misc/Attribute.h"
+#include "Misc/Guid.h"
+#include "Misc/Optional.h"
+#include "MVVMPropertyPath.h"
+#include "Types/MVVMBindingName.h"
+#include "Types/MVVMBindingSource.h"
+#include "Types/MVVMFieldVariant.h"
+#include "UObject/UnrealType.h"
+#include "WidgetBlueprint.h"
+
+namespace UE::MVVM
+{
+	class IFieldPathHelper
+	{
+	public:
+		virtual void GetAvailableFields(TSet<FMVVMConstFieldVariant>& Fields) const = 0;
+		virtual FMVVMConstFieldVariant GetSelectedField() const = 0;
+
+		virtual void GetAvailableSources(TSet<FBindingSource>& OutSources) const = 0;
+		virtual TOptional<FBindingSource> GetSelectedSource() const = 0;
+		virtual void SetSelectedSource(const TOptional<FBindingSource>& Source) const = 0;
+
+		virtual FMVVMBindingName GetBindingName() const = 0;
+		virtual void SetBindingName(FName BindingName) const = 0;
+	};
+
+	class FWidgetFieldPathHelper : public IFieldPathHelper
+	{
+	public:
+		FWidgetFieldPathHelper() {}
+		virtual ~FWidgetFieldPathHelper() {}
+		FWidgetFieldPathHelper(const TAttribute<FMVVMWidgetPropertyPath*>& WidgetPath, const UWidgetBlueprint* InWidgetBlueprint) :
+			PathAttr(WidgetPath),
+			WidgetBlueprint(InWidgetBlueprint)
+		{
+		}
+
+		virtual void GetAvailableFields(TSet<FMVVMConstFieldVariant>& Fields) const override;
+		virtual FMVVMConstFieldVariant GetSelectedField() const override;
+
+		virtual void GetAvailableSources(TSet<FBindingSource>& OutSources) const override;
+		virtual TOptional<FBindingSource> GetSelectedSource() const override;
+		virtual void SetSelectedSource(const TOptional<FBindingSource>& Source) const override;
+
+		virtual FMVVMBindingName GetBindingName() const override;
+		virtual void SetBindingName(FName BindingName) const override;
+
+	private:
+		TAttribute<FMVVMWidgetPropertyPath*> PathAttr;
+		const UWidgetBlueprint* WidgetBlueprint = nullptr;
+	};
+
+	class FViewModelFieldPathHelper : public IFieldPathHelper
+	{
+	public:
+		FViewModelFieldPathHelper() {}
+		virtual ~FViewModelFieldPathHelper() {}
+		FViewModelFieldPathHelper(const TAttribute<FMVVMViewModelPropertyPath*>& ViewModelPath, const UWidgetBlueprint* InWidgetBlueprint) :
+			PathAttr(ViewModelPath),
+			WidgetBlueprint(InWidgetBlueprint)
+		{
+		}
+
+		virtual void GetAvailableFields(TSet<FMVVMConstFieldVariant>& Fields) const override;
+		virtual FMVVMConstFieldVariant GetSelectedField() const override;
+
+		virtual void GetAvailableSources(TSet<FBindingSource>& OutSources) const override;
+		virtual TOptional<FBindingSource> GetSelectedSource() const override;
+		virtual void SetSelectedSource(const TOptional<FBindingSource>& Source) const override;
+
+		virtual FMVVMBindingName GetBindingName() const override;
+		virtual void SetBindingName(FName BindingName) const override;
+
+	private:
+		TAttribute<FMVVMViewModelPropertyPath*> PathAttr;
+		const UWidgetBlueprint* WidgetBlueprint = nullptr;
+	};
+}
diff --git a/Engine/Plugins/Runtime/ModelViewViewModel/Source/ModelViewViewModelEditor/Private/ModelViewViewModelEditorModule.h b/Engine/Plugins/Runtime/ModelViewViewModel/Source/ModelViewViewModelEditor/Private/ModelViewViewModelEditorModule.h
new file mode 100644
--- /dev/null
+++ b/Engine/Plugins/Runtime/ModelViewViewModel/Source/ModelViewViewModelEditor/Private/ModelViewViewModelEditorModule.h
@@ -0,0 +1,37 @@
+// Copyright Epic Games, Inc. All Rights Reserved.
+
+#pragma once
+
+#include "CoreMinimal.h"
+#include "Modules/ModuleInterface.h"
+#include "Modules/ModuleManager.h"
+
+class FMVVMPropertyBindingExtension;
+class FMVVMBindPropertiesDetailView;
+class FWidgetBlueprintApplicationMode;
+class FWorkflowAllowedTabSet;
+class UBlueprint;
+
+/**
+ *
+ */
+class FModelViewViewModelEditorModule : public IModuleInterface
+{
+public:
+	FModelViewViewModelEditorModule() = default;
+
+	//~ Begin IModuleInterface interface
+	 virtual void StartupModule() override;
+	 virtual void ShutdownModule() override;
+	//~ End IModuleInterface interface
+
+private:
+	void HandleRegisterBlueprintEditorTab(const FWidgetBlueprintApplicationMode& ApplicationMode, FWorkflowAllowedTabSet& TabFactories);
+	void HandleRenameVariableReferences(UBlueprint* Blueprint, UClass* VariableClass, const FName& OldVarName, const FName& NewVarName);
+	void HandleDeactiveMode(FWidgetBlueprintApplicationMode& InDesignerMode);
+	void HandleActivateMode(FWidgetBlueprintApplicationMode& InDesignerMode);
+
+private:
+	TUniquePtr<FMVVMBindPropertiesDetailView> BindDetailView;
+	TSharedPtr<FMVVMPropertyBindingExtension> PropertyBindingExtension;
+};
