MVVM: Move the ModelViewViewModel plugin to //UE5/Main/

diff --git a/Engine/Plugins/Runtime/ModelViewViewModel/Source/ModelViewViewModelEditor/Private/Tests/MVVMBindingHelperTest.cpp b/Engine/Plugins/Runtime/ModelViewViewModel/Source/ModelViewViewModelEditor/Private/Tests/MVVMBindingHelperTest.cpp
new file mode 100644
--- /dev/null
+++ b/Engine/Plugins/Runtime/ModelViewViewModel/Source/ModelViewViewModelEditor/Private/Tests/MVVMBindingHelperTest.cpp
@@ -0,0 +1,210 @@
+// Copyright Epic Games, Inc. All Rights Reserved.
+
+#include "MVVMBindingHelperTest.h"
+#include "Misc/AutomationTest.h"
+
+#include "Bindings/MVVMBindingHelper.h"
+
+
+#if WITH_AUTOMATION_WORKER
+
+#define LOCTEXT_NAMESPACE "MVVMBindingHelperTest"
+
+IMPLEMENT_SIMPLE_AUTOMATION_TEST(FMVVMBindingHelperTest, "System.Plugins.MVVM.BindingHelper", EAutomationTestFlags::ApplicationContextMask | EAutomationTestFlags::EngineFilter)
+
+
+bool FMVVMBindingHelperTest::RunTest(const FString& Parameters)
+{
+	auto TestProperty = [this](FName PropertyName
+		, bool bIsValidForSourceBinding, bool bIsValidForDestinationBinding
+		, bool bIsAccessibleDirectlyForSourceBinding, bool bIsAccessibleDirectlyForDestinationBinding
+		, bool bIsAccessibleWithGetterForSourceBinding, bool bIsAccessibleWithSetterForDestinationBinding
+		, bool bTryGetPropertyTypeForSourceBinding, bool bTryGetPropertyTypeForDestinationBinding)
+	{
+		FProperty* Property = UMVVMViewModelBindingHelperTest::StaticClass()->FindPropertyByName(PropertyName);
+		check(Property);
+
+		if (::UE::MVVM::BindingHelper::IsValidForSourceBinding(Property) != bIsValidForSourceBinding)
+		{
+			AddError(FString::Printf(TEXT("IsValidForSourceBinding failed for '%s'."), *PropertyName.ToString()));
+		}
+		if (::UE::MVVM::BindingHelper::IsValidForDestinationBinding(Property) != bIsValidForDestinationBinding)
+		{
+			AddError(FString::Printf(TEXT("IsValidForDestinationBinding failed for '%s'."), *PropertyName.ToString()));
+		}
+		if (::UE::MVVM::BindingHelper::IsAccessibleDirectlyForSourceBinding(Property) != bIsAccessibleDirectlyForSourceBinding)
+		{
+			AddError(FString::Printf(TEXT("IsAccessibleDirectlyForSourceBinding failed for '%s'."), *PropertyName.ToString()));
+		}
+		if (::UE::MVVM::BindingHelper::IsAccessibleDirectlyForDestinationBinding(Property) != bIsAccessibleDirectlyForDestinationBinding)
+		{
+			AddError(FString::Printf(TEXT("IsAccessibleDirectlyForDestinationBinding failed for '%s'."), *PropertyName.ToString()));
+		}
+		if (::UE::MVVM::BindingHelper::IsAccessibleWithGetterForSourceBinding(Property) != bIsAccessibleWithGetterForSourceBinding)
+		{
+			AddError(FString::Printf(TEXT("IsAccessibleWithGetterForSourceBinding failed for '%s'."), *PropertyName.ToString()));
+		}
+		if (::UE::MVVM::BindingHelper::IsAccessibleWithSetterForDestinationBinding(Property) != bIsAccessibleWithSetterForDestinationBinding)
+		{
+			AddError(FString::Printf(TEXT("IsAccessibleWithSetterForDestinationBinding failed for '%s'."), *PropertyName.ToString()));
+		}
+
+		if (::UE::MVVM::BindingHelper::TryGetPropertyTypeForSourceBinding(Property).HasValue() != bTryGetPropertyTypeForSourceBinding)
+		{
+			AddError(FString::Printf(TEXT("TryGetPropertyTypeForSourceBinding failed for '%s'."), *PropertyName.ToString()));
+		}
+		if (::UE::MVVM::BindingHelper::TryGetPropertyTypeForDestinationBinding(Property).HasValue() != bTryGetPropertyTypeForDestinationBinding)
+		{
+			AddError(FString::Printf(TEXT("TryGetPropertyTypeForDestinationBinding failed for '%s'."), *PropertyName.ToString()));
+		}
+	};
+
+	TestProperty(GET_MEMBER_NAME_CHECKED(UMVVMViewModelBindingHelperTest, PropertyA), false, false, false, false, false, false, false, false);
+	TestProperty(GET_MEMBER_NAME_CHECKED(UMVVMViewModelBindingHelperTest, PropertyB), true, true, true, true, false, false, true, true);
+	TestProperty(GET_MEMBER_NAME_CHECKED(UMVVMViewModelBindingHelperTest, PropertyC), true, false, true, false, false, false, true, false);
+	TestProperty(GET_MEMBER_NAME_CHECKED(UMVVMViewModelBindingHelperTest, PropertyD), true, true, true, true, true, true, true, true);
+	TestProperty(GET_MEMBER_NAME_CHECKED(UMVVMViewModelBindingHelperTest, PropertyE), true, false, true, false, true, false, true, false);
+
+	TestProperty("PropertyI", false, false, false, false, false, false, false, false);
+	TestProperty("PropertyJ", true, true, true, true, false, false, true, true);
+	TestProperty("PropertyK", true, false, true, false, false, false, true, false);
+	TestProperty("PropertyL", true, true, true, true, true, true, true, true);
+	TestProperty("PropertyM", true, false, true, false, true, false, true, false);
+
+	TestProperty("PropertyX", false, false, false, false, false, false, false, false);
+
+
+	auto TestFunction = [this](FName FunctionName
+		, bool bIsValidForSourceBinding, bool bIsValidForDestinationBinding, bool bIsValidForRuntimeConversion
+		, FFieldClass* PropertyTypeForSourceBinding, FFieldClass* PropertyTypeForDestinationBinding)
+	{
+		UFunction* Function = UMVVMViewModelBindingHelperTest::StaticClass()->FindFunctionByName(FunctionName);
+		check(Function);
+
+		if (::UE::MVVM::BindingHelper::IsValidForSourceBinding(Function) != bIsValidForSourceBinding)
+		{
+			AddError(FString::Printf(TEXT("IsValidForSourceBinding failed for '%s'."), *FunctionName.ToString()));
+		}
+		if (::UE::MVVM::BindingHelper::IsValidForDestinationBinding(Function) != bIsValidForDestinationBinding)
+		{
+			AddError(FString::Printf(TEXT("IsValidForDestinationBinding failed for '%s'."), *FunctionName.ToString()));
+		}
+		if (::UE::MVVM::BindingHelper::IsValidForRuntimeConversion(Function) != bIsValidForRuntimeConversion)
+		{
+			AddError(FString::Printf(TEXT("IsValidForRuntimeConversion failed for '%s'."), *FunctionName.ToString()));
+		}
+
+		{
+			TValueOrError<const FProperty*, FString> Result = ::UE::MVVM::BindingHelper::TryGetPropertyTypeForSourceBinding(Function);
+			if ((PropertyTypeForSourceBinding != nullptr) != Result.HasValue())
+			{
+				AddError(FString::Printf(TEXT("TryGetPropertyTypeForSourceBinding failed for '%s'."), *FunctionName.ToString()));
+			}
+			else if (PropertyTypeForSourceBinding && Result.GetValue()->GetClass() != PropertyTypeForSourceBinding)
+			{
+				AddError(FString::Printf(TEXT("TryGetPropertyTypeForSourceBinding failed for '%s'."), *FunctionName.ToString()));
+			}
+		}
+		{
+			TValueOrError<const FProperty*, FString> Result = ::UE::MVVM::BindingHelper::TryGetPropertyTypeForDestinationBinding(Function);
+			if ((PropertyTypeForDestinationBinding != nullptr) != Result.HasValue())
+			{
+				AddError(FString::Printf(TEXT("TryGetPropertyTypeForDestinationBinding failed for '%s'."), *FunctionName.ToString()));
+			}
+			else if (PropertyTypeForDestinationBinding && Result.GetValue()->GetClass() != PropertyTypeForDestinationBinding)
+			{
+				AddError(FString::Printf(TEXT("TryGetPropertyTypeForDestinationBinding failed for '%s'."), *FunctionName.ToString()));
+			}
+		}
+	};
+
+	TestFunction("FunctionGetA", false, false, false, nullptr, nullptr);
+	TestFunction("FunctionGetB", false, false, false, nullptr, nullptr);
+	TestFunction("FunctionGetC", true, false, false, FIntProperty::StaticClass(), nullptr);
+	TestFunction("FunctionGetD", false, false, false, nullptr, nullptr);
+	TestFunction("FunctionGetE", false, false, false, nullptr, nullptr);
+	TestFunction("FunctionGetF", true, false, false, FArrayProperty::StaticClass(), nullptr);
+	TestFunction("FunctionGetG", true, false, false, FArrayProperty::StaticClass(), nullptr);
+	TestFunction("FunctionGetH", false, false, false, nullptr, nullptr);
+	TestFunction("FunctionGetI", false, false, false, nullptr, nullptr);
+	TestFunction("FunctionGetJ", false, false, false, nullptr, nullptr);
+
+	TestFunction("FunctionSetA", false, false, false, nullptr, nullptr);
+	TestFunction("FunctionSetB", false, true, false, nullptr, FIntProperty::StaticClass());
+	TestFunction("FunctionSetC", false, false, false, nullptr, nullptr);
+	TestFunction("FunctionSetD", false, false, false, nullptr, nullptr);
+	TestFunction("FunctionSetE", false, false, false, nullptr, nullptr);
+	TestFunction("FunctionSetF", false, false, false, nullptr, nullptr);
+	TestFunction("FunctionSetG", false, true, false, nullptr, FArrayProperty::StaticClass());
+	TestFunction("FunctionSetH", false, false, false, nullptr, nullptr);
+
+	TestFunction("FunctionGetProtected", true, false, false, FIntProperty::StaticClass(), nullptr);
+	TestFunction("FunctionSetProtected", false, true, false, nullptr, FIntProperty::StaticClass());
+	TestFunction("FunctionGetter", true, false, false, FIntProperty::StaticClass(), nullptr);
+	TestFunction("FunctionSetter", false, true, false, nullptr, FIntProperty::StaticClass());
+
+	TestFunction("FunctionConversionA", false, false, false, nullptr, nullptr);
+	TestFunction("FunctionConversionB", false, false, false, nullptr, nullptr);
+	TestFunction("FunctionConversionC", false, false, true, nullptr, nullptr);
+	TestFunction("FunctionConversionD", false, false, false, nullptr, nullptr);
+	TestFunction("FunctionConversionE", false, true, false, nullptr, FIntProperty::StaticClass());
+	TestFunction("FunctionConversionF", false, false, false, nullptr, nullptr);
+	TestFunction("FunctionConversionG", false, false, true, nullptr, nullptr);
+	TestFunction("FunctionConversionH", false, false, true, nullptr, nullptr);
+	TestFunction("FunctionConversionI", false, false, false, nullptr, nullptr);
+	TestFunction("FunctionConversionJ", false, false, false, nullptr, nullptr);
+
+	auto TestConversionFunction = [this](FName FunctionName, FFieldClass* ExpectedReturnProperty, FFieldClass* ExpectedFirstProperty)
+	{
+		UFunction* Function = UMVVMViewModelBindingHelperTest::StaticClass()->FindFunctionByName(FunctionName);
+		check(Function);
+
+		const FProperty* ReturnProperty = ::UE::MVVM::BindingHelper::GetReturnProperty(Function);
+		if (ExpectedReturnProperty)
+		{
+			if (!ReturnProperty || ReturnProperty->GetClass() != ExpectedReturnProperty)
+			{
+				AddError(FString::Printf(TEXT("GetReturnProperty failed for '%s'."), *FunctionName.ToString()));
+			}
+		}
+		else
+		{
+			if (ReturnProperty)
+			{
+				AddError(FString::Printf(TEXT("GetReturnProperty failed for '%s'."), *FunctionName.ToString()));
+			}
+		}
+
+		const FProperty* FirstProperty = ::UE::MVVM::BindingHelper::GetFirstArgumentProperty(Function);
+		if (ExpectedFirstProperty)
+		{
+			if (!FirstProperty || FirstProperty->GetClass() != ExpectedFirstProperty)
+			{
+				AddError(FString::Printf(TEXT("GetFirstArgumentProperty failed for '%s'."), *FunctionName.ToString()));
+			}
+		}
+		else
+		{
+			if (FirstProperty)
+			{
+				AddError(FString::Printf(TEXT("GetFirstArgumentProperty failed for '%s'."), *FunctionName.ToString()));
+			}
+		}
+	};
+
+	TestConversionFunction("FunctionConversionA", nullptr, nullptr);
+	TestConversionFunction("FunctionConversionB", FIntProperty::StaticClass(), nullptr);
+	TestConversionFunction("FunctionConversionC", FIntProperty::StaticClass(), FIntProperty::StaticClass());
+	TestConversionFunction("FunctionConversionD", FIntProperty::StaticClass(), FIntProperty::StaticClass());
+	TestConversionFunction("FunctionConversionE", nullptr, FIntProperty::StaticClass());
+	TestConversionFunction("FunctionConversionF", nullptr, FIntProperty::StaticClass());
+	TestConversionFunction("FunctionConversionG", FArrayProperty::StaticClass(), FSetProperty::StaticClass());
+	TestConversionFunction("FunctionConversionH", FArrayProperty::StaticClass(), FSetProperty::StaticClass());
+	TestConversionFunction("FunctionConversionI", FArrayProperty::StaticClass(), FSetProperty::StaticClass());
+	TestConversionFunction("FunctionConversionJ", FArrayProperty::StaticClass(), FMapProperty::StaticClass());
+
+	return true;
+}
+
+#undef LOCTEXT_NAMESPACE 
+#endif //WITH_AUTOMATION_WORKER
\ No newline at end of file
diff --git a/Engine/Plugins/Runtime/ModelViewViewModel/Source/ModelViewViewModelEditor/Private/Widgets/SMVVMConversionPath.cpp b/Engine/Plugins/Runtime/ModelViewViewModel/Source/ModelViewViewModelEditor/Private/Widgets/SMVVMConversionPath.cpp
new file mode 100644
--- /dev/null
+++ b/Engine/Plugins/Runtime/ModelViewViewModel/Source/ModelViewViewModelEditor/Private/Widgets/SMVVMConversionPath.cpp
@@ -0,0 +1,211 @@
+// Copyright Epic Games, Inc. All Rights Reserved.
+
+#include "SMVVMConversionPath.h"
+
+#include "Editor.h"
+#include "Framework/MultiBox/MultiBoxBuilder.h"
+#include "MVVMBlueprintViewBinding.h"
+#include "MVVMEditorSubsystem.h"
+#include "MVVMPropertyPathHelpers.h"
+#include "MVVMSubsystem.h"
+#include "Styling/StyleColors.h"
+#include "Types/MVVMFieldVariant.h"
+#include "UMGStyle.h"
+#include "Widgets/Images/SImage.h"
+#include "Widgets/Input/SButton.h"
+#include "Widgets/Input/SMenuAnchor.h"
+#include "Widgets/Layout/SBox.h"
+
+#define LOCTEXT_NAMESPACE "MVVMConversionPath"
+
+void SMVVMConversionPath::Construct(const FArguments& InArgs, const UWidgetBlueprint* InWidgetBlueprint, bool bInIsGetter)
+{
+	bIsGetter = bInIsGetter;
+	WidgetBlueprint = InWidgetBlueprint;
+	OnFunctionChanged = InArgs._OnFunctionChanged;
+	Bindings = InArgs._Bindings;
+	check(Bindings.IsSet());
+
+	ChildSlot
+	[
+		SAssignNew(Anchor, SMenuAnchor)
+		.ToolTipText(this, &SMVVMConversionPath::GetFunctionToolTip)
+		.OnGetMenuContent(this, &SMVVMConversionPath::GetFunctionMenuContent)
+		.Visibility(this, &SMVVMConversionPath::IsFunctionVisible)
+		[
+			SNew(SBox)
+			.VAlign(VAlign_Center)
+			.HAlign(HAlign_Left)
+			.Padding(3, 0, 3, 0)
+			[
+				SNew(SButton)
+				.ButtonStyle(FAppStyle::Get(), "SimpleButton")
+				.OnClicked(this, &SMVVMConversionPath::OnButtonClicked)
+				[
+					SNew(SImage)
+					.Image(FUMGStyle::Get().GetBrush(bIsGetter ? "MVVM.GetterFunction" : "MVVM.SetterFunction"))
+					.ColorAndOpacity(this, &SMVVMConversionPath::GetFunctionColor)
+				]
+			]
+		]
+	];
+}
+
+EVisibility SMVVMConversionPath::IsFunctionVisible() const
+{
+	TArray<FMVVMBlueprintViewBinding*> ViewBindings = Bindings.Get(TArray<FMVVMBlueprintViewBinding*>());
+	if (ViewBindings.Num() == 0)
+	{
+		return EVisibility::Collapsed;
+	}
+
+	for (const FMVVMBlueprintViewBinding* Binding : ViewBindings)
+	{
+		bool bShouldBeVisible =
+			Binding->BindingType == EMVVMBindingMode::OneTimeToDestination ||
+			Binding->BindingType == EMVVMBindingMode::OneWayToDestination ||
+			Binding->BindingType == EMVVMBindingMode::TwoWay;
+
+		if (bShouldBeVisible)
+		{
+			return EVisibility::Visible;
+		}
+	}
+	
+	return EVisibility::Hidden;
+}
+
+FString SMVVMConversionPath::GetFunctionPath() const
+{
+	TArray<FMVVMBlueprintViewBinding*> ViewBindings = Bindings.Get(TArray<FMVVMBlueprintViewBinding*>());
+	for (const FMVVMBlueprintViewBinding* Binding : ViewBindings)
+	{
+		FString FunctionPath = bIsGetter ? Binding->Conversion.GetConversionFunctionPath : Binding->Conversion.SetConversionFunctionPath;
+		return FunctionPath;
+	}
+	return FString();
+}
+
+FText SMVVMConversionPath::GetFunctionToolTip() const
+{
+	FString FunctionPath = GetFunctionPath();
+	if (!FunctionPath.IsEmpty())
+	{
+		return FText::FromString(FunctionPath);
+	}
+
+	return bIsGetter ?
+		LOCTEXT("AddGetterConversionFunction", "Add conversion function to be used when getting the value from the viewmodel.") :
+		LOCTEXT("AddSetterConversionFunction", "Add conversion function to be used when setting the value in the viewmodel.");
+}
+
+FSlateColor SMVVMConversionPath::GetFunctionColor() const
+{
+	FString FunctionPath = GetFunctionPath();
+	if (FunctionPath.IsEmpty())
+	{
+		return FStyleColors::Foreground;
+	}
+
+	return FStyleColors::AccentGreen;
+}
+
+FReply SMVVMConversionPath::OnButtonClicked() const
+{
+	Anchor->SetIsOpen(!Anchor->IsOpen());
+	return FReply::Handled();
+}
+
+void SMVVMConversionPath::SetConversionFunction(const UFunction* Function)
+{
+	TArray<FMVVMBlueprintViewBinding*> ViewBindings = Bindings.Get(TArray<FMVVMBlueprintViewBinding*>());
+	if (ViewBindings.Num() == 0)
+	{
+		return;
+	}
+
+	for (FMVVMBlueprintViewBinding* Binding : ViewBindings)
+	{
+		if (bIsGetter)
+		{
+			Binding->Conversion.GetConversionFunctionPath = Function->GetPathName();
+		}
+		else
+		{
+			Binding->Conversion.SetConversionFunctionPath = Function->GetPathName();
+		}
+	}
+
+	if (OnFunctionChanged.IsBound())
+	{
+		OnFunctionChanged.Execute(Function->GetPathName());
+	}
+}
+
+TSharedRef<SWidget> SMVVMConversionPath::GetFunctionMenuContent()
+{
+	TArray<FMVVMBlueprintViewBinding*> ViewBindings = Bindings.Get(TArray<FMVVMBlueprintViewBinding*>());
+	if (ViewBindings.Num() == 0)
+	{
+		return SNullWidget::NullWidget;
+	}
+
+	TSet<const UFunction*> ConversionFunctions;
+
+	for (FMVVMBlueprintViewBinding* Binding : ViewBindings)
+	{
+		UE::MVVM::FViewModelFieldPathHelper ViewModelHelper(&Binding->ViewModelPath, WidgetBlueprint);
+		UE::MVVM::FMVVMConstFieldVariant ViewModelField = ViewModelHelper.GetSelectedField();
+
+		UE::MVVM::FWidgetFieldPathHelper WidgetHelper(&Binding->WidgetPath, WidgetBlueprint);
+		UE::MVVM::FMVVMConstFieldVariant WidgetField = WidgetHelper.GetSelectedField();
+
+		UMVVMSubsystem::FConstDirectionalBindingArgs Args;
+		Args.SourceBinding = bIsGetter ? ViewModelField : WidgetField;
+		Args.DestinationBinding = bIsGetter ? WidgetField : ViewModelField;
+
+		UMVVMEditorSubsystem* EditorSubsystem = GEditor->GetEditorSubsystem<UMVVMEditorSubsystem>();
+		TArray<const UFunction*> FunctionsForThis = EditorSubsystem->GetAvailableConversionFunctions(Args.SourceBinding, Args.DestinationBinding);
+
+		if (ConversionFunctions.Num() > 0)
+		{
+			ConversionFunctions = ConversionFunctions.Intersect(TSet<const UFunction*>(FunctionsForThis));
+		}
+		else
+		{
+			ConversionFunctions = TSet<const UFunction*>(FunctionsForThis);
+		}
+	}
+
+	FMenuBuilder MenuBuilder(true, TSharedPtr<const FUICommandList>());
+
+	if (ConversionFunctions.Num() == 0)
+	{
+		MenuBuilder.AddWidget(
+			SNew(SBox)
+			.Padding(10,0)
+			[
+				SNew(STextBlock)
+					.TextStyle(FAppStyle::Get(), "HintText")
+					.Text(LOCTEXT("NoCompatibleFunctions", "No compatible functions found."))
+			],
+			FText::GetEmpty(),
+			true, // no indent
+			false // searchable
+		);
+	}
+
+	for (const UFunction* Function : ConversionFunctions)
+	{
+		FUIAction Action(FExecuteAction::CreateSP(this, &SMVVMConversionPath::SetConversionFunction, Function));
+		MenuBuilder.AddMenuEntry(
+			Function->GetDisplayNameText(),
+			Function->GetToolTipText(),
+			FSlateIcon(FAppStyle::GetAppStyleSetName(), "Kismet.AllClasses.FunctionIcon"),
+			Action);
+	}
+
+	return MenuBuilder.MakeWidget();
+}
+
+#undef LOCTEXT_NAMESPACE
\ No newline at end of file
diff --git a/Engine/Plugins/Runtime/ModelViewViewModel/Source/ModelViewViewModelEditor/Private/Widgets/SMVVMConversionPath.h b/Engine/Plugins/Runtime/ModelViewViewModel/Source/ModelViewViewModelEditor/Private/Widgets/SMVVMConversionPath.h
new file mode 100644
--- /dev/null
+++ b/Engine/Plugins/Runtime/ModelViewViewModel/Source/ModelViewViewModelEditor/Private/Widgets/SMVVMConversionPath.h
@@ -0,0 +1,40 @@
+// Copyright Epic Games, Inc. All Rights Reserved.
+
+#pragma once
+
+#include "Widgets/SCompoundWidget.h"
+
+struct FMVVMBlueprintViewBinding;
+class SMenuAnchor;
+class UWidgetBlueprint;
+
+class SMVVMConversionPath : public SCompoundWidget
+{
+public:
+	DECLARE_DELEGATE_OneParam(FOnFunctionChanged, const FString&);
+
+	SLATE_BEGIN_ARGS(SMVVMConversionPath)
+		{
+		}
+		SLATE_ATTRIBUTE(TArray<FMVVMBlueprintViewBinding*>, Bindings) 
+		SLATE_EVENT(FOnFunctionChanged, OnFunctionChanged)
+	SLATE_END_ARGS()
+
+	void Construct(const FArguments& InArgs, const UWidgetBlueprint* InWidgetBlueprint, bool bInIsGetter);
+
+private:
+	EVisibility IsFunctionVisible() const;
+	FText GetFunctionToolTip() const;
+	FSlateColor GetFunctionColor() const;
+	void SetConversionFunction(const UFunction* Function);
+	TSharedRef<SWidget> GetFunctionMenuContent();
+	FReply OnButtonClicked() const;
+	FString GetFunctionPath() const;
+
+private:
+	TAttribute<TArray<FMVVMBlueprintViewBinding*>> Bindings;
+	FOnFunctionChanged OnFunctionChanged;
+	TSharedPtr<SMenuAnchor> Anchor;
+	const UWidgetBlueprint* WidgetBlueprint = nullptr;
+	bool bIsGetter = false;
+};
\ No newline at end of file
diff --git a/Engine/Plugins/Runtime/ModelViewViewModel/Source/ModelViewViewModelEditor/Private/Widgets/SMVVMFieldIcon.cpp b/Engine/Plugins/Runtime/ModelViewViewModel/Source/ModelViewViewModelEditor/Private/Widgets/SMVVMFieldIcon.cpp
new file mode 100644
--- /dev/null
+++ b/Engine/Plugins/Runtime/ModelViewViewModel/Source/ModelViewViewModelEditor/Private/Widgets/SMVVMFieldIcon.cpp
@@ -0,0 +1,66 @@
+// Copyright Epic Games, Inc. All Rights Reserved.
+
+#include "Widgets/SMVVMFieldIcon.h"
+
+#include "Bindings/MVVMBindingHelper.h"
+#include "BlueprintEditor.h"
+#include "Styling/SlateBrush.h"
+#include "Styling/SlateColor.h"
+#include "Styling/StyleColors.h"
+#include "Widgets/Images/SLayeredImage.h"
+
+#define LOCTEXT_NAMESPACE "SMVVMFieldIcon"
+
+void SMVVMFieldIcon::Construct(const FArguments& Args)
+{
+	ChildSlot
+	[
+		SAssignNew(LayeredImage, SLayeredImage)
+	];
+
+	RefreshBinding(Args._Field);
+}
+
+void SMVVMFieldIcon::RefreshBinding(const UE::MVVM::FMVVMConstFieldVariant& Field)
+{
+	LayeredImage->SetImage(nullptr);
+	LayeredImage->RemoveAllLayers();
+
+	if (Field.IsEmpty())
+	{
+		return;
+	}
+
+	const FProperty* IconProperty = nullptr;
+
+	if (Field.IsProperty())
+	{
+		IconProperty = Field.GetProperty();
+	}
+	else if (Field.IsFunction())
+	{
+		const UFunction* Function = Field.GetFunction();
+		const FProperty* ReturnProperty = UE::MVVM::BindingHelper::GetReturnProperty(Function);
+		if (ReturnProperty != nullptr)
+		{
+			IconProperty = ReturnProperty;
+		}
+		else
+		{
+			IconProperty = UE::MVVM::BindingHelper::GetFirstArgumentProperty(Function);
+		}
+	}
+
+	if (IconProperty != nullptr)
+	{
+		FSlateColor PrimaryColor, SecondaryColor;
+		const FSlateBrush* SecondaryBrush = nullptr;
+		const FSlateBrush* PrimaryBrush = FBlueprintEditor::GetVarIconAndColorFromProperty(IconProperty, PrimaryColor, SecondaryBrush, SecondaryColor);
+
+		LayeredImage->SetImage(PrimaryBrush);
+		LayeredImage->SetColorAndOpacity(PrimaryColor);
+		LayeredImage->AddLayer(SecondaryBrush, SecondaryColor);
+	}
+}
+
+#undef LOCTEXT_NAMESPACE
diff --git a/Engine/Plugins/Runtime/ModelViewViewModel/Source/ModelViewViewModelEditor/Private/Widgets/SMVVMFieldIcon.h b/Engine/Plugins/Runtime/ModelViewViewModel/Source/ModelViewViewModelEditor/Private/Widgets/SMVVMFieldIcon.h
new file mode 100644
--- /dev/null
+++ b/Engine/Plugins/Runtime/ModelViewViewModel/Source/ModelViewViewModelEditor/Private/Widgets/SMVVMFieldIcon.h
@@ -0,0 +1,22 @@
+// Copyright Epic Games, Inc. All Rights Reserved.
+
+#pragma once
+
+#include "Types/MVVMFieldVariant.h"
+#include "Widgets/SCompoundWidget.h"
+
+class SLayeredImage;
+
+class SMVVMFieldIcon : public SCompoundWidget
+{
+public:
+	SLATE_BEGIN_ARGS(SMVVMFieldIcon) {}
+		SLATE_ARGUMENT(UE::MVVM::FMVVMConstFieldVariant, Field)
+	SLATE_END_ARGS()
+
+	void Construct(const FArguments& Args);
+	void RefreshBinding(const UE::MVVM::FMVVMConstFieldVariant& Field);
+
+private:
+	TSharedPtr<SLayeredImage> LayeredImage;
+};
\ No newline at end of file
diff --git a/Engine/Plugins/Runtime/ModelViewViewModel/Source/ModelViewViewModelEditor/Private/Widgets/SMVVMFieldSelector.h b/Engine/Plugins/Runtime/ModelViewViewModel/Source/ModelViewViewModelEditor/Private/Widgets/SMVVMFieldSelector.h
new file mode 100644
--- /dev/null
+++ b/Engine/Plugins/Runtime/ModelViewViewModel/Source/ModelViewViewModelEditor/Private/Widgets/SMVVMFieldSelector.h
@@ -0,0 +1,94 @@
+// Copyright Epic Games, Inc. All Rights Reserved.
+
+#pragma once
+
+#include "MVVMFieldVariantListTraits.h"
+#include "MVVMPropertyPathHelpers.h" 
+#include "Styling/CoreStyle.h"
+#include "Templates/ValueOrError.h"
+#include "Types/MVVMBindingMode.h"
+#include "Types/MVVMFieldVariant.h"
+#include "UObject/UnrealType.h"
+#include "Widgets/DeclarativeSyntaxSupport.h"
+#include "Widgets/SCompoundWidget.h"
+#include "Widgets/Input/SComboBox.h"
+
+namespace UE::MVVM
+{
+	class IFieldPathHelper;
+}
+
+class SMVVMFieldIcon;
+class STextBlock; 
+
+class SMVVMFieldEntry : public SCompoundWidget
+{
+public:
+
+	using FIsEntryValidResult = TValueOrError<bool, FString>; 
+
+	DECLARE_DELEGATE_RetVal_OneParam(FIsEntryValidResult, FIsEntryValid, UE::MVVM::FMVVMConstFieldVariant);
+
+	SLATE_BEGIN_ARGS(SMVVMFieldEntry) :
+		_TextStyle( &FCoreStyle::Get().GetWidgetStyle<FTextBlockStyle>( "NormalText" ) )
+		{}
+		SLATE_STYLE_ARGUMENT(FTextBlockStyle, TextStyle)
+		SLATE_ARGUMENT(UE::MVVM::FMVVMConstFieldVariant, Field)
+		SLATE_EVENT(FIsEntryValid, OnValidate)
+	SLATE_END_ARGS()
+
+	void Construct(const FArguments& InArgs);
+	void Refresh(); 
+	void SetField(const UE::MVVM::FMVVMConstFieldVariant& InField);
+
+private:
+	UE::MVVM::FMVVMConstFieldVariant Field;
+	FIsEntryValid OnValidate;
+	TSharedPtr<SMVVMFieldIcon> Icon;
+	TSharedPtr<STextBlock> Label;
+};
+
+class SMVVMFieldSelector : public SCompoundWidget
+{
+public:
+	DECLARE_DELEGATE_OneParam(FSelectionChanged, UE::MVVM::FMVVMConstFieldVariant);
+
+	SLATE_BEGIN_ARGS(SMVVMFieldSelector) :
+		_TextStyle(&FCoreStyle::Get().GetWidgetStyle<FTextBlockStyle>("NormalText")),
+		_IsSource(false)
+		{
+		}
+		SLATE_STYLE_ARGUMENT(FTextBlockStyle, TextStyle)
+		SLATE_ARGUMENT(TArray<UE::MVVM::IFieldPathHelper*>, PathHelpers)
+		SLATE_ARGUMENT(TArray<UE::MVVM::IFieldPathHelper*>, CounterpartHelpers)
+		SLATE_ATTRIBUTE(EMVVMBindingMode, BindingMode)
+		SLATE_ARGUMENT(bool, IsSource)
+		SLATE_EVENT(FSelectionChanged, OnSelectionChanged)
+	SLATE_END_ARGS()
+
+	void Construct(const FArguments& InArgs);
+	void Refresh();
+
+private:
+	TSharedRef<SWidget> OnGenerateFieldWidget(UE::MVVM::FMVVMConstFieldVariant FieldPath) const;
+	void OnComboBoxSelectionChanged(UE::MVVM::FMVVMConstFieldVariant Selected, ESelectInfo::Type SelectionType);
+	TOptional<UE::MVVM::FBindingSource> GetSelectedSource() const;
+
+	TValueOrError<bool, FString> IsValidBindingForField(const UE::MVVM::FMVVMConstFieldVariant& Field, const UE::MVVM::FMVVMConstFieldVariant& CounterpartField) const;
+	TValueOrError<bool, FString> ValidateField(UE::MVVM::FMVVMConstFieldVariant Field) const;
+
+private:
+	TArray<UE::MVVM::IFieldPathHelper*> PathHelpers;
+	TArray<UE::MVVM::IFieldPathHelper*> CounterpartHelpers;
+	TAttribute<EMVVMBindingMode> BindingMode;
+	TSharedPtr<SMVVMFieldEntry> SelectedEntry;
+	TOptional<UE::MVVM::FBindingSource> SelectedSource;
+	const FTextBlockStyle* TextStyle = nullptr;
+	bool bIsSource = false;
+
+	TSharedPtr<SComboBox<UE::MVVM::FMVVMConstFieldVariant>> FieldComboBox;
+
+	FSelectionChanged OnSelectionChangedDelegate;
+
+	TArray<UE::MVVM::FMVVMConstFieldVariant> AvailableFields;
+}; 
diff --git a/Engine/Plugins/Runtime/ModelViewViewModel/Source/ModelViewViewModelEditor/Private/Widgets/SMVVMManageViewModelsWidget.cpp b/Engine/Plugins/Runtime/ModelViewViewModel/Source/ModelViewViewModelEditor/Private/Widgets/SMVVMManageViewModelsWidget.cpp
new file mode 100644
--- /dev/null
+++ b/Engine/Plugins/Runtime/ModelViewViewModel/Source/ModelViewViewModelEditor/Private/Widgets/SMVVMManageViewModelsWidget.cpp
@@ -0,0 +1,825 @@
+// Copyright Epic Games, Inc. All Rights Reserved.
+
+#include "SMVVMManageViewModelsWidget.h"
+
+#include "Bindings/MVVMBindingHelper.h"
+#include "Blueprint/UserWidget.h"
+#include "Editor.h"
+#include "Framework/MultiBox/MultiBoxBuilder.h"
+#include "Kismet2/KismetEditorUtilities.h"
+#include "MVVMBlueprintViewModelContext.h"
+#include "MVVMViewModelBase.h"
+#include "SMVVMViewModelContextListWidget.h"
+#include "SMVVMViewModelBindingListWidget.h"
+#include "SPrimaryButton.h"
+#include "SSimpleButton.h"
+#include "Styling/MVVMEditorStyle.h"
+#include "Styling/SlateIconFinder.h"
+#include "Textures/SlateIcon.h"
+#include "Types/MVVMFieldVariant.h"
+#include "UMGStyle.h"
+#include "UObject/UObjectIterator.h"
+#include "Widgets/Images/SImage.h"
+#include "Widgets/Input/SButton.h"
+#include "Widgets/Input/SComboButton.h"
+#include "Widgets/Input/SEditableTextBox.h"
+#include "Widgets/Input/SSearchBox.h"
+#include "Widgets/Layout/SSplitter.h"
+#include "Widgets/Layout/SUniformGridPanel.h"
+#include "Widgets/SWindow.h"
+#include "Widgets/Views/SListView.h"
+#include "Widgets/Views/STreeView.h"
+
+#define LOCTEXT_NAMESPACE "ManageViewModelsWidget"
+
+namespace UE::MVVM::Private
+{
+	class FMVVMViewModelTreeNode : public TSharedFromThis<FMVVMViewModelTreeNode>
+	{
+	public:
+		struct FClassData
+		{
+			TSubclassOf<UMVVMViewModelBase> Class;
+			bool bPassesClassFlagsFilter;
+		};
+
+		struct FTagData
+		{
+			FText Name;
+		};
+
+		FMVVMViewModelTreeNode(FMVVMViewModelTreeNode& OtherNode)
+		{
+			NodeData = OtherNode.NodeData;
+		}
+
+		FMVVMViewModelTreeNode(FClassData& Class)
+		{
+			NodeData.Set<FMVVMViewModelTreeNode::FClassData>(Class);
+		}
+
+		FMVVMViewModelTreeNode(FTagData& Tag)
+		{
+			NodeData.Set<FMVVMViewModelTreeNode::FTagData>(Tag);
+		}
+
+		void AddChild(TSharedRef<FMVVMViewModelTreeNode> Child)
+		{
+			ChildrenList.Add(Child);
+		}
+
+		TArray<TSharedRef<FMVVMViewModelTreeNode>>& GetChildrenList()
+		{
+			return ChildrenList;
+		}
+
+		void UpdateNodeSearchFilterFlag(const FString& SearchString)
+		{
+			bPassesSearchTextFilter = SearchString.IsEmpty() ? true : GetDisplayName().Contains(SearchString);
+
+			for (TSharedRef<FMVVMViewModelTreeNode>& Child : GetChildrenList())
+			{
+				Child->UpdateNodeSearchFilterFlag(SearchString);
+			}
+		}
+
+		TSet<TSharedRef<FMVVMViewModelTreeNode>> GetNodeBranchesToDiscard()
+		{
+			TSet<TSharedRef<FMVVMViewModelTreeNode>> NodeBranchesToDiscard;
+
+			TSharedRef<FMVVMViewModelTreeNode> SharedNode = AsShared();
+			GetNodeBranchesToDiscard_Helper(SharedNode, NodeBranchesToDiscard);
+
+			return NodeBranchesToDiscard;
+		}
+
+		bool IsAddable()
+		{
+			if (FClassData* ClassInfo = GetClassData())
+			{
+				return ClassInfo->bPassesClassFlagsFilter;
+			}
+
+			return false;
+		}
+
+		bool IsPassingSearchFilter()
+		{
+			return bPassesSearchTextFilter;
+		}
+
+		FText GetDisplayNameText()
+		{
+			if (FMVVMViewModelTreeNode::FClassData* ClassInfo = GetClassData())
+			{
+				return ClassInfo->Class->GetDisplayNameText();
+			}
+			else
+			{
+				FMVVMViewModelTreeNode::FTagData* TagInfo = GetTagData();
+				check(TagInfo);
+				return TagInfo->Name;
+			}
+		}
+
+		FString GetDisplayName()
+		{
+			if (FMVVMViewModelTreeNode::FClassData* ClassInfo = GetClassData())
+			{
+				return ClassInfo->Class->GetDisplayNameText().ToString();
+			}
+			else
+			{
+				FMVVMViewModelTreeNode::FTagData* TagInfo = GetTagData();
+				check(TagInfo);
+				return TagInfo->Name.ToString();
+			}
+		}
+
+		TSharedRef<SWidget> GetIcon()
+		{
+			if (FMVVMViewModelTreeNode::FClassData* ClassInfo = GetClassData())
+			{
+				return SNew(SImage).Image(FSlateIconFinder::FindIconBrushForClass(ClassInfo->Class));
+			}
+			else
+			{
+				return SNew(SImage)
+					.Image(FSlateIcon(FMVVMEditorStyle::Get().GetStyleSetName(), "ViewModelSelection.TagIcon").GetIcon());
+			}
+		}
+
+		FClassData* GetClassData()
+		{
+			return NodeData.TryGet<FClassData>();
+		}
+
+		FTagData* GetTagData()
+		{
+			return NodeData.TryGet<FTagData>();
+		}
+
+		void SortChildren()
+		{
+			GetChildrenList().Sort([](const TSharedPtr<FMVVMViewModelTreeNode>& A, const TSharedPtr<FMVVMViewModelTreeNode>& B)
+				{
+					return A->GetDisplayName() <= B->GetDisplayName();
+				});
+		}
+
+	private:
+		static bool GetNodeBranchesToDiscard_Helper(TSharedRef<FMVVMViewModelTreeNode>& Node, TSet<TSharedRef<FMVVMViewModelTreeNode>>& OutNodeBranchesToDiscard)
+		{
+			bool bShouldKeep = false;
+
+			for (TSharedRef<FMVVMViewModelTreeNode>& Child : Node->GetChildrenList())
+			{
+				bool bShouldKeepChild = GetNodeBranchesToDiscard_Helper(Child, OutNodeBranchesToDiscard);
+				bShouldKeep = bShouldKeep || bShouldKeepChild;
+			}
+
+			if (!bShouldKeep)
+			{
+				if (FClassData* Class = Node->GetClassData())
+				{
+					bShouldKeep = bShouldKeep || (Class->bPassesClassFlagsFilter && Node->IsPassingSearchFilter());
+				}
+			}
+
+			if (!bShouldKeep)
+			{
+				OutNodeBranchesToDiscard.Add(Node);
+			}
+
+			return bShouldKeep;
+		}
+
+		TVariant<FClassData, FTagData> NodeData;
+		bool bPassesSearchTextFilter = true;
+		TArray<TSharedRef<FMVVMViewModelTreeNode>> ChildrenList;
+	};
+} // namespace UE::MVVM::Private
+
+void SMVVMManageViewModelsWidget::Construct(const FArguments& InArgs)
+{
+	using UE::MVVM::Private::FMVVMViewModelTreeNode;
+
+	WeakParentWindow = InArgs._ParentWindow;
+	OnViewModelContextsPicked = InArgs._OnViewModelContextsPickedDelegate;
+	ViewModelContextListWidget = InArgs._ManageViewModelsWidget;
+
+	PopulateViewModelsTree();
+
+	TSharedRef<SWidget> ButtonsPanelContent = InArgs._ButtonsPanel.Widget;
+	if (InArgs._ButtonsPanel.Widget == SNullWidget::NullWidget)
+	{
+		SAssignNew(ButtonsPanelContent, SUniformGridPanel)
+			.SlotPadding(FEditorStyle::Get().GetMargin("StandardDialog.SlotPadding"))
+			+ SUniformGridPanel::Slot(0, 0)
+			[
+				SNew(SPrimaryButton)
+				.Text(LOCTEXT("ViewModelFinishButtonText", "Finish"))
+				.OnClicked(this, &SMVVMManageViewModelsWidget::HandleClicked_Finish)
+			]
+			+ SUniformGridPanel::Slot(1, 0)
+			[
+				SNew(SButton)
+				.Text(LOCTEXT("ViewModelCancelButtonText", "Cancel"))
+				.HAlign(HAlign_Center)
+				.OnClicked(this, &SMVVMManageViewModelsWidget::HandleClicked_Cancel)
+			];
+	}
+
+
+	ChildSlot
+	[
+		SNew(SBorder)
+		.BorderImage(FEditorStyle::GetBrush("ToolPanel.GroupBorder"))
+		[
+			SNew(SVerticalBox)
+			+ SVerticalBox::Slot()
+			.AutoHeight()
+			.Padding(6)
+			[
+				GenerateSearchBar()
+			]
+			+ SVerticalBox::Slot()
+			.FillHeight(1.0f)
+			[
+				SNew(SSplitter)
+				.PhysicalSplitterHandleSize(2.0f)
+				.Orientation(Orient_Vertical)
+				+ SSplitter::Slot()
+				.Value(0.6f)
+				[
+					SNew(SSplitter)
+					.PhysicalSplitterHandleSize(2.0f)
+					+ SSplitter::Slot()
+					.Value(0.6f)
+					[
+						SNew(SBorder)
+						.BorderImage(FStyleDefaults::GetNoBrush())
+						[
+							SAssignNew(TreeViewWidget, STreeView<TSharedPtr<FMVVMViewModelTreeNode>>)
+							.SelectionMode(ESelectionMode::Single)
+							.TreeItemsSource(&NodeTreeViewSource)
+							.OnGetChildren(this, &SMVVMManageViewModelsWidget::HandleGetChildrenForTreeView)
+							.OnSetExpansionRecursive(this, &SMVVMManageViewModelsWidget::SetAllExpansionStates_Helper)
+							.OnGenerateRow(this, &SMVVMManageViewModelsWidget::HandleGenerateRowForTreeView)
+							.OnSelectionChanged(this, &SMVVMManageViewModelsWidget::HandleTreeViewSelectionChanged)
+							.OnExpansionChanged(this, &SMVVMManageViewModelsWidget::HandleTreeViewExpansionChanged)
+							.OnMouseButtonDoubleClick(this, &SMVVMManageViewModelsWidget::HandleTreeRowDoubleClicked)
+							.ItemHeight(20.0f)
+						]
+					]
+					+ SSplitter::Slot()
+					.Value(0.4f)
+					[
+						SNew(SBorder)
+						.BorderImage(FStyleDefaults::GetNoBrush())
+						[
+							SAssignNew(ViewModelBindingListWidget, SMVVMViewModelBindingListWidget)
+						]
+					]
+				]
+				+ SSplitter::Slot()
+				.Value(0.4f)
+				[
+					SNew(SVerticalBox)
+					+ SVerticalBox::Slot()
+					.VAlign(VAlign_Center)
+					.Padding(FMargin(6.0f, 3.0f))
+					.AutoHeight()
+					[
+						SNew(STextBlock)
+						.Text(LOCTEXT("SelectedViewModelsHeading", "Selected ViewModels"))
+						.Font(FAppStyle::Get().GetFontStyle("BoldFont"))
+					]
+					+ SVerticalBox::Slot()
+					.FillHeight(1.0f)
+					[
+						ViewModelContextListWidget.ToSharedRef()
+					]
+				]
+			]
+			+ SVerticalBox::Slot()
+			.HAlign(HAlign_Right)
+			.VAlign(VAlign_Bottom)
+			.AutoHeight()
+			.Padding(8)
+			[
+				ButtonsPanelContent
+			]
+		]
+	];
+
+	RefreshTreeView(true);
+}
+
+void SMVVMManageViewModelsWidget::PopulateViewModelsTree()
+{
+	using UE::MVVM::Private::FMVVMViewModelTreeNode;
+	TMap<FName, TSharedRef<FMVVMViewModelTreeNode>> ClassNodeMap;
+
+	TMap<FString, TArray<FMVVMViewModelTreeNode::FClassData>> TagClassMap;
+	FString UntaggedTagName = "<untagged>";
+	TagClassMap.Add(UntaggedTagName);
+
+	{
+		FMVVMViewModelTreeNode::FClassData NullClassData;
+		NullClassData.Class = UMVVMViewModelBase::StaticClass();
+		NullClassData.bPassesClassFlagsFilter = false;
+		ClassListRootNode = MakeShared<FMVVMViewModelTreeNode>(NullClassData);
+	}
+
+	for (TObjectIterator<UClass> It; It; ++It)
+	{
+		bool bPassesClassFlagsFilter = !It->HasAnyClassFlags(CLASS_Abstract | CLASS_Deprecated |
+										CLASS_NewerVersionExists | CLASS_HideDropDown | CLASS_Hidden);
+		bool bIsNotSkeleton = !FKismetEditorUtilities::IsClassABlueprintSkeleton(*It);
+
+		if (It->IsChildOf(UMVVMViewModelBase::StaticClass()) && bIsNotSkeleton)
+		{
+			UClass* Class = *It;
+			FMVVMViewModelTreeNode::FClassData Data;
+			Data.Class = Class;
+			Data.bPassesClassFlagsFilter = bPassesClassFlagsFilter;
+			if (bPassesClassFlagsFilter)
+			{
+				TArray<FString> Tags;
+				Class->GetMetaData("MVVMTags").ParseIntoArray(Tags, TEXT(","));
+
+				if (Tags.IsEmpty())
+				{
+					TagClassMap[UntaggedTagName].Add(Data);
+				} 
+				else
+				{
+					for (const FString& TagName : Tags)
+					{
+						TArray<FMVVMViewModelTreeNode::FClassData>& ClassDataList = TagClassMap.FindOrAdd(TagName);
+						ClassDataList.Add(Data);
+					}
+				}
+
+				ClassListRootNode->AddChild(MakeShared<FMVVMViewModelTreeNode>(Data));
+			}
+
+			ClassNodeMap.Add(Class->GetFName(), MakeShared<FMVVMViewModelTreeNode>(Data));
+		}
+	}
+
+	// Build Tag Hierarchy
+	{
+		FText AllTagsStr = LOCTEXT("MVVMViewModelAllTagsTreeNode", "<all_tags>");
+		FMVVMViewModelTreeNode::FTagData AllTagsNodeData{ AllTagsStr };
+
+		TagTreeRootNode = MakeShared<FMVVMViewModelTreeNode>(AllTagsNodeData);
+		TagClassMap.KeySort(TLess<FString>());
+
+
+		for (TPair<FString, TArray<FMVVMViewModelTreeNode::FClassData>>& Pair : TagClassMap)
+		{
+			FText LocalizedTagText;
+			if (Pair.Key == UntaggedTagName)
+			{
+				LocalizedTagText = LOCTEXT("MVVMViewModelUntaggedTreeNode", "<untagged>");
+			}
+			else
+			{
+				LocalizedTagText = FText::FromString(Pair.Key);
+			}
+			FMVVMViewModelTreeNode::FTagData NewTagData{ LocalizedTagText };
+			TSharedRef<FMVVMViewModelTreeNode> NewTagNode = MakeShared<FMVVMViewModelTreeNode>(NewTagData);
+			TagTreeRootNode->AddChild(NewTagNode);
+			for (FMVVMViewModelTreeNode::FClassData& NewClass : Pair.Value)
+			{
+				NewTagNode->AddChild(MakeShared<FMVVMViewModelTreeNode>(NewClass));
+			}
+		}
+	}
+
+	// Build Class Hierarchy
+	for (const TPair<FName, TSharedRef<FMVVMViewModelTreeNode>>& Pair : ClassNodeMap)
+	{
+		TSubclassOf<UMVVMViewModelBase> Class = GetClassFromNode(Pair.Value);
+
+		if (Class == UMVVMViewModelBase::StaticClass())
+		{
+			ClassTreeRootNode = Pair.Value;
+			continue;
+		}
+
+		{
+			FName SuperClassName = Class->GetSuperClass()->GetFName();
+
+			if (TSharedRef<FMVVMViewModelTreeNode>* SuperClassNodePtr = ClassNodeMap.Find(SuperClassName))
+			{
+				(*SuperClassNodePtr)->AddChild(Pair.Value);
+			}
+		}
+
+	}
+
+	CurrentViewTreeRootNode = ClassListRootNode;
+}
+
+void SMVVMManageViewModelsWidget::SetViewMode(EViewMode Mode)
+{
+	CurrentViewMode = Mode;
+
+	switch (Mode)
+	{
+	case EViewMode::List:
+		CurrentViewTreeRootNode = ClassListRootNode;
+		break;
+	case EViewMode::Hierarchy:
+		CurrentViewTreeRootNode = ClassTreeRootNode;
+		break;
+	case EViewMode::Tag:
+		CurrentViewTreeRootNode = TagTreeRootNode;
+		break;
+	}
+
+	RefreshTreeView();
+}
+
+TSharedRef<SWidget> SMVVMManageViewModelsWidget::HandleGetMenuContent_ViewSettings()
+{
+	FMenuBuilder MenuBuilder(true, NULL);
+
+	SMVVMManageViewModelsWidget* Self = this;
+
+	MenuBuilder.BeginSection("ViewSettings", LOCTEXT("ModeHeading", "View Mode"));
+	{
+		MenuBuilder.AddMenuEntry(
+			LOCTEXT("ListView", "List View"),
+			LOCTEXT("ListViewToolTip", "Show available ViewModels as a flat list"),
+			FSlateIcon(),
+			FUIAction(
+				FExecuteAction::CreateLambda([Self]() { Self->SetViewMode(EViewMode::List); }),
+				FCanExecuteAction(),
+				FIsActionChecked::CreateLambda([Self]() { return Self->CurrentViewMode == EViewMode::List; })
+			),
+			NAME_None,
+			EUserInterfaceActionType::RadioButton
+		);
+
+		MenuBuilder.AddMenuEntry(
+			LOCTEXT("HierarchyView", "Hierarchy View"),
+			LOCTEXT("HierarchyViewToolTip", "Show available ViewModels as a class hierarchy"),
+			FSlateIcon(),
+			FUIAction(
+				FExecuteAction::CreateLambda([Self]() { Self->SetViewMode(EViewMode::Hierarchy); }),
+				FCanExecuteAction(),
+				FIsActionChecked::CreateLambda([Self]() { return Self->CurrentViewMode == EViewMode::Hierarchy; })
+			),
+			NAME_None,
+			EUserInterfaceActionType::RadioButton
+		);
+
+		MenuBuilder.AddMenuEntry(
+			LOCTEXT("TagView", "Tag View"),
+			LOCTEXT("TagViewToolTip", "Show available ViewModels grouped by MVVMTags"),
+			FSlateIcon(),
+			FUIAction(
+				FExecuteAction::CreateLambda([Self]() { Self->SetViewMode(EViewMode::Tag); }),
+				FCanExecuteAction(),
+				FIsActionChecked::CreateLambda([Self]() { return Self->CurrentViewMode == EViewMode::Tag; })
+			),
+			NAME_None,
+			EUserInterfaceActionType::RadioButton
+		);
+	}
+	MenuBuilder.EndSection();
+
+	if (CurrentViewMode == EViewMode::Hierarchy || CurrentViewMode == EViewMode::Tag)
+	{
+		MenuBuilder.BeginSection("HierarchySettings", LOCTEXT("HierarchyHeading", "Hierarchy"));
+		{
+			MenuBuilder.AddMenuEntry(
+				LOCTEXT("ExpandAll", "Expand All"),
+				LOCTEXT("ExpandAllToolTip", "Expand All Items in the Hierarchy"),
+				FSlateIcon(),
+				FUIAction(
+					FExecuteAction::CreateLambda([Self]() {
+						if (!Self->NodeTreeViewSource.IsEmpty())
+						{
+							for (TSharedPtr<FMVVMViewModelTreeNode>& Child : Self->NodeTreeViewSource)
+							{
+								Self->SetAllExpansionStates_Helper(Child, true);
+							}
+						}
+					})
+				)
+			);
+
+			MenuBuilder.AddMenuEntry(
+				LOCTEXT("CollapseAll", "Collapse All"),
+				LOCTEXT("CollapseAllToolTip", "Collapse All Items in the Hierarchy"),
+				FSlateIcon(),
+				FUIAction(
+					FExecuteAction::CreateLambda([Self]() {
+						if (!Self->NodeTreeViewSource.IsEmpty())
+						{
+							for (TSharedPtr<FMVVMViewModelTreeNode>& Child : Self->NodeTreeViewSource)
+							{
+								Self->SetAllExpansionStates_Helper(Child, false);
+							}
+						}					
+					})
+				)
+			);
+		}
+		MenuBuilder.EndSection();
+	}
+
+	return MenuBuilder.MakeWidget();
+}
+
+TSharedRef<SWidget> SMVVMManageViewModelsWidget::GenerateSearchBar()
+{
+	SMVVMManageViewModelsWidget* Self = this;
+
+	return SNew(SHorizontalBox)
+		+ SHorizontalBox::Slot()
+		.Padding(4.0f)
+		.FillWidth(1.0f)
+		[
+			SNew(SSearchBox)
+			.HintText(LOCTEXT("FilterSearch", "Search..."))
+			.ToolTipText(LOCTEXT("FilterSearchHint", "Type here to search"))
+			.OnTextChanged_Lambda([Self](const FText& NewText) { Self->UpdateSearchString(NewText.ToString()); })
+		]
+		+ SHorizontalBox::Slot()
+		.Padding(4.0f)
+		.AutoWidth()
+		[
+			SNew(SComboButton)
+			.ComboButtonStyle(FAppStyle::Get(), "SimpleComboButtonWithIcon")
+			.OnGetMenuContent(this, &SMVVMManageViewModelsWidget::HandleGetMenuContent_ViewSettings)
+			.ButtonContent()
+			[
+				SNew(SImage)
+				.ColorAndOpacity(FSlateColor::UseForeground())
+				.Image(FAppStyle::Get().GetBrush("Icons.Settings"))
+			]
+		];
+
+}
+
+void SMVVMManageViewModelsWidget::UpdateSearchString(const FString& SearchString)
+{
+	bSearchStringIsEmpty = SearchString.IsEmpty();
+
+	ClassListRootNode->UpdateNodeSearchFilterFlag(SearchString);
+	ClassTreeRootNode->UpdateNodeSearchFilterFlag(SearchString);
+	TagTreeRootNode->UpdateNodeSearchFilterFlag(SearchString);
+
+	RefreshTreeView();
+}
+
+void SMVVMManageViewModelsWidget::RefreshTreeView(bool bForceExpandAll /* = false */)
+{
+	using UE::MVVM::Private::FMVVMViewModelTreeNode;
+	RebuildTreeViewSource(NodeTreeViewSource, CurrentViewTreeRootNode.ToSharedRef());
+	TreeViewWidget->RequestListRefresh();
+	if ((CurrentViewTreeRootNode == ClassTreeRootNode && ClassTreeExpansionStateMap.IsEmpty())
+		|| (CurrentViewTreeRootNode == TagTreeRootNode && TagTreeExpansionStateMap.IsEmpty())
+		|| CurrentViewTreeRootNode == ClassListRootNode)
+	{
+		bForceExpandAll = true;
+	}
+	if (bForceExpandAll)
+	{
+		if (!NodeTreeViewSource.IsEmpty())
+		{
+			for (TSharedPtr<FMVVMViewModelTreeNode>& Child : NodeTreeViewSource)
+			{
+				SetAllExpansionStates_Helper(Child, true);
+			}
+		}
+	}
+	else
+	{
+		RestoreExpansionStatesInTree();
+	}
+}
+
+void SMVVMManageViewModelsWidget::RebuildTreeViewSource(TArray<TSharedPtr<FMVVMViewModelTreeNode>>& ViewSource, TSharedRef<FMVVMViewModelTreeNode> RootNode)
+{
+	using UE::MVVM::Private::FMVVMViewModelTreeNode;
+	ViewSource.Empty();
+
+	TSet<TSharedRef<FMVVMViewModelTreeNode>> NodesToDiscard = RootNode->GetNodeBranchesToDiscard();
+
+	if (!NodesToDiscard.Contains(RootNode))
+	{
+		for (TSharedRef<FMVVMViewModelTreeNode>& ChildNode : RootNode->GetChildrenList())
+		{
+			if (!NodesToDiscard.Contains(ChildNode))
+			{
+				TSharedRef<FMVVMViewModelTreeNode> ChildClone = MakeShared<FMVVMViewModelTreeNode>(*ChildNode);
+				ViewSource.Add(ChildClone);
+				RebuildTreeViewSource_Helper(ChildNode, ChildClone, NodesToDiscard);
+			}
+
+		}
+
+		ViewSource.Sort([](const TSharedPtr<FMVVMViewModelTreeNode>& A, const TSharedPtr<FMVVMViewModelTreeNode>& B)
+			{
+				return A->GetDisplayName() < B->GetDisplayName();
+			});
+	}
+}
+
+void SMVVMManageViewModelsWidget::RebuildTreeViewSource_Helper(TSharedRef<FMVVMViewModelTreeNode>& Parent, TSharedRef<FMVVMViewModelTreeNode>& CloneParent, const TSet<TSharedRef<FMVVMViewModelTreeNode>>& NodesToDiscard)
+{
+	for (TSharedRef<FMVVMViewModelTreeNode>& Child : Parent->GetChildrenList())
+	{
+		if (NodesToDiscard.Contains(Child))
+		{
+			continue;
+		}
+
+		TSharedRef<FMVVMViewModelTreeNode> CloneChild = MakeShared<FMVVMViewModelTreeNode>(*Child);
+		CloneParent->AddChild(CloneChild);
+
+		RebuildTreeViewSource_Helper(Child, CloneChild, NodesToDiscard);
+	}
+
+	CloneParent->SortChildren();
+}
+
+void SMVVMManageViewModelsWidget::HandleGetChildrenForTreeView(TSharedPtr<FMVVMViewModelTreeNode> InParent, TArray<TSharedPtr<FMVVMViewModelTreeNode>>& OutChildren)
+{
+	OutChildren.Append(InParent->GetChildrenList());
+}
+
+void SMVVMManageViewModelsWidget::SetAllExpansionStates_Helper(TSharedPtr<FMVVMViewModelTreeNode> InNode, bool bInExpansionState)
+{
+	check(InNode.IsValid());
+	TreeViewWidget->SetItemExpansion(InNode, bInExpansionState);
+
+	if (CurrentViewTreeRootNode == ClassTreeRootNode)
+	{
+		ClassTreeExpansionStateMap.Add(InNode->GetDisplayName(), true);
+	}
+	else if (CurrentViewTreeRootNode == TagTreeRootNode)
+	{
+		TagTreeExpansionStateMap.Add(InNode->GetDisplayName(), true);
+	}
+
+	// Recursively go through the children.
+	for (int32 ChildIndex = 0; ChildIndex < InNode->GetChildrenList().Num(); ChildIndex++)
+	{
+		SetAllExpansionStates_Helper(InNode->GetChildrenList()[ChildIndex], bInExpansionState);
+	}
+}
+
+void SMVVMManageViewModelsWidget::RestoreExpansionStatesInTree()
+{
+	if (CurrentViewTreeRootNode == ClassTreeRootNode)
+	{
+		RestoreExpansionStatesInTree_Helper(NodeTreeViewSource[0].ToSharedRef(), ClassTreeExpansionStateMap, this);
+	}
+	else if (CurrentViewTreeRootNode == TagTreeRootNode)
+	{
+		RestoreExpansionStatesInTree_Helper(NodeTreeViewSource[0].ToSharedRef(), TagTreeExpansionStateMap, this);
+	}
+}
+
+bool SMVVMManageViewModelsWidget::RestoreExpansionStatesInTree_Helper(TSharedRef<FMVVMViewModelTreeNode> InNode, TMap<FString, bool>& ExpansionStateMap, SMVVMManageViewModelsWidget* Self)
+{
+	using UE::MVVM::Private::FMVVMViewModelTreeNode;
+	bool bForceParentOpen = false;
+	for (TSharedRef<FMVVMViewModelTreeNode>& ChildNode : InNode->GetChildrenList())
+	{
+		bForceParentOpen |= RestoreExpansionStatesInTree_Helper(ChildNode, ExpansionStateMap, Self);
+	}
+
+	if (InNode->IsPassingSearchFilter() && !Self->bSearchStringIsEmpty)
+	{
+		bForceParentOpen |= true;
+	}
+
+	if (bForceParentOpen)
+	{
+		ExpansionStateMap[InNode->GetDisplayName()] = true;
+	}
+
+	Self->TreeViewWidget->SetItemExpansion(InNode, ExpansionStateMap.FindOrAdd(InNode->GetDisplayName(), false));
+
+	return bForceParentOpen;
+}
+
+TSharedRef<ITableRow> SMVVMManageViewModelsWidget::HandleGenerateRowForTreeView(TSharedPtr<FMVVMViewModelTreeNode> Item, const TSharedRef<STableViewBase>& OwnerTable)
+{
+	using UE::MVVM::Private::FMVVMViewModelTreeNode;
+	typedef STableRow<TSharedPtr<FMVVMViewModelTreeNode>> RowType;
+
+	bool bShowAddIcon = Item->IsAddable();
+
+	TSharedRef<RowType> NewRow = SNew(RowType, OwnerTable);
+
+	TSharedRef<SHorizontalBox> HBox = SNew(SHorizontalBox);
+
+	if (bShowAddIcon)
+	{
+		HBox->AddSlot()
+		.AutoWidth()
+		[
+			SNew(SSimpleButton)
+			.ToolTipText(LOCTEXT("AddViewModelButtonToolTip", "Add ViewModel"))
+			.OnClicked(this, &SMVVMManageViewModelsWidget::HandleClicked_AddViewModel, Item)
+			.Icon(FSlateIcon(FMVVMEditorStyle::Get().GetStyleSetName(), "ViewModelSelection.AddIcon").GetIcon())
+		];
+	}
+
+	HBox->AddSlot()
+	.AutoWidth()
+	.HAlign(HAlign_Right)
+	.VAlign(VAlign_Center)
+	[
+		Item->GetIcon()
+	];
+
+	HBox->AddSlot()
+	.Padding(4.0f)
+	[
+		SNew(STextBlock).Text(Item->GetDisplayNameText())
+	];
+
+	NewRow->SetContent(HBox);
+
+	return NewRow;
+}
+
+void SMVVMManageViewModelsWidget::HandleTreeViewSelectionChanged(TSharedPtr<FMVVMViewModelTreeNode> Item, ESelectInfo::Type SelectInfo)
+{
+	TSubclassOf<UMVVMViewModelBase> InClass = nullptr;
+	if (Item)
+	{
+		InClass = GetClassFromNode(Item.ToSharedRef());
+	}
+	ViewModelBindingListWidget->SetViewModel(InClass);
+}
+
+void SMVVMManageViewModelsWidget::HandleTreeViewExpansionChanged(TSharedPtr<FMVVMViewModelTreeNode> Item, bool bExpanded)
+{
+	if (CurrentViewTreeRootNode == ClassTreeRootNode)
+	{
+		ClassTreeExpansionStateMap.Add(Item->GetDisplayName(), bExpanded);
+	}
+	else if (CurrentViewTreeRootNode == TagTreeRootNode)
+	{
+		TagTreeExpansionStateMap.Add(Item->GetDisplayName(), bExpanded);
+	}
+}
+
+void SMVVMManageViewModelsWidget::HandleTreeRowDoubleClicked(TSharedPtr<FMVVMViewModelTreeNode> Item)
+{
+	HandleClicked_AddViewModel(Item);
+}
+
+FReply SMVVMManageViewModelsWidget::HandleClicked_AddViewModel(TSharedPtr<FMVVMViewModelTreeNode> InNode)
+{
+	ViewModelContextListWidget->AddViewModelContext(InNode->GetClassData()->Class);
+
+	return FReply::Handled();
+}
+
+FReply SMVVMManageViewModelsWidget::HandleClicked_Finish()
+{
+	TArray<FMVVMBlueprintViewModelContext> AddedContexts = ViewModelContextListWidget->GetViewModelContexts();
+	OnViewModelContextsPicked.ExecuteIfBound(AddedContexts);
+
+	if (WeakParentWindow.IsValid())
+	{
+		WeakParentWindow.Pin()->RequestDestroyWindow();
+	}
+
+	return FReply::Handled();
+}
+
+FReply SMVVMManageViewModelsWidget::HandleClicked_Cancel()
+{
+	if (WeakParentWindow.IsValid())
+	{
+		WeakParentWindow.Pin()->RequestDestroyWindow();
+	}
+	return FReply::Handled();
+}
+
+TSubclassOf<UMVVMViewModelBase> SMVVMManageViewModelsWidget::GetClassFromNode(TSharedRef<FMVVMViewModelTreeNode> ClassNode)
+{
+	using UE::MVVM::Private::FMVVMViewModelTreeNode;
+	if (FMVVMViewModelTreeNode::FClassData* ClassInfo = ClassNode->GetClassData())
+	{
+		return ClassInfo->Class;
+	}
+
+	return nullptr;
+}
+
+#undef LOCTEXT_NAMESPACE
diff --git a/Engine/Plugins/Runtime/ModelViewViewModel/Source/ModelViewViewModelEditor/Private/Widgets/SMVVMManageViewModelsWidget.h b/Engine/Plugins/Runtime/ModelViewViewModel/Source/ModelViewViewModelEditor/Private/Widgets/SMVVMManageViewModelsWidget.h
new file mode 100644
--- /dev/null
+++ b/Engine/Plugins/Runtime/ModelViewViewModel/Source/ModelViewViewModelEditor/Private/Widgets/SMVVMManageViewModelsWidget.h
@@ -0,0 +1,112 @@
+// Copyright Epic Games, Inc. All Rights Reserved.
+
+#pragma once
+
+#include "CoreMinimal.h"
+
+#include "Widgets/SCompoundWidget.h"
+#include "Input/Reply.h"
+#include "MVVMBlueprintViewBinding.h"
+#include "MVVMBlueprintViewModelContext.h"
+
+class ITableRow;
+template<typename T>
+class SListView;
+class SMVVMViewModelContextListWidget;
+class SMVVMViewModelBindingListWidget;
+template<typename T>
+class STreeView;
+class STableViewBase;
+class UMVVMViewModelBase;
+class UUserWidget;
+namespace UE::MVVM
+{
+struct FMVVMFieldVariant;
+namespace Private
+{
+class FMVVMViewModelTreeNode;
+}
+}
+
+DECLARE_DELEGATE_OneParam(FOnViewModelContextsPicked, TArray<FMVVMBlueprintViewModelContext>);
+
+class SMVVMManageViewModelsWidget : public SCompoundWidget
+{
+public:
+	SLATE_BEGIN_ARGS(SMVVMManageViewModelsWidget) {}
+		SLATE_NAMED_SLOT(FArguments, ButtonsPanel)
+		SLATE_ARGUMENT(bool, bUseDefaultButtons)
+		SLATE_ARGUMENT(TWeakPtr<SWindow>, ParentWindow)
+		SLATE_ARGUMENT(FOnViewModelContextsPicked, OnViewModelContextsPickedDelegate)
+		SLATE_ARGUMENT(TSharedPtr<SMVVMViewModelContextListWidget>, ManageViewModelsWidget)
+	SLATE_END_ARGS()
+
+	enum class EViewMode : uint8
+	{
+		List,
+		Hierarchy,
+		Tag
+	};
+
+public:
+	void Construct(const FArguments& InArgs);
+
+private:
+	typedef UE::MVVM::Private::FMVVMViewModelTreeNode FMVVMViewModelTreeNode;
+
+	void PopulateViewModelsTree();
+
+	void SetViewMode(EViewMode Mode);
+
+	TSharedRef<SWidget> GenerateSearchBar();
+
+	void UpdateSearchString(const FString& SearchString);
+
+	void RefreshTreeView(bool bForceExpandAll = false);
+	void RebuildTreeViewSource(TArray<TSharedPtr<FMVVMViewModelTreeNode>>& ViewSource, TSharedRef<FMVVMViewModelTreeNode> RootNode);
+	static void RebuildTreeViewSource_Helper(TSharedRef<FMVVMViewModelTreeNode>& Parent, TSharedRef<FMVVMViewModelTreeNode>& CloneParent, const TSet<TSharedRef<FMVVMViewModelTreeNode>>& NodesToDiscard);
+
+	void SetAllExpansionStates_Helper(TSharedPtr<FMVVMViewModelTreeNode> InNode, bool bInExpansionState);
+	void RestoreExpansionStatesInTree();
+	static bool RestoreExpansionStatesInTree_Helper(TSharedRef<FMVVMViewModelTreeNode> InNode, TMap<FString, bool>& ExpansionStateMap, SMVVMManageViewModelsWidget* Self);
+
+	TSharedRef<SWidget> HandleGetMenuContent_ViewSettings();
+	void HandleGetChildrenForTreeView(TSharedPtr<FMVVMViewModelTreeNode> InParent, TArray<TSharedPtr<FMVVMViewModelTreeNode>>& OutChildren);
+	TSharedRef<ITableRow> HandleGenerateRowForTreeView(TSharedPtr<FMVVMViewModelTreeNode> Item, const TSharedRef<STableViewBase>& OwnerTable);
+	void HandleTreeViewExpansionChanged(TSharedPtr<FMVVMViewModelTreeNode> Item, bool bExpanded);
+	void HandleTreeRowDoubleClicked(TSharedPtr<FMVVMViewModelTreeNode> Item);
+	void HandleTreeViewSelectionChanged(TSharedPtr<FMVVMViewModelTreeNode> Item, ESelectInfo::Type SelectInfo);
+
+	FReply HandleClicked_AddViewModel(TSharedPtr<FMVVMViewModelTreeNode> InNode);
+	FReply HandleClicked_Finish();
+	FReply HandleClicked_Cancel();
+
+	TSubclassOf<UMVVMViewModelBase> GetClassFromNode(TSharedRef<FMVVMViewModelTreeNode> ClassNode);
+
+private:
+	TSharedPtr<FMVVMViewModelTreeNode> ClassListRootNode;
+	TSharedPtr<FMVVMViewModelTreeNode> ClassTreeRootNode;
+	TSharedPtr<FMVVMViewModelTreeNode> TagTreeRootNode;
+
+	TSharedPtr<FMVVMViewModelTreeNode> CurrentViewTreeRootNode;
+
+	TArray<TSharedPtr<FMVVMViewModelTreeNode>> NodeTreeViewSource;
+
+	TMap<FString, bool> ClassTreeExpansionStateMap;
+	TMap<FString, bool> TagTreeExpansionStateMap;
+
+	TSharedPtr<STreeView<TSharedPtr<FMVVMViewModelTreeNode>>> TreeViewWidget;
+
+	EViewMode CurrentViewMode = EViewMode::List;
+
+	FOnViewModelContextsPicked OnViewModelContextsPicked;
+
+	TWeakPtr<SWindow> WeakParentWindow;
+
+	TSharedPtr<SMVVMViewModelBindingListWidget> ViewModelBindingListWidget;
+
+	bool bSearchStringIsEmpty = true;
+	UUserWidget* OuterWidget = nullptr;
+
+	TSharedPtr<SMVVMViewModelContextListWidget> ViewModelContextListWidget;
+};
diff --git a/Engine/Plugins/Runtime/ModelViewViewModel/Source/ModelViewViewModelEditor/Private/Widgets/SMVVMPropertyPath.cpp b/Engine/Plugins/Runtime/ModelViewViewModel/Source/ModelViewViewModelEditor/Private/Widgets/SMVVMPropertyPath.cpp
new file mode 100644
--- /dev/null
+++ b/Engine/Plugins/Runtime/ModelViewViewModel/Source/ModelViewViewModelEditor/Private/Widgets/SMVVMPropertyPath.cpp
@@ -0,0 +1,95 @@
+// Copyright Epic Games, Inc. All Rights Reserved.
+
+#include "Widgets/SMVVMPropertyPath.h"
+
+#include "MVVMBlueprintView.h"
+#include "MVVMWidgetBlueprintExtension_View.h"
+#include "Widgets/SMVVMFieldIcon.h"
+#include "Styling/MVVMEditorStyle.h"
+#include "Widgets/SBoxPanel.h"
+#include "Widgets/Images/SImage.h"
+#include "Widgets/Layout/SBorder.h"
+#include "Widgets/Layout/SSeparator.h"
+#include "Widgets/Text/STextBlock.h"
+
+#define LOCTEXT_NAMESPACE "SMVVMPropertyPath"
+
+void SMVVMPropertyPathBase::Construct(const FArguments& InArgs, const UWidgetBlueprint* InWidgetBlueprint)
+{
+	WidgetBlueprint = InWidgetBlueprint;
+
+	const UE::MVVM::IFieldPathHelper& PathHelper = GetPathHelper();
+	SelectedSource = PathHelper.GetSelectedSource();
+	SelectedField = PathHelper.GetSelectedField();
+
+	auto GetContextText = [this]()
+	{
+		if (SelectedSource.IsSet())
+		{
+			return SelectedSource.GetValue().DisplayName;
+		}
+		return FText::GetEmpty();
+	};
+
+	auto GetPathText = [this]()
+	{
+		if (SelectedField.IsProperty())
+		{
+			return SelectedField.GetProperty()->GetDisplayNameText();
+		}
+
+		if (SelectedField.IsFunction())
+		{
+			return SelectedField.GetFunction()->GetDisplayNameText();
+		}
+
+		return FText::GetEmpty();
+	};
+
+	ChildSlot
+	[
+		SNew(SHorizontalBox)
+		+ SHorizontalBox::Slot()
+		.VAlign(VAlign_Center)
+		.AutoWidth()
+		[
+			SNew(STextBlock)
+			.TextStyle(FMVVMEditorStyle::Get(), "PropertyPath.ContextText")
+			.Text_Lambda(GetContextText)
+		]
+		+ SHorizontalBox::Slot()
+		.Padding(FMargin(5, 0, 0, 0))
+		.VAlign(VAlign_Center)
+		.AutoWidth()
+		[
+			SNew(SMVVMFieldIcon)
+			.Field(SelectedField)
+		]
+		+ SHorizontalBox::Slot()
+		.Padding(FMargin(2, 0, 5, 0))
+		.VAlign(VAlign_Center)
+		.AutoWidth()
+		[
+			SNew(STextBlock)
+			.TextStyle(FMVVMEditorStyle::Get(), "PropertyPath.ContextText")
+			.Text_Lambda(GetPathText)
+		]
+	];
+}
+
+void SMVVMWidgetPropertyPath::Construct(const FArguments& InArgs, const UWidgetBlueprint* InWidgetBlueprint)
+{
+	PathHelper = UE::MVVM::FWidgetFieldPathHelper(InArgs._WidgetPath, InWidgetBlueprint);
+
+	SMVVMPropertyPathBase::Construct(SMVVMPropertyPathBase::FArguments(), InWidgetBlueprint);
+}
+
+void SMVVMViewModelPropertyPath::Construct(const FArguments& InArgs, const UWidgetBlueprint* InWidgetBlueprint)
+{
+	PathHelper = UE::MVVM::FViewModelFieldPathHelper(InArgs._ViewModelPath, InWidgetBlueprint);
+
+	SMVVMPropertyPathBase::Construct(SMVVMPropertyPathBase::FArguments(), InWidgetBlueprint);
+}
+
+
+#undef LOCTEXT_NAMESPACE
\ No newline at end of file
diff --git a/Engine/Plugins/Runtime/ModelViewViewModel/Source/ModelViewViewModelEditor/Private/Widgets/SMVVMPropertyPath.h b/Engine/Plugins/Runtime/ModelViewViewModel/Source/ModelViewViewModelEditor/Private/Widgets/SMVVMPropertyPath.h
new file mode 100644
--- /dev/null
+++ b/Engine/Plugins/Runtime/ModelViewViewModel/Source/ModelViewViewModelEditor/Private/Widgets/SMVVMPropertyPath.h
@@ -0,0 +1,67 @@
+// Copyright Epic Games, Inc. All Rights Reserved.
+
+#pragma once
+
+#include "Widgets/SCompoundWidget.h"
+#include "MVVMBlueprintView.h"
+#include "Widgets/SMVVMFieldIcon.h"
+#include "MVVMPropertyPathHelpers.h"
+#include "Types/MVVMFieldVariant.h"
+
+class SMVVMPropertyPathBase : public SCompoundWidget
+{
+public:
+	SLATE_BEGIN_ARGS(SMVVMPropertyPathBase)
+	{}
+	SLATE_END_ARGS()
+
+	void Construct(const FArguments& InArgs, const UWidgetBlueprint* WidgetBlueprint);
+	virtual void Tick(const FGeometry& AllottedGeometry, const double InCurrentTime, const float InDeltaTime) override;
+
+private:
+	virtual UE::MVVM::IFieldPathHelper& GetPathHelper() = 0;
+
+private:
+	const UWidgetBlueprint* WidgetBlueprint = nullptr;
+	UE::MVVM::FMVVMConstFieldVariant SelectedField;
+	TOptional<UE::MVVM::FBindingSource> SelectedSource;
+};
+
+
+class SMVVMWidgetPropertyPath : public SMVVMPropertyPathBase
+{
+public:
+	SLATE_BEGIN_ARGS(SMVVMWidgetPropertyPath) {}
+		SLATE_ATTRIBUTE(FMVVMWidgetPropertyPath*, WidgetPath)
+	SLATE_END_ARGS()
+
+	void Construct(const FArguments& InArgs, const UWidgetBlueprint* WidgetBlueprint);
+
+private:
+	virtual UE::MVVM::IFieldPathHelper& GetPathHelper() override
+	{ 
+		return PathHelper;
+	}
+
+private:
+	UE::MVVM::FWidgetFieldPathHelper PathHelper;
+};
+
+class SMVVMViewModelPropertyPath : public SMVVMPropertyPathBase
+{
+public:
+	SLATE_BEGIN_ARGS(SMVVMViewModelPropertyPath) {}
+		SLATE_ATTRIBUTE(FMVVMViewModelPropertyPath*, ViewModelPath)
+	SLATE_END_ARGS()
+
+	void Construct(const FArguments& InArgs, const UWidgetBlueprint* WidgetBlueprint);
+
+private:
+	virtual UE::MVVM::IFieldPathHelper& GetPathHelper() override
+	{
+		return PathHelper;
+	}
+
+private:
+	UE::MVVM::FViewModelFieldPathHelper PathHelper;
+};
diff --git a/Engine/Plugins/Runtime/ModelViewViewModel/Source/ModelViewViewModelEditor/Private/Widgets/SMVVMSourceEntry.cpp b/Engine/Plugins/Runtime/ModelViewViewModel/Source/ModelViewViewModelEditor/Private/Widgets/SMVVMSourceEntry.cpp
new file mode 100644
--- /dev/null
+++ b/Engine/Plugins/Runtime/ModelViewViewModel/Source/ModelViewViewModelEditor/Private/Widgets/SMVVMSourceEntry.cpp
@@ -0,0 +1,79 @@
+// Copyright Epic Games, Inc. All Rights Reserved.
+ 
+#include "Widgets/SMVVMSourceEntry.h"
+
+#include "Styling/SlateIconFinder.h"
+#include "Widgets/SBoxPanel.h"
+#include "Widgets/Images/SImage.h"
+#include "Widgets/Text/STextBlock.h"
+
+#define LOCTEXT_NAMESPACE "SMVVMSource"
+
+static const FSlateBrush* GetSourceIcon(const TOptional<UE::MVVM::FBindingSource>& Source)
+{
+	if (!Source.IsSet())
+	{
+		return nullptr;
+	}
+
+	return FSlateIconFinder::FindIconBrushForClass(Source.GetValue().Class);
+}
+
+static FLinearColor GetSourceColor(const TOptional<UE::MVVM::FBindingSource>& Source)
+{
+	if (!Source.IsSet() || Source.GetValue().Class == nullptr || Source.GetValue().Class->GetName().IsEmpty())
+	{
+		return FLinearColor::White;
+	}
+
+	uint32 Hash = GetTypeHash(Source.GetValue().Class->GetName());
+	FLinearColor Color = FLinearColor::White;
+	Color.R = ((Hash * 1 % 96) + 32) / 256.f;
+	Color.G = ((Hash * 2 % 96) + 32) / 256.f;
+	Color.B = ((Hash * 3 % 96) + 32) / 256.f;
+	return Color;
+}
+
+static FText GetSourceDisplayName(const TOptional<UE::MVVM::FBindingSource>& Source)
+{
+	if (!Source.IsSet())
+	{
+		return LOCTEXT("None", "<None>");
+	}
+
+	return Source.GetValue().DisplayName;
+}
+
+void SMVVMSourceEntry::Construct(const FArguments& InArgs)
+{
+	ChildSlot
+	[
+		SNew(SHorizontalBox)
+		+ SHorizontalBox::Slot()
+		.HAlign(HAlign_Left)
+		.VAlign(VAlign_Center)
+		.AutoWidth()
+		[
+			SAssignNew(Image, SImage)
+		]
+		+ SHorizontalBox::Slot()
+		.Padding(4, 0, 0, 0)
+		.HAlign(HAlign_Fill)
+		.VAlign(VAlign_Center)
+		[
+			SAssignNew(Label, STextBlock)
+			.TextStyle(InArgs._TextStyle)
+		]
+	];
+
+	RefreshSource(InArgs._Source);
+}
+
+void SMVVMSourceEntry::RefreshSource(const TOptional<UE::MVVM::FBindingSource>& Source)
+{
+	Image->SetImage(GetSourceIcon(Source));
+	Image->SetColorAndOpacity(GetSourceColor(Source));
+	Label->SetText(GetSourceDisplayName(Source));
+}
+
+#undef LOCTEXT_NAMESPACE
\ No newline at end of file
diff --git a/Engine/Plugins/Runtime/ModelViewViewModel/Source/ModelViewViewModelEditor/Private/Widgets/SMVVMSourceEntry.h b/Engine/Plugins/Runtime/ModelViewViewModel/Source/ModelViewViewModelEditor/Private/Widgets/SMVVMSourceEntry.h
new file mode 100644
--- /dev/null
+++ b/Engine/Plugins/Runtime/ModelViewViewModel/Source/ModelViewViewModelEditor/Private/Widgets/SMVVMSourceEntry.h
@@ -0,0 +1,34 @@
+// Copyright Epic Games, Inc. All Rights Reserved.
+
+#pragma once
+
+#include "MVVMPropertyPathHelpers.h"
+#include "Styling/CoreStyle.h"
+#include "Styling/ISlateStyle.h"
+#include "Styling/SlateTypes.h"
+#include "Styling/SlateWidgetStyleAsset.h"
+#include "Types/MVVMFieldVariant.h"
+#include "Widgets/DeclarativeSyntaxSupport.h"
+#include "Widgets/SCompoundWidget.h"
+
+class SImage;
+class STextBlock;
+
+class SMVVMSourceEntry : public SCompoundWidget
+{
+public:
+	SLATE_BEGIN_ARGS(SMVVMSourceEntry) :
+		_TextStyle(&FCoreStyle::Get().GetWidgetStyle<FTextBlockStyle>("NormalText"))
+		{
+		}
+		SLATE_STYLE_ARGUMENT(FTextBlockStyle, TextStyle)
+		SLATE_ARGUMENT(TOptional<UE::MVVM::FBindingSource>, Source)
+	SLATE_END_ARGS()
+
+	void Construct(const FArguments& Args);
+	void RefreshSource(const TOptional<UE::MVVM::FBindingSource>& Source);
+
+private:
+	TSharedPtr<STextBlock> Label;
+	TSharedPtr<SImage> Image;
+};
\ No newline at end of file
diff --git a/Engine/Plugins/Runtime/ModelViewViewModel/Source/ModelViewViewModelEditor/Private/Widgets/SMVVMSourceSelector.cpp b/Engine/Plugins/Runtime/ModelViewViewModel/Source/ModelViewViewModelEditor/Private/Widgets/SMVVMSourceSelector.cpp
new file mode 100644
--- /dev/null
+++ b/Engine/Plugins/Runtime/ModelViewViewModel/Source/ModelViewViewModelEditor/Private/Widgets/SMVVMSourceSelector.cpp
@@ -0,0 +1,115 @@
+// Copyright Epic Games, Inc. All Rights Reserved.
+
+#include "SMVVMSourceSelector.h"
+
+#include "Algo/Transform.h"
+#include "Widgets/SMVVMFieldIcon.h"
+#include "Widgets/SMVVMSourceEntry.h"
+
+#define LOCTEXT_NAMESPACE "MVVMSourceSelector"
+
+using namespace UE::MVVM;
+
+static FBindingSource GetSelectedSource(const TArray<IFieldPathHelper*>& Helpers)
+{
+	bool bFirst = true;
+	TOptional<FBindingSource> Result;
+
+	for (const IFieldPathHelper* Helper : Helpers)
+	{
+		TOptional<FBindingSource> ThisSelection = Helper->GetSelectedSource();
+		if (bFirst)
+		{
+			Result = ThisSelection;
+		}
+		else if (Result != ThisSelection)
+		{
+			// all don't share the same value
+			Result.Reset();
+			break;
+		}
+	}
+
+	return Result.Get(FBindingSource());
+}
+
+void SMVVMSourceSelector::Construct(const FArguments& Args)
+{
+	PathHelpers = Args._PathHelpers;
+	OnSelectionChanged = Args._OnSelectionChanged;
+	TextStyle = Args._TextStyle;
+
+	Refresh();
+
+	ChildSlot
+	[
+		SAssignNew(SourceComboBox, SComboBox<FBindingSource>)
+		.OptionsSource(&AvailableSources)
+		.OnGenerateWidget_Lambda([this](FBindingSource Source)
+		{
+			return SNew(SMVVMSourceEntry)
+				.TextStyle(TextStyle)
+				.Source(Source);
+		})
+		.OnSelectionChanged(this, &SMVVMSourceSelector::OnComboBoxSelectionChanged)
+		.InitiallySelectedItem(SelectedSource)
+		[
+			SAssignNew(SelectedSourceWidget, SMVVMSourceEntry)
+			.TextStyle(TextStyle)
+			.Source(SelectedSource)
+		]
+	];
+}
+
+void SMVVMSourceSelector::OnComboBoxSelectionChanged(FBindingSource Selected, ESelectInfo::Type SelectionType)
+{
+	SelectedSource = Selected;
+
+	TOptional<FBindingSource> OptionalSource = SelectedSource;
+
+	SelectedSourceWidget->RefreshSource(OptionalSource);
+
+	OnSelectionChanged.ExecuteIfBound(OptionalSource);
+}
+
+void SMVVMSourceSelector::Refresh()
+{
+	AvailableSources.Reset();
+	SelectedSource = FBindingSource();
+
+	TArray<IFieldPathHelper*> Helpers = PathHelpers.Get(TArray<IFieldPathHelper*>());
+
+	bool bFirst = true;
+
+	TSet<FBindingSource> SourceIntersection;
+	TOptional<FBindingSource> NewSelection;
+
+	for (const IFieldPathHelper* Helper : Helpers)
+	{
+		// get available sources
+		TSet<FBindingSource> Sources;
+		Helper->GetAvailableSources(Sources);
+
+		if (bFirst)
+		{
+			bFirst = false; 
+
+			SourceIntersection.Append(Sources);
+		}
+		else
+		{
+			SourceIntersection.Intersect(Sources);
+		}
+	}
+
+	AvailableSources = SourceIntersection.Array();
+	SelectedSource = GetSelectedSource(Helpers);
+
+	if (SourceComboBox.IsValid())
+	{
+		SourceComboBox->RefreshOptions();
+		SourceComboBox->SetSelectedItem(SelectedSource);
+	}
+}
+
+#undef LOCTEXT_NAMESPACE
diff --git a/Engine/Plugins/Runtime/ModelViewViewModel/Source/ModelViewViewModelEditor/Private/Widgets/SMVVMSourceSelector.h b/Engine/Plugins/Runtime/ModelViewViewModel/Source/ModelViewViewModelEditor/Private/Widgets/SMVVMSourceSelector.h
new file mode 100644
--- /dev/null
+++ b/Engine/Plugins/Runtime/ModelViewViewModel/Source/ModelViewViewModelEditor/Private/Widgets/SMVVMSourceSelector.h
@@ -0,0 +1,43 @@
+// Copyright Epic Games, Inc. All Rights Reserved.
+
+#pragma once
+
+#include "CoreMinimal.h"
+#include "MVVMPropertyPathHelpers.h"
+#include "Types/MVVMBindingSource.h"
+#include "Widgets/DeclarativeSyntaxSupport.h"
+#include "Widgets/SCompoundWidget.h"
+#include "Widgets/Input/SComboBox.h"
+
+class SMVVMSourceEntry;
+
+class SMVVMSourceSelector : public SCompoundWidget
+{
+public:
+	DECLARE_DELEGATE_OneParam(FSelectionChanged, TOptional<UE::MVVM::FBindingSource>);
+
+	SLATE_BEGIN_ARGS(SMVVMSourceSelector) :
+		_TextStyle( &FCoreStyle::Get().GetWidgetStyle<FTextBlockStyle>( "NormalText" ) )
+		{
+		}
+		SLATE_STYLE_ARGUMENT(FTextBlockStyle, TextStyle)
+		SLATE_ATTRIBUTE(TArray<UE::MVVM::IFieldPathHelper*>, PathHelpers)
+		SLATE_EVENT(FSelectionChanged, OnSelectionChanged)
+	SLATE_END_ARGS()
+
+	void Construct(const FArguments& InArgs);
+	void Refresh();
+
+private:
+	void OnComboBoxSelectionChanged(UE::MVVM::FBindingSource Selected, ESelectInfo::Type SelectionType);
+
+private:
+	TAttribute<TArray<UE::MVVM::IFieldPathHelper*>> PathHelpers;
+	FSelectionChanged OnSelectionChanged;
+	TSharedPtr<SComboBox<UE::MVVM::FBindingSource>> SourceComboBox;
+	TArray<UE::MVVM::FBindingSource> AvailableSources;
+	UE::MVVM::FBindingSource SelectedSource;
+	TSharedPtr<SMVVMSourceEntry> SelectedSourceWidget;
+	const FTextBlockStyle* TextStyle = nullptr;
+};
+
diff --git a/Engine/Plugins/Runtime/ModelViewViewModel/Source/ModelViewViewModelEditor/Private/Widgets/SMVVMViewBindingListView.cpp b/Engine/Plugins/Runtime/ModelViewViewModel/Source/ModelViewViewModelEditor/Private/Widgets/SMVVMViewBindingListView.cpp
new file mode 100644
--- /dev/null
+++ b/Engine/Plugins/Runtime/ModelViewViewModel/Source/ModelViewViewModelEditor/Private/Widgets/SMVVMViewBindingListView.cpp
@@ -0,0 +1,802 @@
+// Copyright Epic Games, Inc. All Rights Reserved.
+
+#include "Widgets/SMVVMViewBindingListView.h"
+
+#include "Bindings/MVVMBindingHelper.h"
+#include "Dialog/SCustomDialog.h"
+#include "Editor.h"
+#include "EditorStyleSet.h"
+#include "Engine/Engine.h"
+#include "Features/IModularFeatures.h"
+#include "Framework/MultiBox/MultiBoxBuilder.h"
+#include "IPropertyAccessEditor.h"
+#include "MVVMBlueprintView.h"
+#include "MVVMBlueprintViewBinding.h"
+#include "MVVMEditorSubsystem.h"
+#include "MVVMPropertyPathHelpers.h"
+#include "MVVMSubsystem.h"
+#include "MVVMWidgetBlueprintExtension_View.h"
+#include "SEnumCombo.h"
+#include "SSimpleButton.h" 
+#include "Styling/StyleColors.h"
+#include "UMGStyle.h"
+#include "Widgets/SMVVMConversionPath.h"
+#include "Widgets/SMVVMFieldSelector.h"
+#include "Widgets/SMVVMPropertyPath.h"
+#include "Widgets/SMVVMSourceSelector.h"
+#include "Widgets/SMVVMViewBindingPanel.h"
+#include "Widgets/SNullWidget.h"
+#include "Widgets/Images/SImage.h"
+#include "Widgets/Input/SCheckBox.h"
+#include "Widgets/Input/SComboBox.h"
+#include "Widgets/Input/SEditableTextBox.h"
+#include "Widgets/Text/STextBlock.h"
+#include "Widgets/Views/SListView.h"
+#include "Widgets/Views/STableViewBase.h"
+
+#define LOCTEXT_NAMESPACE "BindingListView"
+
+struct FMVVMViewBindingListEntry
+{
+	FMVVMViewBindingListEntry(int32 InIndex) : Index(InIndex) {}
+	int32 Index;
+};
+
+class SMVVMViewBindingListEntryRow : public SMultiColumnTableRow<FMVVMViewBindingListEntryPtr>
+{
+public:
+	static FName EnabledColumnName;
+	static FName ErrorColumnName;
+	static FName ViewModelColumnName;
+	static FName ViewModelPropertyColumnName;
+	static FName ModeColumnName;
+	static FName WidgetColumnName;
+	static FName WidgetPropertyColumnName;
+	static FName UpdateColumnName;
+	static FName ConversionFunctionColumnName;
+	static FName DropDownOptionsColumnName;
+
+public:
+	SLATE_BEGIN_ARGS(SMVVMViewBindingListEntryRow) {}
+	SLATE_END_ARGS()
+
+	void Construct(const FArguments& Args, const TSharedRef<STableViewBase>& OwnerTableView, const FMVVMViewBindingListEntryPtr& InEntry, UMVVMBlueprintView* InBlueprintView, UWidgetBlueprint* InWidgetBlueprint)
+	{
+		Entry = InEntry;
+		BlueprintView = InBlueprintView;
+		WidgetBlueprint = InWidgetBlueprint;
+
+		OnBlueprintChangedHandle = WidgetBlueprint->OnChanged().AddSP(this, &SMVVMViewBindingListEntryRow::HandleBlueprintChanged);
+
+		FMVVMBlueprintViewBinding* ViewModelBinding = GetThisViewBinding();
+
+		ViewModelHelper = MakeUnique<UE::MVVM::FViewModelFieldPathHelper>(&ViewModelBinding->ViewModelPath, WidgetBlueprint);
+		WidgetHelper = MakeUnique<UE::MVVM::FWidgetFieldPathHelper>(&ViewModelBinding->WidgetPath, WidgetBlueprint);
+
+		SMultiColumnTableRow<FMVVMViewBindingListEntryPtr>::Construct(
+			FSuperRowType::FArguments()
+			.Padding(1.0f),
+			OwnerTableView
+		);
+	}
+
+	~SMVVMViewBindingListEntryRow()
+	{
+		WidgetBlueprint->OnChanged().Remove(OnBlueprintChangedHandle);
+	}
+
+	FMVVMBlueprintViewBinding* GetThisViewBinding() const
+	{
+		if (UMVVMBlueprintView* BlueprintViewPtr = BlueprintView.Get())
+		{
+			FMVVMBlueprintViewBinding* ViewBinding = BlueprintViewPtr->GetBindingAt(Entry->Index);
+			return ViewBinding;
+		}
+
+		return nullptr;
+	}
+
+	TArray<FMVVMBlueprintViewBinding*> GetThisViewBindingAsArray() const
+	{
+		TArray<FMVVMBlueprintViewBinding*> Result;
+
+		FMVVMBlueprintViewBinding* ViewBinding = GetThisViewBinding();
+		if (ViewBinding != nullptr)
+		{
+			Result.Add(ViewBinding);
+		}
+
+		return Result;
+	}
+
+	/** Overridden from SMultiColumnTableRow.  Generates a widget for this column of the list view. */
+	virtual TSharedRef<SWidget> GenerateWidgetForColumn(const FName& ColumnName) override
+	{
+		FMVVMBlueprintViewBinding* ViewModelBinding = GetThisViewBinding();
+
+		if (ColumnName == EnabledColumnName)
+		{
+			return SNew(SCheckBox)
+				.IsChecked(this, &SMVVMViewBindingListEntryRow::IsBindingEnabled)
+				.OnCheckStateChanged(this, &SMVVMViewBindingListEntryRow::OnIsBindingEnableChanged);
+		}
+		else if (ColumnName == ErrorColumnName)
+		{
+			return SNew(SSimpleButton)
+				.Icon(FAppStyle::Get().GetBrush("Icons.Error"))
+				.Visibility(this, &SMVVMViewBindingListEntryRow::GetErrorVisibility)
+				.ToolTipText(this, &SMVVMViewBindingListEntryRow::GetErrorToolTip)
+				.OnClicked(this, &SMVVMViewBindingListEntryRow::OnErrorButtonClicked);
+		}
+		else if (ColumnName == ViewModelColumnName)
+		{
+			return SNew(SBox)
+				.Padding(FMargin(2, 0))
+				.VAlign(VAlign_Center)
+				[
+					SAssignNew(ViewModelSourceSelector, SMVVMSourceSelector)
+					.PathHelpers(this, &SMVVMViewBindingListEntryRow::GetViewModelHelpers)
+					.OnSelectionChanged(this, &SMVVMViewBindingListEntryRow::OnViewModelSelectionChanged)
+				];
+		}
+		else if (ColumnName == ViewModelPropertyColumnName)
+		{
+
+			return SNew(SBox)
+				.Padding(FMargin(2, 0))
+				.VAlign(VAlign_Center)
+				[
+					SAssignNew(ViewModelFieldSelector, SMVVMFieldSelector)
+					.PathHelpers(GetViewModelHelpers())
+					.CounterpartHelpers(GetWidgetHelpers())
+					.BindingMode(this, &SMVVMViewBindingListEntryRow::GetCurrentBindingMode)
+					.IsSource(true)
+					.OnSelectionChanged(this, &SMVVMViewBindingListEntryRow::OnViewModelPropertySelectionChanged)
+				];
+		}
+		else if (ColumnName == ModeColumnName)
+		{
+			UEnum* ModeEnum = StaticEnum<EMVVMBindingMode>();
+			for (int32 Index = 0; Index < ModeEnum->NumEnums() - 1; ++Index)
+			{
+				const bool bIsHidden = ModeEnum->HasMetaData(TEXT("Hidden"), Index);
+				if (!bIsHidden)
+				{
+					ModeNames.Add(ModeEnum->GetNameByIndex(Index));
+				}
+			}
+
+			return SNew(SComboBox<FName>)
+				.OptionsSource(&ModeNames)
+				.InitiallySelectedItem(ModeEnum->GetNameByValue((int64) ViewModelBinding->BindingType))
+				.OnSelectionChanged(this, &SMVVMViewBindingListEntryRow::OnModeSelectionChanged)
+				.OnGenerateWidget(this, &SMVVMViewBindingListEntryRow::GenerateModeWidget)
+				.Content()
+				[
+					SNew(SBox)
+					.HAlign(HAlign_Center)
+					.VAlign(VAlign_Center)
+					.WidthOverride(16)
+					.HeightOverride(16)
+					[
+						SNew(SImage)
+						.Image(this, &SMVVMViewBindingListEntryRow::GetCurrentModeBrush)
+					]
+				];
+		}
+		else if (ColumnName == WidgetColumnName)
+		{
+			return SNew(SBox)
+				.Padding(FMargin(2, 0))
+				.VAlign(VAlign_Center)
+				[
+					SAssignNew(WidgetSourceSelector, SMVVMSourceSelector)
+					.PathHelpers(this, &SMVVMViewBindingListEntryRow::GetWidgetHelpers)
+					.OnSelectionChanged(this, &SMVVMViewBindingListEntryRow::OnWidgetSelectionChanged)
+				];
+		}
+		else if (ColumnName == WidgetPropertyColumnName)
+		{
+			return SNew(SBox)
+				.Padding(FMargin(2, 0))
+				.VAlign(VAlign_Center)
+				[
+					SAssignNew(WidgetFieldSelector, SMVVMFieldSelector)
+					.PathHelpers(GetWidgetHelpers())
+					.CounterpartHelpers(GetViewModelHelpers())
+					.BindingMode(this, &SMVVMViewBindingListEntryRow::GetCurrentBindingMode)
+					.IsSource(false)
+					.OnSelectionChanged(this, &SMVVMViewBindingListEntryRow::OnWidgetPropertySelectionChanged)
+				];
+		}
+		else if (ColumnName == UpdateColumnName)
+		{
+			const UEnum* UpdateModeEnum = StaticEnum<EMVVMViewBindingUpdateMode>();
+			return SNew(SBox)
+				.Padding(FMargin(2, 0))
+				.VAlign(VAlign_Center)
+				[
+					SNew(SEnumComboBox, UpdateModeEnum)
+					.OnEnumSelectionChanged(this, &SMVVMViewBindingListEntryRow::OnUpdateModeSelectionChanged)
+					.CurrentValue(this, &SMVVMViewBindingListEntryRow::GetUpdateModeValue)
+				];
+		}
+		else if (ColumnName == ConversionFunctionColumnName)
+		{
+			return SNew(SHorizontalBox)
+				+ SHorizontalBox::Slot()
+				.VAlign(VAlign_Center)
+				.HAlign(HAlign_Left)
+				.AutoWidth()
+				[
+					SNew(SMVVMConversionPath, WidgetBlueprint, false)
+					.Bindings(this, &SMVVMViewBindingListEntryRow::GetThisViewBindingAsArray)
+				]
+				+ SHorizontalBox::Slot()
+				.VAlign(VAlign_Center)
+				.HAlign(HAlign_Left)
+				.AutoWidth()
+				[
+					SNew(SMVVMConversionPath, WidgetBlueprint, true)
+					.Bindings(this, &SMVVMViewBindingListEntryRow::GetThisViewBindingAsArray)
+				];
+		}
+		else if (ColumnName == DropDownOptionsColumnName)
+		{
+			return SAssignNew(ContextMenuOptionHelper, SButton)
+				.ToolTipText(LOCTEXT("DropDownOptionsToolTip", "Context Menu for Binding"))
+				.ButtonStyle(&FAppStyle::Get().GetWidgetStyle<FButtonStyle>("SimpleButton"))
+				.OnClicked(this, &SMVVMViewBindingListEntryRow::HandleDropDownOptionsPressed)
+				[
+					SNew(SBox)
+					.Padding(FMargin(3, 0))
+					[
+						SNew(SImage)
+						.Image(&FAppStyle::Get().GetWidgetStyle<FComboButtonStyle>("SegmentedCombo.Right").DownArrowImage)
+						.ColorAndOpacity(FSlateColor::UseForeground())
+					]
+				];
+		}
+		
+		return SNullWidget::NullWidget;
+	}
+
+private:
+
+	ECheckBoxState IsBindingEnabled() const
+	{
+		if (const FMVVMBlueprintViewBinding* ViewModelBinding = GetThisViewBinding())
+		{
+			return ViewModelBinding->bEnabled ? ECheckBoxState::Checked : ECheckBoxState::Unchecked;
+		}
+		return ECheckBoxState::Undetermined;
+	}
+
+	EVisibility GetErrorVisibility() const
+	{
+		return GetThisViewBinding()->Errors.IsEmpty() ? EVisibility::Hidden : EVisibility::Visible;
+	}
+
+	FText GetErrorToolTip() const
+	{
+		static const FText NewLineText = FText::FromString(TEXT("\n"));
+		FText HintText = LOCTEXT("ErrorButtonText", "Errors: (Click to show in a separate window)");
+		FText ErrorsText = FText::Join(NewLineText, GetThisViewBinding()->Errors);
+
+		return FText::Join(NewLineText, HintText, ErrorsText);
+	}
+
+	FReply OnErrorButtonClicked()
+	{
+		ErrorDialog.Reset();
+		ErrorItems.Reset();
+
+		if (const FMVVMBlueprintViewBinding* ViewModelBinding = GetThisViewBinding())
+		{
+			for (const FText& ErrorText : ViewModelBinding->Errors)
+			{
+				ErrorItems.Add(MakeShared<FText>(ErrorText));
+			}
+
+			ErrorDialog = SNew(SCustomDialog)
+				.Buttons({
+					SCustomDialog::FButton(LOCTEXT("OK", "OK"))
+				})
+				.Content()
+				[
+					SNew(SListView<TSharedPtr<FText>>)
+					.ListItemsSource(&ErrorItems)
+					.OnGenerateRow(this, &SMVVMViewBindingListEntryRow::OnGenerateErrorRow)
+				];
+
+			ErrorDialog->Show();
+		}
+
+		return FReply::Handled();
+	}
+
+	EMVVMBindingMode GetCurrentBindingMode() const
+	{
+		const FMVVMBlueprintViewBinding* ViewModelBinding = GetThisViewBinding();
+		return ViewModelBinding->BindingType;
+	}
+
+	UMVVMSubsystem::FConstDirectionalBindingArgs GetCurrentBindingArgs(bool bIsGetter) const
+	{
+		UE::MVVM::FMVVMConstFieldVariant ViewModelField = ViewModelHelper->GetSelectedField();
+		UE::MVVM::FMVVMConstFieldVariant WidgetField = WidgetHelper->GetSelectedField();
+
+		const FMVVMBlueprintViewBinding* ViewModelBinding = GetThisViewBinding();
+
+		UMVVMSubsystem::FConstDirectionalBindingArgs Args;
+		if (bIsGetter)
+		{
+			Args.SourceBinding = ViewModelField;
+			Args.DestinationBinding = WidgetField;
+		}
+		else
+		{
+			Args.SourceBinding = WidgetField;
+			Args.DestinationBinding = ViewModelField;
+		}
+
+		return MoveTemp(Args);
+	}
+
+	TSharedRef<ITableRow> OnGenerateErrorRow(TSharedPtr<FText> Text, const TSharedRef<STableViewBase>& TableView) const
+	{
+		return SNew(STableRow<TSharedPtr<FText>>, TableView)
+			.Content()
+			[
+				SNew(SEditableTextBox)
+				.IsReadOnly(true)
+				.Text(*Text.Get())
+			];
+	}
+
+	void OnViewModelSelectionChanged(TOptional<UE::MVVM::FBindingSource> Source)
+	{
+		OnSourceSelectionChanged(Source, ViewModelHelper.Get());
+		if (ViewModelFieldSelector.IsValid())
+		{
+			ViewModelFieldSelector->Refresh();
+		}
+	}
+
+	void OnWidgetSelectionChanged(TOptional<UE::MVVM::FBindingSource> Source)
+	{
+		OnSourceSelectionChanged(Source, WidgetHelper.Get());
+		if (WidgetFieldSelector.IsValid())
+		{
+			WidgetFieldSelector->Refresh();
+		}
+	}
+
+	void OnSourceSelectionChanged(TOptional<UE::MVVM::FBindingSource> SelectedSource, UE::MVVM::IFieldPathHelper* PathHelper)
+	{
+		if (UMVVMBlueprintView* BlueprintViewPtr = BlueprintView.Get())
+		{
+			BlueprintViewPtr->PreEditChange(UMVVMBlueprintView::StaticClass()->FindPropertyByName("Bindings"));
+
+			PathHelper->SetSelectedSource(SelectedSource);
+			PathHelper->SetBindingName(FName()); // Might make sense to keep this around in case we retarget to a compatible widget or switch back.
+
+			BlueprintViewPtr->PostEditChange();
+		}
+	}
+
+	void OnViewModelPropertySelectionChanged(UE::MVVM::FMVVMConstFieldVariant SelectedField)
+	{
+		UE::MVVM::FMVVMConstFieldVariant CurrentField = ViewModelHelper->GetSelectedField();
+		if (CurrentField != SelectedField)
+		{
+			OnPropertySelectionChanged(SelectedField, ViewModelHelper.Get());
+
+			if (WidgetFieldSelector.IsValid())
+			{
+				WidgetFieldSelector->Refresh();
+			}
+		}
+	}
+
+	void OnWidgetPropertySelectionChanged(UE::MVVM::FMVVMConstFieldVariant SelectedField)
+	{
+		UE::MVVM::FMVVMConstFieldVariant CurrentField = WidgetHelper->GetSelectedField();
+		if (CurrentField != SelectedField)
+		{
+			OnPropertySelectionChanged(SelectedField, WidgetHelper.Get());
+
+			if (ViewModelFieldSelector.IsValid())
+			{
+				ViewModelFieldSelector->Refresh();
+			}
+		}
+	}
+
+	void OnPropertySelectionChanged(UE::MVVM::FMVVMConstFieldVariant SelectedField, UE::MVVM::IFieldPathHelper* PathHelper)
+	{
+		if (UMVVMBlueprintView* BlueprintViewPtr = BlueprintView.Get())
+		{
+			BlueprintViewPtr->PreEditChange(UMVVMBlueprintView::StaticClass()->FindPropertyByName("Bindings"));
+
+			PathHelper->SetBindingName(SelectedField.GetName());
+
+			BlueprintViewPtr->PostEditChange();
+		}
+	}
+
+	void OnUpdateModeSelectionChanged(int32 Value, ESelectInfo::Type)
+	{
+		if (FMVVMBlueprintViewBinding* ViewModelBinding = GetThisViewBinding())
+		{
+			if (ViewModelBinding->UpdateMode != (EMVVMViewBindingUpdateMode) Value)
+			{
+				if (UMVVMBlueprintView* BlueprintViewPtr = BlueprintView.Get())
+				{
+					BlueprintViewPtr->PreEditChange(UMVVMBlueprintView::StaticClass()->FindPropertyByName("Bindings"));
+
+					ViewModelBinding->UpdateMode = (EMVVMViewBindingUpdateMode) Value;
+
+					BlueprintViewPtr->PostEditChange();
+				}
+			}
+		}
+	}
+
+	int32 GetUpdateModeValue() const
+	{
+		return (int32) GetThisViewBinding()->UpdateMode;
+	}
+
+	void OnIsBindingEnableChanged(ECheckBoxState NewState)
+	{
+		if (NewState == ECheckBoxState::Undetermined)
+		{
+			return;
+		}
+
+		if (FMVVMBlueprintViewBinding* ViewModelBinding = GetThisViewBinding())
+		{
+			bool bNewEnabled = (NewState == ECheckBoxState::Checked);
+			if (ViewModelBinding->bEnabled != bNewEnabled)
+			{
+				if (UMVVMBlueprintView* BlueprintViewPtr = BlueprintView.Get())
+				{
+					BlueprintViewPtr->PreEditChange(UMVVMBlueprintView::StaticClass()->FindPropertyByName("Bindings"));
+
+					ViewModelBinding->bEnabled = bNewEnabled;
+
+					BlueprintViewPtr->PostEditChange();
+				}
+			}
+		}
+	}
+
+	const FSlateBrush* GetModeBrush(EMVVMBindingMode BindingMode) const
+	{
+		switch (BindingMode)
+		{
+		case EMVVMBindingMode::OneTimeToDestination:
+			return FUMGStyle::Get().GetBrush("MVVM.OneTime");
+		case EMVVMBindingMode::OneWayToDestination:
+			return FUMGStyle::Get().GetBrush("MVVM.OneWay");
+		case EMVVMBindingMode::OneWayToSource:
+			return FUMGStyle::Get().GetBrush("MVVM.OneWayToSource");
+		case EMVVMBindingMode::OneTimeToSource:
+			return nullptr;
+		case EMVVMBindingMode::TwoWay:
+			return FUMGStyle::Get().GetBrush("MVVM.TwoWay");
+		default:
+			return nullptr;
+		}
+	}
+
+	const FSlateBrush* GetCurrentModeBrush() const
+	{
+		return GetModeBrush(GetThisViewBinding()->BindingType);
+	}
+
+	const FText& GetModeLabel(EMVVMBindingMode BindingMode) const
+	{
+		static FText OneTimeToDestinationLabel = LOCTEXT("OneTimeToDestinationLabel", "One Time To Widget");
+		static FText OneWayToDestinationLabel = LOCTEXT("OneWayToDestinationLabel", "One Way To Widget");
+		static FText OneWayToSourceLabel = LOCTEXT("OneWayToSourceLabel", "One Way To View Model");
+		static FText OneTimeToSourceLabel = LOCTEXT("OneTimeToSourceLabel", "One Time To View Model");
+		static FText TwoWayLabel = LOCTEXT("TwoWayLabel", "Two Way");
+
+		switch (BindingMode)
+		{
+		case EMVVMBindingMode::OneTimeToDestination:
+			return OneTimeToDestinationLabel;
+		case EMVVMBindingMode::OneWayToDestination:
+			return OneWayToDestinationLabel;
+		case EMVVMBindingMode::OneWayToSource:
+			return OneWayToSourceLabel;
+		case EMVVMBindingMode::OneTimeToSource:
+			return OneTimeToSourceLabel;
+		case EMVVMBindingMode::TwoWay:
+			return TwoWayLabel;
+		default:
+			return FText::GetEmpty();
+		}
+	}
+
+	TSharedRef<SWidget> GenerateModeWidget(FName ValueName) const
+	{
+		const UEnum* ModeEnum = StaticEnum<EMVVMBindingMode>();
+		int32 Index = ModeEnum->GetIndexByName(ValueName);
+		EMVVMBindingMode MVVMBindingMode = EMVVMBindingMode(Index);
+		return SNew(SHorizontalBox)
+			+SHorizontalBox::Slot()
+			.AutoWidth()
+			[
+				SNew(SBox)
+				.HAlign(HAlign_Center)
+				.VAlign(VAlign_Center)
+				.WidthOverride(16)
+				.HeightOverride(16)
+				[
+					SNew(SImage)
+					.Image(GetModeBrush(MVVMBindingMode))
+				]
+			]
+			+SHorizontalBox::Slot()
+			.AutoWidth()
+			.Padding(2.0, 0, 0, 0)
+			[
+				SNew(STextBlock)
+				.Text(GetModeLabel(MVVMBindingMode))
+				.ToolTipText(ModeEnum->GetToolTipTextByIndex(Index))
+			];
+	}
+
+	void OnModeSelectionChanged(FName ValueName, ESelectInfo::Type)
+	{
+		if (FMVVMBlueprintViewBinding* ViewModelBinding = GetThisViewBinding())
+		{
+			const UEnum* ModeEnum = StaticEnum<EMVVMBindingMode>();
+			EMVVMBindingMode NewMode = (EMVVMBindingMode) ModeEnum->GetValueByName(ValueName);
+
+			if (ViewModelBinding->BindingType != NewMode)
+			{
+				if (UMVVMBlueprintView* BlueprintViewPtr = BlueprintView.Get())
+				{
+					BlueprintViewPtr->PreEditChange(UMVVMBlueprintView::StaticClass()->FindPropertyByName("Bindings"));
+
+					ViewModelBinding->BindingType = NewMode;
+
+					BlueprintViewPtr->PostEditChange();
+				}
+
+				ViewModelFieldSelector->Refresh();
+				WidgetFieldSelector->Refresh();
+			}
+		}
+	}
+
+	FReply HandleDropDownOptionsPressed() 
+	{
+		if (TSharedPtr<ITypedTableView<FMVVMViewBindingListEntryPtr>> ListView = OwnerTablePtr.Pin())
+		{
+			if (TSharedPtr<SMVVMViewBindingListView> ParentList = StaticCastSharedPtr<SMVVMViewBindingListView>(ListView->AsWidget()->GetParentWidget()))
+			{
+				// Get the context menu content. If invalid, don't open a menu.
+				ListView->Private_SetItemSelection(Entry, true);
+				TSharedPtr<SWidget> MenuContent = ParentList->OnSourceConstructContextMenu();
+
+				if (MenuContent.IsValid())
+				{
+					const FVector2D& SummonLocation = ContextMenuOptionHelper->GetCachedGeometry().GetRenderBoundingRect().GetBottomLeft();
+					FWidgetPath WidgetPath;
+					FSlateApplication::Get().PushMenu(ParentList->AsShared(), WidgetPath, MenuContent.ToSharedRef(), SummonLocation, FPopupTransitionEffect(FPopupTransitionEffect::ContextMenu));
+				}
+			}
+		}
+
+		return FReply::Handled();
+	}
+
+	TArray<UE::MVVM::IFieldPathHelper*> GetWidgetHelpers() const
+	{
+		return TArray<UE::MVVM::IFieldPathHelper*> { WidgetHelper.Get() };
+	}
+
+	TArray<UE::MVVM::IFieldPathHelper*> GetViewModelHelpers() const
+	{
+		return TArray<UE::MVVM::IFieldPathHelper*> { ViewModelHelper.Get() };
+	}
+
+	void HandleBlueprintChanged(UBlueprint* Blueprint)
+	{
+		ViewModelSourceSelector->Refresh();
+		ViewModelFieldSelector->Refresh();
+		WidgetSourceSelector->Refresh();
+		WidgetFieldSelector->Refresh();
+	}
+
+private:
+	FMVVMViewBindingListEntryPtr Entry;
+	TWeakObjectPtr<UMVVMBlueprintView> BlueprintView;
+	UWidgetBlueprint* WidgetBlueprint { nullptr };
+	TUniquePtr<UE::MVVM::FViewModelFieldPathHelper> ViewModelHelper;
+	TUniquePtr<UE::MVVM::FWidgetFieldPathHelper> WidgetHelper;
+	TSharedPtr<SMVVMSourceSelector> ViewModelSourceSelector;
+	TSharedPtr<SMVVMFieldSelector> ViewModelFieldSelector;
+	TSharedPtr<SMVVMSourceSelector> WidgetSourceSelector;
+	TSharedPtr<SMVVMFieldSelector> WidgetFieldSelector;
+	TSharedPtr<SWidget> ContextMenuOptionHelper;
+	TSharedPtr<SCustomDialog> ErrorDialog;
+	TSharedPtr<SMenuAnchor> GetterConversionFunctionAnchor;
+	TSharedPtr<SMenuAnchor> SetterConversionFunctionAnchor;
+	TArray<TSharedPtr<FText>> ErrorItems;
+	TArray<FName> ModeNames;
+	FDelegateHandle OnBlueprintChangedHandle;
+};
+
+FName SMVVMViewBindingListEntryRow::EnabledColumnName = "Enabled";
+FName SMVVMViewBindingListEntryRow::ErrorColumnName = "Error";
+FName SMVVMViewBindingListEntryRow::ViewModelColumnName = "ViewModel";
+FName SMVVMViewBindingListEntryRow::ViewModelPropertyColumnName = "ViewModelProperty";
+FName SMVVMViewBindingListEntryRow::ModeColumnName = "Mode";
+FName SMVVMViewBindingListEntryRow::WidgetColumnName = "Widget";
+FName SMVVMViewBindingListEntryRow::WidgetPropertyColumnName = "WidgetProperty";
+FName SMVVMViewBindingListEntryRow::UpdateColumnName = "Update";
+FName SMVVMViewBindingListEntryRow::ConversionFunctionColumnName = "ConversionFunction";
+FName SMVVMViewBindingListEntryRow::DropDownOptionsColumnName = "DropdownOptions";
+
+/** */
+void SMVVMViewBindingListView::Construct(const FArguments& InArgs, TSharedPtr<SMVVMViewBindingPanel> Owner, UMVVMWidgetBlueprintExtension_View* InMVVMExtension)
+{
+	BindingPanel = Owner;
+	MVVMExtension = InMVVMExtension;
+	check(InMVVMExtension);
+
+	MVVMExtension->OnBlueprintViewChangedDelegate().AddSP(this, &SMVVMViewBindingListView::RequestListRefresh);
+	MVVMExtension->GetBlueprintView()->OnBindingsUpdated.AddSP(this, &SMVVMViewBindingListView::RequestListRefresh);
+	MVVMExtension->GetBlueprintView()->OnViewModelsUpdated.AddSP(this, &SMVVMViewBindingListView::RequestListRefresh);
+
+	RequestListRefresh();
+
+	ChildSlot
+	[
+		SAssignNew(ListView, SListView<FMVVMViewBindingListEntryPtr>)
+		.ListItemsSource(&SourceData)
+		.SelectionMode(ESelectionMode::Single)
+		.OnGenerateRow(this, &SMVVMViewBindingListView::MakeSourceListViewWidget)
+		.OnContextMenuOpening(this, &SMVVMViewBindingListView::OnSourceConstructContextMenu)
+		.OnSelectionChanged(this, &SMVVMViewBindingListView::OnSourceListSelectionChanged)
+		.HeaderRow
+		(
+			SNew(SHeaderRow)
+			+ SHeaderRow::Column(SMVVMViewBindingListEntryRow::EnabledColumnName)
+			.DefaultLabel(FText::GetEmpty())
+			.FixedWidth(25.f)
+			+ SHeaderRow::Column(SMVVMViewBindingListEntryRow::ErrorColumnName)
+			.DefaultLabel(FText::GetEmpty())
+			.FixedWidth(25.f)
+			.HeaderContent()
+			[
+				SNew(SBox)
+				.WidthOverride(16)
+				.HeightOverride(16)
+				.VAlign(VAlign_Center)
+				.HAlign(HAlign_Center)
+				[
+					SNew(SImage)
+					.Image(FAppStyle::Get().GetBrush("Icons.Error"))
+				]
+			]
+			+ SHeaderRow::Column(SMVVMViewBindingListEntryRow::ViewModelColumnName)
+			.FillWidth(0.125f)
+			.DefaultLabel(LOCTEXT("ViewModel", "View Model"))
+			+ SHeaderRow::Column(SMVVMViewBindingListEntryRow::ViewModelPropertyColumnName)
+			.FillWidth(0.125f)
+			.DefaultLabel(LOCTEXT("ViewModelProperty", "View Model Property"))
+			+ SHeaderRow::Column(SMVVMViewBindingListEntryRow::ModeColumnName)
+			.FixedWidth(52)
+			.DefaultLabel(LOCTEXT("Mode", "Mode"))
+			+ SHeaderRow::Column(SMVVMViewBindingListEntryRow::WidgetColumnName)
+			.FillWidth(0.125f)
+			.DefaultLabel(LOCTEXT("Widget", "Widget"))
+			+ SHeaderRow::Column(SMVVMViewBindingListEntryRow::WidgetPropertyColumnName)
+			.FillWidth(0.125f)
+			.DefaultLabel(LOCTEXT("Widget", "Widget Property"))
+			+ SHeaderRow::Column(SMVVMViewBindingListEntryRow::UpdateColumnName)
+			.FillWidth(0.05f)
+			.DefaultLabel(LOCTEXT("Update", "Update"))
+			+ SHeaderRow::Column(SMVVMViewBindingListEntryRow::ConversionFunctionColumnName)
+			.FillWidth(0.05f)
+			.DefaultLabel(LOCTEXT("Conversion", "Conversion"))
+			+ SHeaderRow::Column(SMVVMViewBindingListEntryRow::DropDownOptionsColumnName)
+			.FixedWidth(25.f)
+			.DefaultLabel(FText::GetEmpty())
+		)
+	];
+}
+
+
+SMVVMViewBindingListView::~SMVVMViewBindingListView()
+{
+	if (UMVVMWidgetBlueprintExtension_View* MVVMExtensionPtr = MVVMExtension.Get())
+	{
+		MVVMExtensionPtr->OnBlueprintViewChangedDelegate().RemoveAll(this);
+		MVVMExtensionPtr->GetBlueprintView()->OnBindingsUpdated.RemoveAll(this);
+		MVVMExtensionPtr->GetBlueprintView()->OnViewModelsUpdated.RemoveAll(this);
+	}
+}
+
+void SMVVMViewBindingListView::RequestListRefresh()
+{
+	SourceData.Reset();
+	if (const UMVVMWidgetBlueprintExtension_View* MVVMExtensionPtr = MVVMExtension.Get())
+	{
+		if (const UMVVMBlueprintView* BlueprintView = MVVMExtensionPtr->GetBlueprintView())
+		{
+			const int32 Count = BlueprintView->GetBindings().Num();
+			for (int32 Index = 0; Index < Count; ++Index)
+			{
+				SourceData.Add(MakeShared<FMVVMViewBindingListEntry>(Index));
+			}
+		}
+	}
+
+	if (ListView.IsValid())
+	{
+		ListView->RequestListRefresh();
+	}
+}
+
+
+TSharedRef<ITableRow> SMVVMViewBindingListView::MakeSourceListViewWidget(FMVVMViewBindingListEntryPtr Entry, const TSharedRef<STableViewBase>& OwnerTable) const
+{
+	if (UMVVMWidgetBlueprintExtension_View* MVVMExtensionPtr = MVVMExtension.Get())
+	{
+		return SNew(SMVVMViewBindingListEntryRow, OwnerTable, Entry, MVVMExtensionPtr->GetBlueprintView(), MVVMExtensionPtr->GetWidgetBlueprint());
+	}
+	return SNew(SMVVMViewBindingListEntryRow, OwnerTable, Entry, nullptr, nullptr);
+}
+
+TSharedPtr<SWidget> SMVVMViewBindingListView::OnSourceConstructContextMenu()
+{
+	const bool bShouldCloseWindowAfterMenuSelection = true;
+	FMenuBuilder MenuBuilder(bShouldCloseWindowAfterMenuSelection, nullptr);
+
+	TArray<FMVVMViewBindingListEntryPtr> Selection = ListView->GetSelectedItems();
+	if (Selection.Num() > 0)
+	{
+		FUIAction RemoveAction;
+		RemoveAction.ExecuteAction = FExecuteAction::CreateLambda([this, ToRemove = Selection[0]->Index]()
+			{
+				if (UMVVMWidgetBlueprintExtension_View* MVVMExtensionPtr = MVVMExtension.Get())
+				{
+					if (TSharedPtr<SMVVMViewBindingPanel> BindingPanelPtr = BindingPanel.Pin())
+					{
+						BindingPanelPtr->OnBindingListSelectionChanged(INDEX_NONE);
+					}
+
+					UMVVMBlueprintView* BlueprintView = MVVMExtensionPtr->GetBlueprintView();
+					BlueprintView->RemoveBindingAt(ToRemove);
+				}
+			});
+		MenuBuilder.AddMenuEntry(LOCTEXT("RemoveBinding", "Remove Binding"), LOCTEXT("RemoveBindingTooltip", "Remove this binding."), FSlateIcon(FAppStyle::GetAppStyleSetName(), "Icons.Delete"), RemoveAction);
+	}
+
+	return MenuBuilder.MakeWidget();
+}
+
+void SMVVMViewBindingListView::OnSourceListSelectionChanged(FMVVMViewBindingListEntryPtr Entry, ESelectInfo::Type SelectionType) const
+{
+	if (bSelectionChangedGuard)
+	{
+		return;
+	}
+	TGuardValue<bool> ReentrantGuard(bSelectionChangedGuard, true);
+
+	const int32 SelectionIndex = Entry ? Entry->Index : INDEX_NONE;
+	if (TSharedPtr<SMVVMViewBindingPanel> BindingPanelPtr = BindingPanel.Pin())
+	{
+		BindingPanelPtr->OnBindingListSelectionChanged(SelectionIndex);
+	}
+}
+
+#undef LOCTEXT_NAMESPACE
diff --git a/Engine/Plugins/Runtime/ModelViewViewModel/Source/ModelViewViewModelEditor/Private/Widgets/SMVVMViewBindingListView.h b/Engine/Plugins/Runtime/ModelViewViewModel/Source/ModelViewViewModelEditor/Private/Widgets/SMVVMViewBindingListView.h
new file mode 100644
--- /dev/null
+++ b/Engine/Plugins/Runtime/ModelViewViewModel/Source/ModelViewViewModelEditor/Private/Widgets/SMVVMViewBindingListView.h
@@ -0,0 +1,43 @@
+// Copyright Epic Games, Inc. All Rights Reserved.
+
+#pragma once
+
+#include "CoreMinimal.h"
+
+#include "Widgets/DeclarativeSyntaxSupport.h"
+#include "Widgets/SCompoundWidget.h"
+
+class ITableRow;
+template <typename ItemType> class SListView;
+class STableViewBase;
+class SMVVMViewBindingPanel;
+class UMVVMWidgetBlueprintExtension_View;
+
+struct FMVVMViewBindingListEntry;
+using FMVVMViewBindingListEntryPtr = TSharedPtr<FMVVMViewBindingListEntry>;
+
+class SMVVMViewBindingListView : public SCompoundWidget
+{
+public:
+	SLATE_BEGIN_ARGS(SMVVMViewBindingListView) {}
+	SLATE_END_ARGS()
+
+	void Construct(const FArguments& InArgs, TSharedPtr<SMVVMViewBindingPanel> Owner, UMVVMWidgetBlueprintExtension_View* MVVMExtension);
+	~SMVVMViewBindingListView();
+
+	void RequestListRefresh();
+
+	/** Constructs context menu used for right click and dropdown button */
+	TSharedPtr<SWidget> OnSourceConstructContextMenu();
+
+private:
+	TSharedRef<ITableRow> MakeSourceListViewWidget(FMVVMViewBindingListEntryPtr Entry, const TSharedRef<STableViewBase>& OwnerTable) const;
+	void OnSourceListSelectionChanged(FMVVMViewBindingListEntryPtr Entry, ESelectInfo::Type SelectionType) const;
+
+private:
+	TWeakPtr<SMVVMViewBindingPanel> BindingPanel;
+	TSharedPtr<SListView<FMVVMViewBindingListEntryPtr>> ListView;
+	TArray<FMVVMViewBindingListEntryPtr> SourceData;
+	TWeakObjectPtr<UMVVMWidgetBlueprintExtension_View> MVVMExtension;
+	mutable bool bSelectionChangedGuard = false;
+};
\ No newline at end of file
