MVVM: Move the ModelViewViewModel plugin to //UE5/Main/

diff --git a/Engine/Plugins/Runtime/ModelViewViewModel/Source/ModelViewViewModelEditor/Private/Widgets/SMVVMViewBindingPanel.cpp b/Engine/Plugins/Runtime/ModelViewViewModel/Source/ModelViewViewModelEditor/Private/Widgets/SMVVMViewBindingPanel.cpp
new file mode 100644
--- /dev/null
+++ b/Engine/Plugins/Runtime/ModelViewViewModel/Source/ModelViewViewModelEditor/Private/Widgets/SMVVMViewBindingPanel.cpp
@@ -0,0 +1,468 @@
+// Copyright Epic Games, Inc. All Rights Reserved.
+
+#include "Widgets/SMVVMViewBindingPanel.h"
+
+#include "Widgets/SMVVMManageViewModelsWidget.h"
+#include "Widgets/SMVVMViewModelContextListWidget.h"
+#include "Widgets/SMVVMViewBindingListView.h"
+
+#include "WidgetBlueprintEditor.h"
+#include "WidgetBlueprintToolMenuContext.h"
+#include "MVVMBlueprintView.h"
+#include "MVVMWidgetBlueprintExtension_View.h"
+#include "Customizations/MVVMConversionPathCustomization.h"
+#include "Customizations/MVVMPropertyPathCustomization.h"
+
+#include "EditorStyleSet.h"
+#include "IStructureDetailsView.h"
+#include "PropertyEditorModule.h"
+#include "ToolMenuContext.h"
+#include "ToolMenus.h"
+
+#include "Framework/MultiBox/MultiBoxBuilder.h"
+#include "Widgets/Input/SButton.h"
+#include "Widgets/SNullWidget.h"
+#include "Widgets/Layout/SBorder.h"
+#include "Widgets/Layout/SSpacer.h"
+#include "Widgets/Layout/SSplitter.h"
+
+#include "StatusBarSubsystem.h"
+#include "Tabs/MVVMBindingSummoner.h"
+#include "Widgets/Docking/SDockTab.h"
+
+#define LOCTEXT_NAMESPACE "BindingPanel"
+
+
+/** */
+void SMVVMViewBindingPanel::Construct(const FArguments& InArgs, TSharedPtr<FWidgetBlueprintEditor> WidgetBlueprintEditor, bool bInIsDrawerTab)
+{
+	WeakBlueprintEditor = WidgetBlueprintEditor;
+	bIsDrawerTab = bInIsDrawerTab;
+
+	UWidgetBlueprint* WidgetBlueprint = WidgetBlueprintEditor->GetWidgetBlueprintObj();
+	check(WidgetBlueprint);
+
+	UMVVMWidgetBlueprintExtension_View* MVVMExtensionPtr = UMVVMWidgetBlueprintExtension_View::GetExtension<UMVVMWidgetBlueprintExtension_View>(WidgetBlueprint);
+	MVVMExtension = MVVMExtensionPtr;
+	if (MVVMExtensionPtr)
+	{
+		BlueprintViewChangedDelegateHandle = MVVMExtensionPtr->OnBlueprintViewChangedDelegate().AddSP(this, &SMVVMViewBindingPanel::HandleBlueprintViewChangedDelegate);
+	}
+	else
+	{
+		WidgetBlueprint->OnExtensionAdded.AddRaw(this, &SMVVMViewBindingPanel::HandleExtensionAdded);
+	}
+
+	{
+		// Connection Settings
+		FPropertyEditorModule& PropertyEditorModule = FModuleManager::GetModuleChecked<FPropertyEditorModule>("PropertyEditor");
+
+		FDetailsViewArgs DetailsViewArgs;
+		DetailsViewArgs.bUpdatesFromSelection = false;
+		DetailsViewArgs.bLockable = false;
+		DetailsViewArgs.bShowPropertyMatrixButton = false;
+		DetailsViewArgs.NameAreaSettings = FDetailsViewArgs::HideNameArea;
+		DetailsViewArgs.ViewIdentifier = NAME_None;
+		DetailsView = PropertyEditorModule.CreateDetailView(DetailsViewArgs);
+		
+		//DetailsView->OnFinishedChangingProperties().AddSP(this, &SLiveLinkClientPanel::OnPropertyChanged);
+		FStructureDetailsViewArgs StructureDetailsViewArgs;
+		StructDetailsView = PropertyEditorModule.CreateStructureDetailView(DetailsViewArgs, StructureDetailsViewArgs, TSharedPtr<FStructOnScope>());
+		StructDetailsView->GetDetailsView()->RegisterInstancedCustomPropertyTypeLayout(FMVVMViewModelPropertyPath::StaticStruct()->GetFName(), FOnGetPropertyTypeCustomizationInstance::CreateStatic(&UE::MVVM::FViewModelPropertyPathCustomization::MakeInstance, WidgetBlueprint));
+		StructDetailsView->GetDetailsView()->RegisterInstancedCustomPropertyTypeLayout(FMVVMWidgetPropertyPath::StaticStruct()->GetFName(), FOnGetPropertyTypeCustomizationInstance::CreateStatic(&UE::MVVM::FWidgetPropertyPathCustomization::MakeInstance, WidgetBlueprint));
+		StructDetailsView->GetDetailsView()->RegisterInstancedCustomPropertyTypeLayout(FMVVMBlueprintViewConversionPath::StaticStruct()->GetFName(), FOnGetPropertyTypeCustomizationInstance::CreateStatic(&UE::MVVM::FConversionPathCustomization::MakeInstance, WidgetBlueprint));
+	}
+
+	HandleBlueprintViewChangedDelegate();
+}
+
+
+SMVVMViewBindingPanel::~SMVVMViewBindingPanel()
+{
+	if (UMVVMWidgetBlueprintExtension_View* Extension = MVVMExtension.Get())
+	{
+		Extension->OnBlueprintViewChangedDelegate().Remove(BlueprintViewChangedDelegateHandle);
+	}
+	if (TSharedPtr<FWidgetBlueprintEditor> WidgetEditor = WeakBlueprintEditor.Pin())
+	{
+		UWidgetBlueprint* WidgetBlueprint = WidgetEditor->GetWidgetBlueprintObj();
+
+		if (WidgetBlueprint)
+		{
+			WidgetBlueprint->OnExtensionAdded.RemoveAll(this);
+		}
+	}
+}
+
+FReply SMVVMViewBindingPanel::OnKeyDown(const FGeometry& MyGeometry, const FKeyEvent& InKeyEvent)
+{
+	FReply Reply = FReply::Unhandled();
+	if (TSharedPtr<FWidgetBlueprintEditor> WidgetEditor = WeakBlueprintEditor.Pin())
+	{
+		if (WidgetEditor->GetToolkitCommands()->ProcessCommandBindings(InKeyEvent))
+		{
+			Reply = FReply::Handled();
+		}
+	}
+
+	return Reply;
+}
+
+bool SMVVMViewBindingPanel::SupportsKeyboardFocus() const
+{
+	return true;
+}
+
+void SMVVMViewBindingPanel::HandleBlueprintViewChangedDelegate()
+{
+	UMVVMWidgetBlueprintExtension_View* MVVMExtensionPtr = MVVMExtension.Get();
+
+	bool bHasExtension = MVVMExtensionPtr && MVVMExtensionPtr->GetBlueprintView() != nullptr;
+	ChildSlot
+	[
+		SNew(SOverlay)
+		+ SOverlay::Slot()
+		.HAlign(HAlign_Fill)
+		.VAlign(VAlign_Fill)
+		[
+			bHasExtension ? GenerateEditViewWidget() : GenerateCreateViewWidget()
+		]
+		+SOverlay::Slot()
+		.HAlign(HAlign_Right)
+		.VAlign(VAlign_Top)
+		.Padding(FMargin(bHasExtension ? 160.0 : 24.0, 10.0))
+		[
+			CreateDrawerDockButton()
+		]
+	];
+}
+
+
+void SMVVMViewBindingPanel::OnBindingListSelectionChanged(int32 Index)
+{
+	bool bDefault = true;
+	if (UMVVMWidgetBlueprintExtension_View* MVVMExtensionPtr = MVVMExtension.Get())
+	{
+		if (UMVVMBlueprintView* BlueprintView = MVVMExtensionPtr->GetBlueprintView())
+		{
+			FMVVMBlueprintViewBinding* Binding = BlueprintView->GetBindingAt(Index);
+			if (Binding != nullptr)
+			{
+				TSharedRef<FStructOnScope> StructScope = MakeShared<FStructOnScope>(FMVVMBlueprintViewBinding::StaticStruct(), reinterpret_cast<uint8*>(Binding));
+				StructDetailsView->SetStructureData(StructScope);
+				DetailContainer->SetContent(StructDetailsView->GetWidget().ToSharedRef());
+				bDefault = false;
+			}
+		}
+	}
+
+	if (bDefault)
+	{
+		DetailsView->SetObject(nullptr);
+		DetailContainer->SetContent(DetailsView.ToSharedRef());
+	}
+}
+
+
+void SMVVMViewBindingPanel::AddDefaultBinding()
+{
+	if (UMVVMWidgetBlueprintExtension_View* MVVMExtensionPtr = MVVMExtension.Get())
+	{
+		if (UMVVMBlueprintView* EditorData = MVVMExtensionPtr->GetBlueprintView())
+		{
+			EditorData->AddDefaultBinding();
+			ListView->RequestListRefresh();
+		}
+	}
+}
+
+
+bool SMVVMViewBindingPanel::CanAddBinding() const
+{
+	return true;
+}
+
+void SMVVMViewBindingPanel::ShowManageViewModelsWindow()
+{
+	if (UMVVMWidgetBlueprintExtension_View* MVVMExtensionPtr = MVVMExtension.Get())
+	{
+		if (UMVVMBlueprintView* BlueprintView = MVVMExtensionPtr->GetBlueprintView())
+		{
+			TWeakObjectPtr<UMVVMWidgetBlueprintExtension_View> LocalMVVMExtension = MVVMExtension;
+			FOnViewModelContextsPicked ViewModelContextsPickedDelegate = FOnViewModelContextsPicked::CreateLambda
+			(
+				[LocalMVVMExtension](TArray<FMVVMBlueprintViewModelContext> UpdatedContexts)
+				{
+					UMVVMWidgetBlueprintExtension_View* MVVMExtensionPtr = LocalMVVMExtension.Get();
+					if (ensure(MVVMExtensionPtr))
+					{
+						if (UMVVMBlueprintView* BlueprintView = MVVMExtensionPtr->GetBlueprintView())
+						{
+							BlueprintView->SetViewModels(UpdatedContexts);
+						}
+					}
+				}
+			);
+
+			TSharedRef<SWindow> ManageViewModelsWindow = SNew(SWindow)
+				.Title(LOCTEXT("MVVMManageViewMoedlsWindowHeader", "Manage ViewModels"))
+				.SupportsMaximize(false)
+				.ClientSize(FVector2D(800.0f, 600.0f));
+
+			TSharedRef<SMVVMViewModelContextListWidget> ContextListWidget = SNew(SMVVMViewModelContextListWidget)
+				.Bindings(BlueprintView->GetBindings())
+				.ExistingViewModelContexts(TArray<FMVVMBlueprintViewModelContext>(BlueprintView->GetViewModels()))
+				.WidgetBlueprint(WeakBlueprintEditor.Pin()->GetWidgetBlueprintObj())
+				.OwningWidget(Cast<UBlueprintGeneratedClass>(WeakBlueprintEditor.Pin()->GetBlueprintObj()->SkeletonGeneratedClass));
+
+			TSharedRef<SMVVMManageViewModelsWidget> ManageViewModelsWidget = SNew(SMVVMManageViewModelsWidget)
+				.ParentWindow(ManageViewModelsWindow)
+				.OnViewModelContextsPickedDelegate(ViewModelContextsPickedDelegate)
+				.ManageViewModelsWidget(ContextListWidget);
+
+			ManageViewModelsWindow->SetContent(ManageViewModelsWidget);
+			GEditor->EditorAddModalWindow(ManageViewModelsWindow);
+		}
+	}
+}
+
+TSharedRef<SWidget> SMVVMViewBindingPanel::CreateDrawerDockButton()
+{
+	if(bIsDrawerTab)
+	{
+		return 
+			SNew(SButton)
+			.ButtonStyle(FAppStyle::Get(), "SimpleButton")
+			.ToolTipText(LOCTEXT("DockInLayout_Tooltip", "Docks MVVM drawer in tab."))
+			.ContentPadding(FMargin(1, 0))
+			.OnClicked(this, &SMVVMViewBindingPanel::CreateDrawerDockButtonClicked)
+			[
+				SNew(SHorizontalBox)
+				+ SHorizontalBox::Slot()
+				.AutoWidth()
+				.VAlign(VAlign_Center)
+				.Padding(4.0, 0.0f)
+				[
+					SNew(SImage)
+					.ColorAndOpacity(FSlateColor::UseForeground())
+					.Image(FAppStyle::Get().GetBrush("EditorViewport.SubMenu.Layouts"))
+				]
+				+ SHorizontalBox::Slot()
+				.VAlign(VAlign_Center)
+				.Padding(4.0, 0.0f)
+				[
+					SNew(STextBlock)
+					.Text(LOCTEXT("DockInLayout", "Dock in Layout"))
+					.ColorAndOpacity(FSlateColor::UseForeground())
+				]
+			];
+	}
+
+	return SNullWidget::NullWidget;
+}
+
+FReply SMVVMViewBindingPanel::CreateDrawerDockButtonClicked()
+{
+	if (TSharedPtr<FWidgetBlueprintEditor> WidgetEditor = WeakBlueprintEditor.Pin())
+	{
+		GEditor->GetEditorSubsystem<UStatusBarSubsystem>()->ForceDismissDrawer();
+
+		if (TSharedPtr<SDockTab> ExistingTab = WidgetEditor->GetToolkitHost()->GetTabManager()->TryInvokeTab(FMVVMBindingSummoner::TabID))
+		{
+			ExistingTab->ActivateInParent(ETabActivationCause::SetDirectly); 
+		}
+	}
+
+	return FReply::Handled();
+}
+
+void SMVVMViewBindingPanel::HandleExtensionAdded(TObjectPtr<UBlueprintExtension> NewExtension)
+{
+	if (UMVVMWidgetBlueprintExtension_View* MVVMExtensionPtr = Cast<UMVVMWidgetBlueprintExtension_View>(NewExtension))
+	{
+		UWidgetBlueprint* WidgetBlueprint = MVVMExtensionPtr->GetWidgetBlueprint();
+
+		if (WidgetBlueprint)
+		{
+			WidgetBlueprint->OnExtensionAdded.RemoveAll(this);
+
+			MVVMExtension = MVVMExtensionPtr;
+
+			if (!BlueprintViewChangedDelegateHandle.IsValid())
+			{
+				BlueprintViewChangedDelegateHandle = MVVMExtensionPtr->OnBlueprintViewChangedDelegate().AddSP(this, &SMVVMViewBindingPanel::HandleBlueprintViewChangedDelegate);
+			}
+
+			if (MVVMExtensionPtr->GetBlueprintView() == nullptr)
+			{
+				MVVMExtensionPtr->CreateBlueprintViewInstance();
+			}
+
+			HandleBlueprintViewChangedDelegate();
+		}
+	}
+}
+
+TSharedRef<SWidget> SMVVMViewBindingPanel::GenerateSettingsMenu()
+{
+	UMVVMWidgetBlueprintExtension_View* MVVMExtensionPtr = MVVMExtension.Get();
+	DetailsView->SetObject(MVVMExtensionPtr && MVVMExtensionPtr->GetBlueprintView() ? MVVMExtensionPtr->GetBlueprintView() : nullptr);
+	DetailContainer->SetContent(DetailsView.ToSharedRef());
+
+	UWidgetBlueprintToolMenuContext* WidgetBlueprintMenuContext = NewObject<UWidgetBlueprintToolMenuContext>();
+	WidgetBlueprintMenuContext->WidgetBlueprintEditor = WeakBlueprintEditor;
+
+	FToolMenuContext MenuContext(WidgetBlueprintMenuContext);
+	return UToolMenus::Get()->GenerateWidget("MVVMEditor.Panel.Toolbar.Settings", MenuContext);
+}
+
+
+void SMVVMViewBindingPanel::RegisterSettingsMenu()
+{
+	UToolMenu* Menu = UToolMenus::Get()->RegisterMenu("MVVMEditor.Panel.Toolbar.Settings");
+}
+
+
+TSharedRef<SWidget> SMVVMViewBindingPanel::GenerateCreateViewWidget()
+{
+	return SNew(SBorder)
+		[
+			SNew(SHorizontalBox)
+
+			+ SHorizontalBox::Slot()
+			.HAlign(EHorizontalAlignment::HAlign_Center)
+			.VAlign(EVerticalAlignment::VAlign_Center)
+			[
+				SNew(SButton)
+				.OnClicked(this, &SMVVMViewBindingPanel::HandleCreateViewClicked)
+				.Text(LOCTEXT("CreateViewButtonText", "Create View"))
+			]
+		];
+}
+
+
+TSharedRef<SWidget> SMVVMViewBindingPanel::GenerateEditViewWidget()
+{
+	FSlimHorizontalToolBarBuilder ToolbarBuilderGlobal(TSharedPtr<const FUICommandList>(), FMultiBoxCustomization::None);
+	{
+		ToolbarBuilderGlobal.BeginSection("Binding");
+		ToolbarBuilderGlobal.AddToolBarButton(
+			FUIAction(
+				FExecuteAction::CreateSP(this, &SMVVMViewBindingPanel::AddDefaultBinding),
+				FCanExecuteAction::CreateSP(this, &SMVVMViewBindingPanel::CanAddBinding),
+				FGetActionCheckState()
+			),
+			NAME_None,
+			LOCTEXT("AddBinding", "Add Binding"),
+			LOCTEXT("AddBindingTooltip", "Add an empty Binding"),
+			FSlateIcon(/*FWidgetReflectorStyle::GetStyleSetName(), WidgetReflectorIcon::TakeSnapshot*/),
+			EUserInterfaceActionType::Button
+		);
+		ToolbarBuilderGlobal.EndSection();
+	}
+
+	{
+		ToolbarBuilderGlobal.BeginSection("ManageViewModels");
+		ToolbarBuilderGlobal.AddToolBarButton(
+			FUIAction(
+				FExecuteAction::CreateSP(this, &SMVVMViewBindingPanel::ShowManageViewModelsWindow),
+				FCanExecuteAction(),
+				FGetActionCheckState()
+			),
+			NAME_None,
+			LOCTEXT("ManageViewModels", "Manage ViewModels"),
+			LOCTEXT("ManageViewModelsTooltip", "Manage ViewModels for this widget"),
+			FSlateIcon(),
+			EUserInterfaceActionType::Button
+		);
+		ToolbarBuilderGlobal.EndSection();
+	}
+
+	{
+		ToolbarBuilderGlobal.AddWidget(SNew(SSpacer), NAME_None, true, HAlign_Right);
+	}
+
+	{
+		ToolbarBuilderGlobal.BeginSection("Options");
+		ToolbarBuilderGlobal.AddComboButton(
+			FUIAction(
+				FExecuteAction(),
+				FCanExecuteAction(),
+				FGetActionCheckState()
+			),
+			FOnGetContent::CreateSP(this, &SMVVMViewBindingPanel::GenerateSettingsMenu),
+			LOCTEXT("Settings", "Settings"),
+			LOCTEXT("SettingsTooltip", "ModelView Settings"),
+			FSlateIcon(FEditorStyle::GetStyleSetName(), "LevelEditor.GameSettings"),
+			false
+		);
+		ToolbarBuilderGlobal.EndSection();
+	}
+
+	return SNew(SVerticalBox)
+
+		+ SVerticalBox::Slot()
+		.AutoHeight()
+		.Padding(FMargin(0.0f, 2.0f))
+		[
+			ToolbarBuilderGlobal.MakeWidget()
+		]
+
+		+ SVerticalBox::Slot()
+		.FillHeight(1.0f)
+		[
+			SNew(SBorder)
+			.BorderImage(FEditorStyle::GetBrush("ToolPanel.GroupBorder"))
+			.Padding(FMargin(4.0f, 4.0f))
+			[
+				SNew(SSplitter)
+				.Orientation(EOrientation::Orient_Horizontal)
+
+				+ SSplitter::Slot()
+				.Value(0.75f)
+				[
+					SAssignNew(ListView, SMVVMViewBindingListView, SharedThis(this), MVVMExtension.Get())
+				]
+
+				+ SSplitter::Slot()
+				.Value(0.25f)
+				[
+					SAssignNew(DetailContainer, SBorder)
+					[
+						DetailsView.ToSharedRef()
+					]
+				]
+			]
+		];
+}
+
+
+FReply SMVVMViewBindingPanel::HandleCreateViewClicked()
+{
+	if (TSharedPtr<FWidgetBlueprintEditor> WidgetBlueprintEditor = WeakBlueprintEditor.Pin())
+	{
+		UWidgetBlueprint* WidgetBlueprint = WidgetBlueprintEditor->GetWidgetBlueprintObj();
+		check(WidgetBlueprint);
+
+		WidgetBlueprint->OnExtensionAdded.RemoveAll(this);
+		UMVVMWidgetBlueprintExtension_View* MVVMExtensionPtr = UMVVMWidgetBlueprintExtension_View::RequestExtension<UMVVMWidgetBlueprintExtension_View>(WidgetBlueprint);
+		check(MVVMExtensionPtr);
+		MVVMExtension = MVVMExtensionPtr;
+
+		if (!BlueprintViewChangedDelegateHandle.IsValid())
+		{
+			BlueprintViewChangedDelegateHandle = MVVMExtensionPtr->OnBlueprintViewChangedDelegate().AddSP(this, &SMVVMViewBindingPanel::HandleBlueprintViewChangedDelegate);
+		}
+
+		if (MVVMExtensionPtr->GetBlueprintView() == nullptr)
+		{
+			MVVMExtensionPtr->CreateBlueprintViewInstance();
+		}
+
+		HandleBlueprintViewChangedDelegate();
+	}
+
+	return FReply::Handled();
+}
+
+#undef LOCTEXT_NAMESPACE
\ No newline at end of file
diff --git a/Engine/Plugins/Runtime/ModelViewViewModel/Source/ModelViewViewModelEditor/Private/Widgets/SMVVMViewBindingPanel.h b/Engine/Plugins/Runtime/ModelViewViewModel/Source/ModelViewViewModelEditor/Private/Widgets/SMVVMViewBindingPanel.h
new file mode 100644
--- /dev/null
+++ b/Engine/Plugins/Runtime/ModelViewViewModel/Source/ModelViewViewModelEditor/Private/Widgets/SMVVMViewBindingPanel.h
@@ -0,0 +1,65 @@
+// Copyright Epic Games, Inc. All Rights Reserved.
+
+#pragma once
+
+#include "CoreMinimal.h"
+
+#include "Widgets/DeclarativeSyntaxSupport.h"
+#include "Widgets/SCompoundWidget.h"
+#include "Layout/Visibility.h"
+#include "Input/Reply.h"
+
+class FWidgetBlueprintEditor;
+class IDetailsView;
+class IStructureDetailsView;
+class SBorder;
+class SMVVMViewBindingListView;
+class UBlueprintExtension;
+class UMVVMWidgetBlueprintExtension_View;
+
+class SMVVMViewBindingPanel : public SCompoundWidget
+{
+public:
+	SLATE_BEGIN_ARGS(SMVVMViewBindingPanel) {}
+	SLATE_END_ARGS()
+
+	void Construct(const FArguments& InArgs, TSharedPtr<FWidgetBlueprintEditor>, bool bInIsDrawerTab);
+	virtual ~SMVVMViewBindingPanel();
+
+	//~ Begin SWidget Interface
+	virtual FReply OnKeyDown(const FGeometry& MyGeometry, const FKeyEvent& InKeyEvent) override;
+	virtual bool SupportsKeyboardFocus() const override;
+	//~ End SWidget Interface
+
+	void OnBindingListSelectionChanged(int32 Index);
+
+	static void RegisterSettingsMenu();
+
+private:
+	void HandleBlueprintViewChangedDelegate();
+
+	TSharedRef<SWidget> GenerateCreateViewWidget();
+	FReply HandleCreateViewClicked();
+
+	TSharedRef<SWidget> GenerateEditViewWidget();
+	void AddDefaultBinding();
+	bool CanAddBinding() const;
+	TSharedRef<SWidget> GenerateSettingsMenu();
+
+	void ShowManageViewModelsWindow();
+
+	TSharedRef<SWidget> CreateDrawerDockButton();
+	FReply CreateDrawerDockButtonClicked();
+
+	void HandleExtensionAdded(TObjectPtr<UBlueprintExtension> NewExtension);
+
+private:
+	TWeakPtr<FWidgetBlueprintEditor> WeakBlueprintEditor;
+	TSharedPtr<SMVVMViewBindingListView> ListView;
+	TSharedPtr<SBorder> DetailContainer;
+	TSharedPtr<IDetailsView> DetailsView;
+	TSharedPtr<IStructureDetailsView> StructDetailsView;
+	TWeakObjectPtr<UMVVMWidgetBlueprintExtension_View> MVVMExtension;
+	FDelegateHandle BlueprintViewChangedDelegateHandle;
+	bool bIsDrawerTab;
+};
diff --git a/Engine/Plugins/Runtime/ModelViewViewModel/Source/ModelViewViewModelEditor/Private/Widgets/SMVVMViewModelBindingListWidget.cpp b/Engine/Plugins/Runtime/ModelViewViewModel/Source/ModelViewViewModelEditor/Private/Widgets/SMVVMViewModelBindingListWidget.cpp
new file mode 100644
--- /dev/null
+++ b/Engine/Plugins/Runtime/ModelViewViewModel/Source/ModelViewViewModelEditor/Private/Widgets/SMVVMViewModelBindingListWidget.cpp
@@ -0,0 +1,261 @@
+// Copyright Epic Games, Inc. All Rights Reserved.
+
+#include "SMVVMViewModelBindingListWidget.h"
+
+#include "Bindings/MVVMBindingHelper.h"
+#include "Editor.h"
+#include "Engine/Engine.h"
+#include "MVVMEditorSubsystem.h"
+#include "MVVMSubsystem.h"
+#include "MVVMViewModelBase.h"
+#include "Styling/SlateIconFinder.h"
+#include "Textures/SlateIcon.h"
+#include "Types/MVVMAvailableBinding.h"
+#include "Types/MVVMFieldVariant.h"
+#include "Widgets/Images/SImage.h"
+#include "Widgets/SMVVMFieldIcon.h"
+#include "Widgets/Views/STreeView.h"
+
+namespace UE::MVVM::Private
+{
+	class FMVVMAvailableBindingNode
+	{
+	public:
+		FMVVMAvailableBindingNode(const FMVVMAvailableBinding& InBinding, TSubclassOf<UMVVMViewModelBase> InClass)
+			:
+			Binding(InBinding),
+			Class(InClass),
+			bIsRoot(false)
+		{}
+
+		FMVVMAvailableBindingNode(TSubclassOf<UMVVMViewModelBase> InClass)
+			:
+			Class(InClass),
+			bIsRoot(true)
+		{}
+
+		const FMVVMAvailableBinding Binding;
+		const TSubclassOf<UMVVMViewModelBase> Class;
+		const bool bIsRoot = false;
+
+		TArray<TSharedPtr<FMVVMAvailableBindingNode>> ChildNodes;
+
+		TSubclassOf<UMVVMViewModelBase> GetViewModelClassProperty()
+		{
+			if (bIsRoot)
+			{
+				return nullptr;
+			}
+
+			UE::MVVM::FMVVMFieldVariant FieldVariant = UE::MVVM::BindingHelper::FindFieldByName(Class, Binding.GetBindingName());
+			if (FieldVariant.IsProperty())
+			{
+				FProperty* Property = FieldVariant.GetProperty();
+
+				if (FObjectProperty* ObjectProperty = CastField<FObjectProperty>(Property))
+				{
+					if (ObjectProperty->PropertyClass->IsChildOf(UMVVMViewModelBase::StaticClass()))
+					{
+						return ObjectProperty->PropertyClass;
+					}
+				}
+			}
+			
+			return nullptr;
+		}
+
+		FText GetDisplayName()
+		{
+			if (bIsRoot)
+			{
+				return Class->GetDisplayNameText();
+			}
+
+			return FText::FromString(Binding.GetBindingName().ToString());
+		}
+
+		FText GetToolTip()
+		{
+			if (bIsRoot)
+			{
+				return FText::GetEmpty();
+			}
+
+			UE::MVVM::FMVVMFieldVariant FieldVariant = UE::MVVM::BindingHelper::FindFieldByName(Class, Binding.GetBindingName());
+			if (FieldVariant.IsProperty())
+			{
+				FProperty* Property = FieldVariant.GetProperty();
+				check(Property);
+				return FText::FormatOrdered(INVTEXT("{0}\n{1}"), FText::FromString(Property->GetCPPType()), Property->GetMetaDataText("ToolTip"));
+			}
+
+			if (FieldVariant.IsFunction())
+			{
+				UFunction* Function = FieldVariant.GetFunction();
+				check(Function);
+				FProperty* Property = Function->GetReturnProperty();
+				if (!Property)
+				{
+					return FText::FormatOrdered(INVTEXT("void\n{0}"), Function->GetMetaDataText("ToolTip"));
+				}
+				return FText::FormatOrdered(INVTEXT("{0}\n{1}"), FText::FromString(Property->GetCPPType()), Function->GetMetaDataText("ToolTip"));
+			}
+
+			return FText::GetEmpty();
+		}
+
+		TSharedRef<SWidget> GetIcon()
+		{
+			if (bIsRoot)
+			{
+				return SNew(SImage).Image(FSlateIconFinder::FindIconBrushForClass(Class));
+			}
+
+			UE::MVVM::FMVVMFieldVariant FieldVariant = UE::MVVM::BindingHelper::FindFieldByName(Class, Binding.GetBindingName());
+			
+			return SNew(SMVVMFieldIcon)
+				.Field(FieldVariant);
+		}
+
+		static bool Sort(const TSharedPtr<FMVVMAvailableBindingNode>& Node, const TSharedPtr<FMVVMAvailableBindingNode>& OtherNode)
+		{
+			TSubclassOf<UMVVMViewModelBase> NodeClass = Node->GetViewModelClassProperty();
+			TSubclassOf<UMVVMViewModelBase> OtherNodeClass = OtherNode->GetViewModelClassProperty();
+			const FString& NodeStr = Node->Binding.GetBindingName().ToString();
+			const FString& OtherNodeStr = OtherNode->Binding.GetBindingName().ToString();
+			if (NodeClass && OtherNodeClass)
+			{
+				return NodeStr < OtherNodeStr;
+			}
+
+			if (NodeClass)
+			{
+				return true;
+			}
+
+			if (OtherNodeClass)
+			{
+				return false;
+			}
+
+			return NodeStr < OtherNodeStr;
+		}
+	};
+
+	static TArray<FMVVMAvailableBinding> GetBindingsList(TSubclassOf<UMVVMViewModelBase> Class)
+	{
+		TArray<FMVVMAvailableBinding> ViewModelAvailableBindingsList = GEngine->GetEngineSubsystem<UMVVMSubsystem>()->GetViewModelAvailableBindings(Class);
+		const TArray<UE::MVVM::FMVVMConstFieldVariant> ChildViewModels = GEditor->GetEditorSubsystem<UMVVMEditorSubsystem>()->GetChildViewModels(Class);
+		for (const UE::MVVM::FMVVMConstFieldVariant& Field : ChildViewModels)
+		{
+			ViewModelAvailableBindingsList.Add(FMVVMAvailableBinding(FMVVMBindingName(Field.GetName()), true, false));
+		}
+
+		return ViewModelAvailableBindingsList;
+	}
+} //namespace UE::MVVM::Private
+
+void SMVVMViewModelBindingListWidget::Construct(const FArguments& InArgs)
+{
+	using UE::MVVM::Private::FMVVMAvailableBindingNode;
+	ChildSlot
+	[
+		SAssignNew(TreeWidget, STreeView<TSharedPtr<FMVVMAvailableBindingNode>>)
+		.TreeItemsSource(&TreeSource)
+		.OnGetChildren(this, &SMVVMViewModelBindingListWidget::HandleGetChildren)
+		.OnGenerateRow(this, &SMVVMViewModelBindingListWidget::HandleGenerateRow)
+		.ItemHeight(20.0f)
+	];
+}
+
+void SMVVMViewModelBindingListWidget::SetViewModel(TSubclassOf<UMVVMViewModelBase> ViewModelClass)
+{
+	using UE::MVVM::Private::FMVVMAvailableBindingNode;
+	TreeSource.Empty();
+	if (ViewModelClass)
+	{
+		TSharedRef<FMVVMAvailableBindingNode> RootNode = MakeShared<FMVVMAvailableBindingNode>(ViewModelClass);
+		TArray<FMVVMAvailableBinding> ViewModelAvailableBindingsList = UE::MVVM::Private::GetBindingsList(ViewModelClass);
+		for (const FMVVMAvailableBinding& Binding : ViewModelAvailableBindingsList)
+		{
+			RootNode->ChildNodes.Add(MakeShared<FMVVMAvailableBindingNode>(Binding, ViewModelClass));
+		}
+		RootNode->ChildNodes.Sort(FMVVMAvailableBindingNode::Sort);
+		TreeSource.Add(RootNode);
+		TreeWidget->SetItemExpansion(RootNode, true);
+	}
+	TreeWidget->RequestListRefresh();
+}
+
+TSharedRef<ITableRow> SMVVMViewModelBindingListWidget::HandleGenerateRow(TSharedPtr<FMVVMAvailableBindingNode> Item, const TSharedRef<STableViewBase>& OwnerTable)
+{
+	using UE::MVVM::Private::FMVVMAvailableBindingNode;
+	typedef STableRow<TSharedPtr<FMVVMAvailableBindingNode>> RowType;
+
+	TSharedRef<RowType> NewRow = SNew(RowType, OwnerTable);
+
+
+	NewRow->SetContent(
+		SNew(SHorizontalBox)
+			.ToolTipText(Item->GetToolTip())
+			+ SHorizontalBox::Slot()
+			.AutoWidth()
+			.HAlign(HAlign_Right)
+			.VAlign(VAlign_Center)
+			[
+				Item->GetIcon()
+			]
+			+ SHorizontalBox::Slot()
+			.Padding(4.0f)
+			[
+				SNew(STextBlock)
+				.Visibility(EVisibility::Visible)
+				.Text(Item->GetDisplayName())
+			]
+	);
+
+	return NewRow;
+}
+
+void SMVVMViewModelBindingListWidget::HandleGetChildren(TSharedPtr<FMVVMAvailableBindingNode> InParent, TArray<TSharedPtr<FMVVMAvailableBindingNode>>& OutChildren)
+{
+	using UE::MVVM::Private::FMVVMAvailableBindingNode;
+
+	if (!InParent->ChildNodes.IsEmpty())
+	{
+		OutChildren = InParent->ChildNodes;
+		return;
+	}
+
+	UE::MVVM::FMVVMFieldVariant FieldVariant = UE::MVVM::BindingHelper::FindFieldByName(InParent->Class, InParent->Binding.GetBindingName());
+
+	if (FieldVariant.IsProperty())
+	{
+		FProperty* Property = FieldVariant.GetProperty();
+
+		if (FObjectProperty* ObjectProperty = CastField<FObjectProperty>(Property))
+		{
+			if (ObjectProperty->PropertyClass->IsChildOf(UMVVMViewModelBase::StaticClass()))
+			{
+				if (InParent->ChildNodes.IsEmpty())
+				{
+					TSubclassOf<UMVVMViewModelBase> ChildClass = ObjectProperty->PropertyClass;
+					TArray<FMVVMAvailableBinding> ViewModelAvailableBindingsList = UE::MVVM::Private::GetBindingsList(ChildClass);
+					for (const FMVVMAvailableBinding& Binding : ViewModelAvailableBindingsList)
+					{
+						TSharedPtr<FMVVMAvailableBindingNode> ChildNode = MakeShared<FMVVMAvailableBindingNode>(Binding, ChildClass);
+						InParent->ChildNodes.Add(ChildNode);
+					}
+					InParent->ChildNodes.Sort(FMVVMAvailableBindingNode::Sort);
+				}
+				OutChildren.Append(InParent->ChildNodes);
+			}
+		}
+	}
+}
+
+TSharedRef<SWidget> SMVVMViewModelBindingListWidget::GetFieldIcon(const UE::MVVM::FMVVMFieldVariant& FieldVariant)
+{
+	return SNew(SMVVMFieldIcon)
+		.Field(FieldVariant);
+}
diff --git a/Engine/Plugins/Runtime/ModelViewViewModel/Source/ModelViewViewModelEditor/Private/Widgets/SMVVMViewModelBindingListWidget.h b/Engine/Plugins/Runtime/ModelViewViewModel/Source/ModelViewViewModelEditor/Private/Widgets/SMVVMViewModelBindingListWidget.h
new file mode 100644
--- /dev/null
+++ b/Engine/Plugins/Runtime/ModelViewViewModel/Source/ModelViewViewModelEditor/Private/Widgets/SMVVMViewModelBindingListWidget.h
@@ -0,0 +1,47 @@
+// Copyright Epic Games, Inc. All Rights Reserved.
+
+#pragma once
+
+#include "CoreMinimal.h"
+#include "Templates/SubclassOf.h"
+#include "Widgets/SCompoundWidget.h"
+
+struct FMVVMAvailableBinding;
+
+template<typename T>
+class STreeView;
+class ITableRow;
+class STableViewBase;
+class UMVVMViewModelBase;
+namespace UE::MVVM
+{
+struct FMVVMFieldVariant;
+namespace Private
+{
+	class FMVVMAvailableBindingNode;
+}
+}
+
+
+class SMVVMViewModelBindingListWidget : public SCompoundWidget
+{
+public:
+	SLATE_BEGIN_ARGS(SMVVMViewModelBindingListWidget) {}
+	SLATE_END_ARGS()
+
+	void Construct(const FArguments& InArgs);
+	void SetViewModel(TSubclassOf<UMVVMViewModelBase> ViewModelClass);
+
+private:
+	typedef UE::MVVM::Private::FMVVMAvailableBindingNode FMVVMAvailableBindingNode;
+
+	TSharedRef<ITableRow> HandleGenerateRow(TSharedPtr<FMVVMAvailableBindingNode> Item, const TSharedRef<STableViewBase>& OwnerTable);
+	void HandleGetChildren(TSharedPtr<FMVVMAvailableBindingNode> InParent, TArray<TSharedPtr<FMVVMAvailableBindingNode>>& OutChildren);
+
+	TSharedRef<SWidget> GetFieldIcon(const UE::MVVM::FMVVMFieldVariant& FieldVariant);
+
+	TSharedPtr<STreeView<TSharedPtr<FMVVMAvailableBindingNode>>> TreeWidget;
+	TArray<TSharedPtr<FMVVMAvailableBindingNode>> TreeSource;
+
+	TSubclassOf<UMVVMViewModelBase> CurrentViewModel;
+};
\ No newline at end of file
diff --git a/Engine/Plugins/Runtime/ModelViewViewModel/Source/ModelViewViewModelEditor/Private/Widgets/SMVVMViewModelContextListWidget.cpp b/Engine/Plugins/Runtime/ModelViewViewModel/Source/ModelViewViewModelEditor/Private/Widgets/SMVVMViewModelContextListWidget.cpp
new file mode 100644
--- /dev/null
+++ b/Engine/Plugins/Runtime/ModelViewViewModel/Source/ModelViewViewModelEditor/Private/Widgets/SMVVMViewModelContextListWidget.cpp
@@ -0,0 +1,913 @@
+// Copyright Epic Games, Inc. All Rights Reserved.
+
+#include "SMVVMViewModelContextListWidget.h"
+
+#include "Algo/Transform.h"
+#include "Bindings/MVVMBindingHelper.h"
+#include "Editor.h"
+#include "EditorStyleSet.h"
+#include "Engine/BlueprintGeneratedClass.h"
+#include "Framework/Application/SlateApplication.h"
+#include "Framework/MultiBox/MultiBoxBuilder.h"
+#include "IPropertyAccessEditor.h"
+#include "Kismet2/Kismet2NameValidators.h"
+#include "SPrimaryButton.h"
+#include "SSimpleButton.h"
+#include "Styling/MVVMEditorStyle.h"
+#include "Styling/SlateIconFinder.h"
+#include "UMGStyle.h"
+#include "Widgets/Input/SComboButton.h"
+#include "Widgets/Input/SEditableTextBox.h"
+#include "Widgets/Layout/SSpacer.h"
+#include "Widgets/Layout/SUniformGridPanel.h"
+#include "Widgets/Views/SListView.h"
+#include "WidgetBlueprint.h"
+
+#define LOCTEXT_NAMESPACE "ViewModelContextListWidget"
+
+namespace UE::MVVM::Private
+{
+	static bool IsPropertyTypeChildOf(const FProperty* Property, TSubclassOf<UMVVMViewModelBase> ParentClass)
+	{
+		if (const FObjectPropertyBase* ObjectProperty = CastField<FObjectPropertyBase>(Property))
+		{
+			if (ObjectProperty->PropertyClass->IsChildOf(ParentClass))
+			{
+				return true;
+			}
+		}
+
+		return false;
+	}
+
+	static bool HasBindableViewModelFieldRecursive(UStruct* InStruct, TSubclassOf<UMVVMViewModelBase> InParentClass, TSet<UStruct*>& VisitedSet, uint32 RecursionDepth)
+	{
+		if (RecursionDepth > 10)
+		{
+			return false;
+		}
+
+		if (VisitedSet.Contains(InStruct))
+		{
+			return false;
+		}
+
+		VisitedSet.Add(InStruct);
+
+		for (TFieldIterator<FProperty> PropIt(InStruct, EFieldIteratorFlags::IncludeSuper); PropIt; ++PropIt)
+		{
+			FProperty* Property = *PropIt;
+
+			if (IsPropertyTypeChildOf(Property, InParentClass) && UE::MVVM::BindingHelper::IsValidForSourceBinding(Property))
+			{
+				return true;
+			}
+
+			bool bFoundValidField = false;
+
+			if (const FObjectPropertyBase* ObjectPropertyBase = CastField<FObjectPropertyBase>(Property))
+			{
+				bFoundValidField = HasBindableViewModelFieldRecursive(ObjectPropertyBase->PropertyClass, InParentClass, VisitedSet, RecursionDepth + 1);
+			}
+			else if (const FStructProperty* StructProperty = CastField<FStructProperty>(Property))
+			{
+				bFoundValidField = HasBindableViewModelFieldRecursive(StructProperty->Struct, InParentClass, VisitedSet, RecursionDepth + 1);
+			}
+
+			if (bFoundValidField)
+			{
+				return true;
+			}
+		}
+
+		for (TFieldIterator<UFunction> FuncIt(InStruct, EFieldIteratorFlags::IncludeSuper); FuncIt; ++FuncIt)
+		{
+			UFunction* Function = *FuncIt;
+
+			const FProperty* ReturnProperty = UE::MVVM::BindingHelper::GetReturnProperty(Function);
+
+			if (IsPropertyTypeChildOf(ReturnProperty, InParentClass) && UE::MVVM::BindingHelper::IsValidForSourceBinding(ReturnProperty))
+			{
+				return true;
+			}
+
+			bool bFoundValidField = false;
+
+			if (const FObjectPropertyBase* ObjectPropertyBase = CastField<FObjectPropertyBase>(ReturnProperty))
+			{
+				bFoundValidField = HasBindableViewModelFieldRecursive(ObjectPropertyBase->PropertyClass, InParentClass, VisitedSet, RecursionDepth + 1);
+			}
+			else if (const FStructProperty* StructProperty = CastField<FStructProperty>(ReturnProperty))
+			{
+				bFoundValidField = HasBindableViewModelFieldRecursive(StructProperty->Struct, InParentClass, VisitedSet, RecursionDepth + 1);
+			}
+
+			if (bFoundValidField)
+			{
+				return true;
+			}
+		}
+
+		return false;
+	}
+
+	DECLARE_DELEGATE_OneParam(FOnViewModelContextRemovedDelegate, FGuid)
+	DECLARE_DELEGATE_RetVal_TwoParams(bool, FOnViewModelContextRenamedDelegate, FGuid, FText);
+
+	class SMVVMManageViewModelsListEntryRow : public SMultiColumnTableRow<TSharedPtr<FMVVMBlueprintViewModelContext>>
+	{
+	public:
+		static FName RemoveButtonColumnName;
+		static FName ClassColumnName;
+		static FName ContextIdColumnName;
+		static FName CreationTypeColumnName;
+		static FName CreationGetterColumnName;
+
+	public:
+		SLATE_BEGIN_ARGS(SMVVMManageViewModelsListEntryRow) {}
+			SLATE_ARGUMENT(UBlueprintGeneratedClass* , OwningWidget)
+			SLATE_ARGUMENT(UWidgetBlueprint* , WidgetBlueprint)
+			SLATE_EVENT(FOnViewModelContextRenamedDelegate, OnViewModelContextRenamed)
+			SLATE_EVENT(FOnViewModelContextRemovedDelegate, OnViewModelContextRemoved)
+		SLATE_END_ARGS()
+
+		void Construct(const FArguments& InArgs, const TSharedRef<STableViewBase>& OwnerTableView, const TSharedPtr<FMVVMBlueprintViewModelContext>& InEntry)
+		{
+			Entry = InEntry;
+			OwningWidget = InArgs._OwningWidget;
+			OnViewModelContextRenamed = InArgs._OnViewModelContextRenamed;
+			OnViewModelContextRemoved = InArgs._OnViewModelContextRemoved;
+
+			NameValidator = MakeUnique<FKismetNameValidator>(InArgs._WidgetBlueprint, Entry->GetViewModelName(), OwningWidget);
+
+			SMultiColumnTableRow<TSharedPtr<FMVVMBlueprintViewModelContext>>::Construct(
+				FSuperRowType::FArguments()
+				.Padding(1.0f),
+				OwnerTableView
+			);
+
+			UpdateGetterContainer();
+		}
+
+		virtual TSharedRef<SWidget> GenerateWidgetForColumn(const FName& InColumnName) override
+		{
+			SMVVMManageViewModelsListEntryRow* Self = this;
+			if (InColumnName == RemoveButtonColumnName)
+			{
+				return SNew(SSimpleButton)
+					.ToolTipText(LOCTEXT("RemoveViewModelContextButtonToolTip", "Remove ViewModelContext"))
+					.Icon(FSlateIcon(FMVVMEditorStyle::Get().GetStyleSetName(), "ViewModelSelection.RemoveIcon").GetIcon())
+					.OnClicked_Lambda(
+						[Self]() {
+							Self->OnViewModelContextRemoved.ExecuteIfBound(Self->Entry->GetViewModelId());
+							return FReply::Handled();
+						}
+					);
+			}
+			else if (InColumnName == ClassColumnName)
+			{
+				return SNew(SBox)
+					.Padding(FMargin(2.0f, 0.0f))
+					.VAlign(VAlign_Center)
+					[
+						SNew(STextBlock)
+						.Text(Entry->GetViewModelClass()->GetDisplayNameText())
+					];
+			}
+			else if (InColumnName == ContextIdColumnName)
+			{
+				TSharedPtr<FMVVMBlueprintViewModelContext> TempEntry = Entry;
+
+				return SNew(SBox)
+					.Padding(FMargin(2.0f, 0.0f))
+					[
+						SNew(SEditableTextBox)
+						.Text(Entry->GetDisplayName())
+						.OnTextCommitted_Lambda(
+							[TempEntry](const FText& NewText, ETextCommit::Type CommitType)
+							{
+								TempEntry->OverrideDisplayName = NewText;
+							})
+						.OnVerifyTextChanged_Lambda(
+							[Self](const FText& InText, FText& OutErrorMessage) -> bool
+							{
+								if (InText.IsEmptyOrWhitespace())
+								{
+									OutErrorMessage = LOCTEXT("ViewModelContextIdEmptyOrWhitespaceErrorMsg", "The ContextId cannot be empty or have whitespaces.");
+									return false;
+								}
+								for (TCHAR Char : InText.ToString())
+								{
+									if (FText::IsWhitespace(Char))
+									{
+										OutErrorMessage = LOCTEXT("ViewModelContextIdWhitespaceErrorMsg", "The ContextId cannot have whitespaces.");
+										return false;
+									}
+								}
+
+								EValidatorResult ValidatorResult = Self->NameValidator->IsValid(InText.ToString());
+								if (ValidatorResult == EValidatorResult::Ok)
+								{
+									if (Self->OnViewModelContextRenamed.IsBound())
+									{
+										if (!Self->OnViewModelContextRenamed.Execute(Self->Entry->GetViewModelId(), InText))
+										{
+											ValidatorResult = EValidatorResult::AlreadyInUse;
+										}
+									}
+								}
+
+								if (ValidatorResult != EValidatorResult::Ok)
+								{
+									OutErrorMessage = Self->NameValidator->GetErrorText(InText.ToString(), ValidatorResult);
+									return false;
+								}
+								return true;
+							})
+					];
+			}
+			else if (InColumnName == CreationTypeColumnName)
+			{
+				FProperty* Property = FMVVMBlueprintViewModelContext::StaticStruct()->FindPropertyByName(TEXT("CreationType"));
+
+				TSharedRef<SComboButton> CreationTypeWidget = 
+					SNew(SComboButton)
+					.OnGetMenuContent(this, &SMVVMManageViewModelsListEntryRow::HandleGetMenuContent_CreationType)
+					.ButtonContent()
+					[
+						SNew(SBox)
+						.Padding(FMargin(2.0f, 0.0f))
+						[
+							SAssignNew(CreationTypeTextBlock, STextBlock)
+							.Text(UEnum::GetDisplayValueAsText(Entry->CreationType))
+						]
+					];
+				return CreationTypeWidget;
+			}
+			else if (InColumnName == CreationGetterColumnName)
+			{
+				return SAssignNew(GetterWidgetContainer, SHorizontalBox);
+			}
+
+			return SNullWidget::NullWidget;
+		}
+
+	private:
+		void HandleAddBinding(TArray<TSharedPtr<FBindingChainElement>> NewBindingChain)
+		{
+			FString PropertyPath;
+
+			for (TSharedPtr<FBindingChainElement>& ChainElement : NewBindingChain)
+			{
+				if (!PropertyPath.IsEmpty())
+				{
+					PropertyPath += TEXT(".");
+				}
+				PropertyPath += ChainElement->Field.GetName();
+			}
+
+			Entry->ViewModelPropertyPath = PropertyPath;
+			PropertyPathTextBlock->SetText(FText::FromString(PropertyPath));
+			PropertyPathTextBlock->SetToolTipText(FText::FromString(PropertyPath));
+		}
+
+		void UpdateGetterContainer()
+		{
+			TSharedPtr<FMVVMBlueprintViewModelContext> TempEntry = Entry;
+			GetterWidgetContainer->ClearChildren();
+
+			if (Entry->CreationType == EMVVMBlueprintViewModelContextCreationType::GlobalViewModelCollection)
+			{
+				GetterWidgetContainer->AddSlot()
+				[
+					SNew(SBox)
+					.Padding(FMargin(2.0f, 0.0f))
+					[
+						SNew(SEditableTextBox)
+						.Text(FText::FromName(Entry->GlobalViewModelIdentifier))
+						.OnTextCommitted_Lambda(
+							[TempEntry](const FText& NewText, ETextCommit::Type CommitType)
+							{
+								TempEntry->GlobalViewModelIdentifier = FName(NewText.ToString());
+							})
+					]
+				];
+			}
+			else if (Entry->CreationType == EMVVMBlueprintViewModelContextCreationType::PropertyPath)
+			{
+				TArray<TSharedPtr<FBindingChainElement>> BindingChain;
+				FMenuBuilder MenuBuilder(true, NULL);
+				GeneratePropertyPathMenuContent(MenuBuilder, OwningWidget, BindingChain);
+
+				GetterWidgetContainer->AddSlot()
+				[
+					SNew(SBox)
+					.Padding(FMargin(2.0f, 0.0f))
+					[
+						SNew(SComboButton)
+						.ButtonContent()
+						[
+							SAssignNew(PropertyPathTextBlock, STextBlock)
+							.Text(FText::FromString(Entry->ViewModelPropertyPath))
+							.ToolTipText(FText::FromString(Entry->ViewModelPropertyPath))
+						]
+						.MenuContent()
+						[
+							MenuBuilder.MakeWidget()
+						]
+					]
+				];
+			}
+
+		}
+
+		TSharedRef<SWidget> HandleGetMenuContent_CreationType()
+		{
+			FMenuBuilder MenuBuilder(true, NULL);
+			SMVVMManageViewModelsListEntryRow* Self = this;
+
+			MenuBuilder.AddMenuEntry(
+				FUIAction(
+					FExecuteAction::CreateLambda([Self]() 
+						{ 
+							Self->Entry->CreationType = EMVVMBlueprintViewModelContextCreationType::CreateInstance;
+							Self->CreationTypeTextBlock->SetText(UEnum::GetDisplayValueAsText(Self->Entry->CreationType));
+							Self->UpdateGetterContainer();
+						})
+				),
+				SNew(STextBlock).Text(LOCTEXT("CreateInstance", "Create new Instance"))
+			);
+
+			MenuBuilder.AddMenuEntry(
+				FUIAction(
+					FExecuteAction::CreateLambda([Self]() 
+						{ 
+							Self->Entry->CreationType = EMVVMBlueprintViewModelContextCreationType::GlobalViewModelCollection; 
+							Self->CreationTypeTextBlock->SetText(UEnum::GetDisplayValueAsText(Self->Entry->CreationType));
+							Self->UpdateGetterContainer();
+						})
+				),
+				SNew(STextBlock).Text(LOCTEXT("GlobalViewModelCollection", "Get from Global collection"))
+			);
+
+			MenuBuilder.AddMenuEntry(
+				FUIAction(
+					FExecuteAction::CreateLambda([Self]() 
+						{ 
+							Self->Entry->CreationType = EMVVMBlueprintViewModelContextCreationType::PropertyPath; 
+							Self->CreationTypeTextBlock->SetText(UEnum::GetDisplayValueAsText(Self->Entry->CreationType));
+							Self->UpdateGetterContainer();
+						})
+				),
+				SNew(STextBlock).Text(LOCTEXT("PropertyPath", "Get from Property Path"))
+			);
+
+			return MenuBuilder.MakeWidget();
+		}
+
+		void GeneratePropertyPathMenuContent(FMenuBuilder& MenuBuilder, UStruct* InOwnerStruct, TArray<TSharedPtr<FBindingChainElement>> InBindingChain)
+		{
+			using namespace UE::MVVM::Private;
+
+			auto MakePropertyWidget = [](FProperty* InProperty)
+			{
+				return SNew(SHorizontalBox)
+					.ToolTipText(InProperty->GetToolTipText())
+					+ SHorizontalBox::Slot()
+					.AutoWidth()
+					[
+						SNew(SSpacer)
+						.Size(FVector2D(18.0f, 0.0f))
+					]
+				+ SHorizontalBox::Slot()
+					.AutoWidth()
+					.VAlign(VAlign_Center)
+					.Padding(1.0f, 0.0f)
+					[
+						SNew(SImage)
+						.Image(FSlateIconFinder::FindIconBrushForClass(UMVVMViewModelBase::StaticClass()))
+					]
+				+ SHorizontalBox::Slot()
+					.AutoWidth()
+					.VAlign(VAlign_Center)
+					.Padding(4.0f, 0.0f)
+					[
+						SNew(STextBlock)
+						.Text(InProperty->GetDisplayNameText())
+					];
+			};
+
+			auto MakeFunctionWidget = [](const UFunction* Info)
+			{
+				return SNew(SHorizontalBox)
+					.ToolTipText(Info->GetMetaDataText("ToolTip"))
+					+ SHorizontalBox::Slot()
+					.AutoWidth()
+					[
+						SNew(SSpacer)
+						.Size(FVector2D(18.0f, 0.0f))
+					]
+				+ SHorizontalBox::Slot()
+					.AutoWidth()
+					.VAlign(VAlign_Center)
+					.Padding(1.0f, 0.0f)
+					[
+						SNew(SImage)
+						.Image(FEditorStyle::Get().GetBrush("GraphEditor.Function_16x"))
+					]
+				+ SHorizontalBox::Slot()
+					.AutoWidth()
+					.VAlign(VAlign_Center)
+					.Padding(4.0f, 0.0f)
+					[
+						SNew(STextBlock)
+						.Text(Info->GetDisplayNameText())
+					];
+			};
+
+			UClass* InBindableClass = Cast<UClass>(InOwnerStruct);
+			TSubclassOf<UMVVMViewModelBase> ViewModelClass = Entry->GetViewModelClass();
+
+			MenuBuilder.BeginSection("Functions", LOCTEXT("Functions", "Functions"));
+			for (TFieldIterator<UFunction> FuncIt(InBindableClass, EFieldIteratorFlags::IncludeSuper); FuncIt; ++FuncIt)
+			{
+				UFunction* Function = *FuncIt;
+
+				TArray<TSharedPtr<FBindingChainElement>> NewBindingChain(InBindingChain);
+				NewBindingChain.Emplace(MakeShared<FBindingChainElement>(Function));
+
+				const FProperty* ReturnProperty = UE::MVVM::BindingHelper::GetReturnProperty(Function);
+
+				if (IsPropertyTypeChildOf(ReturnProperty, ViewModelClass))
+				{
+					MenuBuilder.AddMenuEntry(
+						FUIAction(FExecuteAction::CreateSP(this, &SMVVMManageViewModelsListEntryRow::HandleAddBinding, NewBindingChain)),
+						MakeFunctionWidget(Function)
+					);
+				}
+				else
+				{
+					TSet<UStruct*> VisitedSet;
+					bool bFoundValidField = false;
+					if (const FObjectPropertyBase* ObjectPropertyBase = CastField<FObjectPropertyBase>(ReturnProperty))
+					{
+						if (HasBindableViewModelFieldRecursive(ObjectPropertyBase->PropertyClass, ViewModelClass, VisitedSet, 1))
+						{
+							MenuBuilder.AddSubMenu(
+								MakeFunctionWidget(Function),
+								FNewMenuDelegate::CreateSP(this, &SMVVMManageViewModelsListEntryRow::GeneratePropertyPathMenuContent, static_cast<UStruct*>(ObjectPropertyBase->PropertyClass), NewBindingChain)
+							);
+						}
+					}
+					else if (const FStructProperty* StructProperty = CastField<FStructProperty>(ReturnProperty))
+					{
+						if (HasBindableViewModelFieldRecursive(StructProperty->Struct, ViewModelClass, VisitedSet, 1))
+						{
+							MenuBuilder.AddSubMenu(
+								MakeFunctionWidget(Function),
+								FNewMenuDelegate::CreateSP(this, &SMVVMManageViewModelsListEntryRow::GeneratePropertyPathMenuContent, static_cast<UStruct*>(StructProperty->Struct), NewBindingChain)
+							);
+						}
+					}
+				}
+			}
+			MenuBuilder.EndSection(); //Functions
+
+			MenuBuilder.BeginSection("Properties", LOCTEXT("Properties", "Properties"));
+			for (TFieldIterator<FProperty> PropIt(InBindableClass, EFieldIteratorFlags::IncludeSuper); PropIt; ++PropIt)
+			{
+				FProperty* Property = *PropIt;
+
+				TArray<TSharedPtr<FBindingChainElement>> NewBindingChain(InBindingChain);
+				NewBindingChain.Emplace(MakeShared<FBindingChainElement>(Property));
+
+				if (IsPropertyTypeChildOf(Property, ViewModelClass))
+				{
+					MenuBuilder.AddMenuEntry(
+						FUIAction(FExecuteAction::CreateSP(this, &SMVVMManageViewModelsListEntryRow::HandleAddBinding, NewBindingChain)),
+						MakePropertyWidget(Property)
+					);
+				}
+				else
+				{
+					TSet<UStruct*> VisitedSet;
+					bool bFoundValidField = false;
+					if (const FObjectPropertyBase* ObjectPropertyBase = CastField<FObjectPropertyBase>(Property))
+					{
+						if (HasBindableViewModelFieldRecursive(ObjectPropertyBase->PropertyClass, ViewModelClass, VisitedSet, 1))
+						{
+							MenuBuilder.AddSubMenu(
+								MakePropertyWidget(Property),
+								FNewMenuDelegate::CreateSP(this, &SMVVMManageViewModelsListEntryRow::GeneratePropertyPathMenuContent, static_cast<UStruct*>(ObjectPropertyBase->PropertyClass), NewBindingChain)
+							);
+						}
+					}
+					else if (const FStructProperty* StructProperty = CastField<FStructProperty>(Property))
+					{
+						if (HasBindableViewModelFieldRecursive(StructProperty->Struct, ViewModelClass, VisitedSet, 1))
+						{
+							MenuBuilder.AddSubMenu(
+								MakePropertyWidget(Property),
+								FNewMenuDelegate::CreateSP(this, &SMVVMManageViewModelsListEntryRow::GeneratePropertyPathMenuContent, static_cast<UStruct*>(StructProperty->Struct), NewBindingChain)
+							);
+						}
+					}
+				}
+			}
+			MenuBuilder.EndSection(); //Properties
+		}
+
+		TSharedPtr<FMVVMBlueprintViewModelContext> Entry;
+		UBlueprintGeneratedClass* OwningWidget;
+		TSharedPtr<SHorizontalBox> GetterWidgetContainer;
+		TSharedPtr<STextBlock> PropertyPathTextBlock;
+		TSharedPtr<STextBlock> CreationTypeTextBlock;
+		TUniquePtr<FKismetNameValidator> NameValidator;
+
+		FOnViewModelContextRenamedDelegate OnViewModelContextRenamed;
+		FOnViewModelContextRemovedDelegate OnViewModelContextRemoved;
+	};
+
+	FName SMVVMManageViewModelsListEntryRow::RemoveButtonColumnName = "RemoveButton";
+	FName SMVVMManageViewModelsListEntryRow::ClassColumnName = "ViewModel";
+	FName SMVVMManageViewModelsListEntryRow::ContextIdColumnName = "ContextId";
+	FName SMVVMManageViewModelsListEntryRow::CreationTypeColumnName = "CreationType";
+	FName SMVVMManageViewModelsListEntryRow::CreationGetterColumnName = "Getter";
+
+	class SMVVMViewModelDeleteConfirmationDialog : public SCompoundWidget
+	{
+	public:
+		SLATE_BEGIN_ARGS(SMVVMViewModelDeleteConfirmationDialog) {}
+			SLATE_ARGUMENT(TArray<TSharedPtr<SHorizontalBox>>, ConflictList)
+			SLATE_ARGUMENT(TWeakPtr<SWindow>, ParentWindow)
+		SLATE_END_ARGS()
+
+	public:
+		void Construct(const FArguments& InArgs)
+		{
+			ConflictListSource = InArgs._ConflictList;
+			WeakParentWindow = InArgs._ParentWindow;
+
+			SAssignNew(ConflictListViewWidget, SListView<TSharedPtr<SHorizontalBox>>)
+				.ListItemsSource(&ConflictListSource)
+				.SelectionMode(ESelectionMode::None)
+				.ItemHeight(30.0f)
+				.OnGenerateRow(this, &SMVVMViewModelDeleteConfirmationDialog::HandleGenerateRow);
+		
+			ChildSlot
+			[
+				SNew(SBorder)
+				.BorderImage(FEditorStyle::GetBrush("AssetDeleteDialog.Background"))
+				.Padding(10)
+				[
+					SNew(SVerticalBox)
+					+ SVerticalBox::Slot()
+					.AutoHeight()
+					[
+						SNew(SBorder)
+						.BorderImage(FEditorStyle::GetBrush("ToolPanel.GroupBorder"))
+						.Padding(5.0f)
+						[
+							SNew(STextBlock)
+							.Font( FEditorStyle::GetFontStyle( "BoldFont" ) )
+							.Text(LOCTEXT("References", "This ViewModel is still referenced by the following bindings"))
+							.ShadowOffset( FVector2D( 1.0f, 1.0f ) )
+						]
+					]
+					+ SVerticalBox::Slot()
+					.FillHeight(1.0f)
+					[
+						SNew(SBorder)
+						.BorderImage(FEditorStyle::GetBrush("ToolPanel.GroupBorder"))
+						.Padding(5.0f)
+						[
+							ConflictListViewWidget.ToSharedRef()
+						]
+					]
+					+ SVerticalBox::Slot()
+					.AutoHeight()
+					.Padding(4,4)
+					[
+						SNew(SButton)
+						.HAlign(HAlign_Center)
+						.Text( LOCTEXT("MVVMViewModelForceDelete", "Force Delete ViewModel"))
+						.ToolTipText(LOCTEXT("MVVMViewModelForceDeleteTooltipText", "These bindings will be in an invalid state and you must remove or fix them manually."))
+						.ButtonStyle(FEditorStyle::Get(), "FlatButton.Danger")
+						.TextStyle(FEditorStyle::Get(), "FlatButton.DefaultTextStyle")
+						.OnClicked(this, &SMVVMViewModelDeleteConfirmationDialog::HandleForceDelete)
+					]
+					+ SVerticalBox::Slot()
+					.AutoHeight()
+					.Padding(4,4)
+					[
+						SNew(SButton)
+						.HAlign(HAlign_Center)
+						.Text( LOCTEXT("Cancel", "Cancel"))
+						.ToolTipText(LOCTEXT("CancelDeleteTooltipText", "Cancel the delete"))
+						.ButtonStyle(FEditorStyle::Get(), "FlatButton.Default")
+						.TextStyle(FEditorStyle::Get(), "FlatButton.DefaultTextStyle")
+						.OnClicked(this, &SMVVMViewModelDeleteConfirmationDialog::HandleCancel)
+					]
+				]
+			];
+		}
+
+		bool ShouldDeleteViewModel()
+		{
+			return bShouldDelete;
+		}
+	private:
+		TSharedRef<ITableRow> HandleGenerateRow(TSharedPtr<SHorizontalBox> Item, const TSharedRef<STableViewBase>& OwnerTable)
+		{
+			typedef STableRow<TSharedPtr<SHorizontalBox>> RowType;
+
+			TSharedRef<RowType> NewRow = SNew(RowType, OwnerTable);
+
+			NewRow->SetContent(Item.ToSharedRef());
+
+			return NewRow;
+		}
+
+		FReply HandleCancel()
+		{
+			bShouldDelete = false;
+			if (WeakParentWindow.IsValid())
+			{
+				WeakParentWindow.Pin()->RequestDestroyWindow();
+			}
+			return FReply::Handled();
+		}
+
+		FReply HandleForceDelete()
+		{
+			bShouldDelete = true;
+			if (WeakParentWindow.IsValid())
+			{
+				WeakParentWindow.Pin()->RequestDestroyWindow();
+			}
+			return FReply::Handled();
+		}
+
+	private:
+		TSharedPtr<SListView<TSharedPtr<SHorizontalBox>>> ConflictListViewWidget;
+		TArray<TSharedPtr<SHorizontalBox>> ConflictListSource;
+		bool bShouldDelete = false;
+		TWeakPtr<SWindow> WeakParentWindow;
+	};
+
+	const FSlateBrush* ManageViewModelsGetModeBrush(EMVVMBindingMode BindingMode)
+	{
+		switch (BindingMode)
+		{
+		case EMVVMBindingMode::OneTimeToDestination:
+			return FUMGStyle::Get().GetBrush("MVVM.OneTime");
+		case EMVVMBindingMode::OneWayToDestination:
+			return FUMGStyle::Get().GetBrush("MVVM.OneWay");
+		case EMVVMBindingMode::OneWayToSource:
+			return FUMGStyle::Get().GetBrush("MVVM.OneWayToSource");
+		case EMVVMBindingMode::OneTimeToSource:
+			return nullptr;
+		case EMVVMBindingMode::TwoWay:
+			return FUMGStyle::Get().GetBrush("MVVM.TwoWay");
+		default:
+			return nullptr;
+		}
+	}
+} // namespace UE::MVVM::Private
+
+void SMVVMViewModelContextListWidget::Construct(const FArguments& InArgs)
+{
+	using UE::MVVM::Private::SMVVMManageViewModelsListEntryRow;
+
+	OwningWidget = InArgs._OwningWidget;
+	WidgetBlueprint = InArgs._WidgetBlueprint;
+	Bindings = InArgs._Bindings;
+	OnViewModelContextsUpdated = InArgs._OnViewModelContextsUpdated;
+	WeakParentWindow = InArgs._ParentWindow;
+
+	TSharedRef<SWidget> ButtonsPanelContent = InArgs._ButtonsPanel.Widget;
+	if (InArgs._ButtonsPanel.Widget == SNullWidget::NullWidget && WeakParentWindow.IsValid())
+	{
+		SAssignNew(ButtonsPanelContent, SUniformGridPanel)
+			.SlotPadding(FEditorStyle::Get().GetMargin("StandardDialog.SlotPadding"))
+			+ SUniformGridPanel::Slot(0, 0)
+			[
+				SNew(SPrimaryButton)
+				.Text(LOCTEXT("ViewModelFinishButtonText", "Finish"))
+				.OnClicked(this, &SMVVMViewModelContextListWidget::HandleClicked_Finish)
+			]
+			+ SUniformGridPanel::Slot(1, 0)
+			[
+				SNew(SButton)
+				.Text(LOCTEXT("ViewModelCancelButtonText", "Cancel"))
+				.HAlign(HAlign_Center)
+				.OnClicked(this, &SMVVMViewModelContextListWidget::HandleClicked_Cancel)
+			];
+	}
+
+	for (const FMVVMBlueprintViewModelContext& Context : InArgs._ExistingViewModelContexts)
+	{
+		ContextListSource.Add(MakeShared<FMVVMBlueprintViewModelContext>(Context));
+	}
+
+	ChildSlot
+	[
+		SNew(SVerticalBox)
+		+ SVerticalBox::Slot()
+		.FillHeight(1.0f)
+		.Padding(6)
+		[
+			SAssignNew(ContextListWidget, SListView<TSharedPtr<FMVVMBlueprintViewModelContext>>)
+			.SelectionMode(ESelectionMode::Single)
+			.ListItemsSource(&ContextListSource)
+			.OnGenerateRow(this, &SMVVMViewModelContextListWidget::HandleGenerateRowForListView)
+			.ItemHeight(20.0f)
+			.HeaderRow
+			(
+				SNew(SHeaderRow)
+				+ SHeaderRow::Column(SMVVMManageViewModelsListEntryRow::RemoveButtonColumnName)
+				.FixedWidth(30.0f)
+				.DefaultLabel(LOCTEXT("", "")) 
+				+ SHeaderRow::Column(SMVVMManageViewModelsListEntryRow::ClassColumnName)
+				.DefaultLabel(LOCTEXT("Class", "Class"))
+				+ SHeaderRow::Column(SMVVMManageViewModelsListEntryRow::ContextIdColumnName)
+				.DefaultLabel(LOCTEXT("ContextId", "Context Id"))
+				+ SHeaderRow::Column(SMVVMManageViewModelsListEntryRow::CreationTypeColumnName)
+				.DefaultLabel(LOCTEXT("CreationType", "Creation Type"))
+				+ SHeaderRow::Column(SMVVMManageViewModelsListEntryRow::CreationGetterColumnName)
+				.DefaultLabel(LOCTEXT("Getter", "Getter"))
+			)
+		]
+		+ SVerticalBox::Slot()
+		.HAlign(HAlign_Right)
+		.VAlign(VAlign_Bottom)
+		.AutoHeight()
+		.Padding(8)
+		[
+			ButtonsPanelContent
+		]
+	];
+}
+
+void SMVVMViewModelContextListWidget::AddViewModelContext(TSubclassOf<UMVVMViewModelBase> ViewModelClass)
+{
+	TSet<FString> ExistingContextNames;
+	for (const TSharedPtr<FMVVMBlueprintViewModelContext>& Context : ContextListSource)
+	{
+		ExistingContextNames.Add(Context->GetDisplayName().ToString());
+	}
+
+	TSharedPtr<FMVVMBlueprintViewModelContext> NewContext = MakeShared<FMVVMBlueprintViewModelContext>(ViewModelClass, FGuid::NewGuid());
+	FName TempNewName = ViewModelClass->GetFName();
+	FKismetNameValidator NameValidator(WidgetBlueprint);
+	while ((NameValidator.IsValid(TempNewName) != EValidatorResult::Ok) || (ExistingContextNames.Contains(TempNewName.ToString())))
+	{
+		TempNewName = MakeUniqueObjectName(OwningWidget, ViewModelClass);
+	}
+	NewContext->OverrideDisplayName = FText::FromName(TempNewName);
+	ContextListSource.Add(NewContext);
+	ContextListWidget->RequestListRefresh();
+}
+
+TSharedRef<ITableRow> SMVVMViewModelContextListWidget::HandleGenerateRowForListView(TSharedPtr<FMVVMBlueprintViewModelContext> Item, const TSharedRef<STableViewBase>& OwnerTable)
+{
+	using UE::MVVM::Private::SMVVMManageViewModelsListEntryRow;
+	return SNew(SMVVMManageViewModelsListEntryRow, OwnerTable, Item)
+		.WidgetBlueprint(WidgetBlueprint)
+		.OwningWidget(OwningWidget)
+		.OnViewModelContextRenamed(this, &SMVVMViewModelContextListWidget::IsContextNameAvailable)
+		.OnViewModelContextRemoved(this, &SMVVMViewModelContextListWidget::RemoveViewModelContext);
+}
+
+bool SMVVMViewModelContextListWidget::IsContextNameAvailable(FGuid Guid, FText ContextName)
+{
+	for (TSharedPtr<FMVVMBlueprintViewModelContext>& Context : ContextListSource)
+	{
+		if (Context->GetViewModelId() == Guid)
+		{
+			continue;
+		}
+
+		if (Context->OverrideDisplayName.EqualTo(ContextName))
+		{
+			return false;
+		}
+	}
+
+	return true;
+}
+
+void SMVVMViewModelContextListWidget::RemoveViewModelContext(FGuid Guid)
+{
+	TSharedPtr<FMVVMBlueprintViewModelContext> ContextToRemove;
+	for (TSharedPtr<FMVVMBlueprintViewModelContext>& Context : ContextListSource)
+	{
+		if (Context->GetViewModelId() == Guid)
+		{
+			ContextToRemove = Context;
+			break;
+		}
+	}
+
+	if (ValidateRemoveViewModelContext(ContextToRemove))
+	{
+		ContextListSource.Remove(ContextToRemove);
+		ContextListWidget->RequestListRefresh();
+	}
+}
+
+bool SMVVMViewModelContextListWidget::ValidateRemoveViewModelContext(TSharedPtr<FMVVMBlueprintViewModelContext> ContextToRemove)
+{
+	using UE::MVVM::Private::SMVVMViewModelDeleteConfirmationDialog;
+	TArray<TSharedPtr<SHorizontalBox>> ConflictList;
+	for (const FMVVMBlueprintViewBinding& Binding : Bindings)
+	{
+		if (Binding.ViewModelPath.ContextId == ContextToRemove->GetViewModelId())
+		{
+			TSharedRef<SHorizontalBox> Container = SNew(SHorizontalBox);
+			Container->AddSlot()
+			.Padding(FMargin(3.0f, 0.0f))
+			.AutoWidth()
+			[
+				SNew(STextBlock)
+				.Text(FText::FromString(Binding.ViewModelPath.GetBindingName().ToString()))
+			];
+			Container->AddSlot()
+			.Padding(FMargin(3.0f, 0.0f))
+			.AutoWidth()
+			[
+				SNew(SBox)
+				.HAlign(HAlign_Center)
+				.VAlign(VAlign_Center)
+				.WidthOverride(16)
+				.HeightOverride(16)
+				[
+					SNew(SImage)
+					.Image(UE::MVVM::Private::ManageViewModelsGetModeBrush(Binding.BindingType))
+				]
+			];
+			Container->AddSlot()
+			.Padding(FMargin(3.0f, 0.0f))
+			.AutoWidth()
+			[
+				SNew(STextBlock)
+				.Text(FText::FormatOrdered(INVTEXT("{0}.{1}"), FText::FromString(Binding.WidgetPath.WidgetName.ToString()), FText::FromString(Binding.WidgetPath.GetBindingName().ToString())))
+			];
+			ConflictList.Add(Container);
+		}
+	}
+
+	if (!ConflictList.IsEmpty())
+	{
+		TSharedRef<SWindow> ViewModelDeleteConfirmationWindow = SNew(SWindow)
+			.Title(LOCTEXT("MVVMViewModelDeleteConfirmationWindowHeader", "Delete ViewModel Confirmation"))
+			.SupportsMaximize(false)
+			.ClientSize(FVector2D(400.0f, 300.0f));
+
+		TSharedRef<SMVVMViewModelDeleteConfirmationDialog> DeleteConfirmationDialog = SNew(SMVVMViewModelDeleteConfirmationDialog)
+			.ConflictList(ConflictList)
+			.ParentWindow(ViewModelDeleteConfirmationWindow);
+
+		ViewModelDeleteConfirmationWindow->SetContent(DeleteConfirmationDialog);
+
+		GEditor->EditorAddModalWindow(ViewModelDeleteConfirmationWindow);
+
+		return DeleteConfirmationDialog->ShouldDeleteViewModel();
+	}
+
+	return true;
+}
+
+TArray<FMVVMBlueprintViewModelContext> SMVVMViewModelContextListWidget::GetViewModelContexts()
+{
+	TArray<FMVVMBlueprintViewModelContext> NewContextList;
+	Algo::Transform(ContextListSource, NewContextList, [](TSharedPtr<FMVVMBlueprintViewModelContext> InElement)
+		{
+			return *InElement.Get();
+		});
+
+	return NewContextList;
+}
+
+FReply SMVVMViewModelContextListWidget::HandleClicked_Finish()
+{
+	OnViewModelContextsUpdated.ExecuteIfBound(GetViewModelContexts());
+
+	if (WeakParentWindow.IsValid())
+	{
+		WeakParentWindow.Pin()->RequestDestroyWindow();
+	}
+
+	return FReply::Handled();
+}
+
+FReply SMVVMViewModelContextListWidget::HandleClicked_Cancel()
+{
+	if (WeakParentWindow.IsValid())
+	{
+		WeakParentWindow.Pin()->RequestDestroyWindow();
+	}
+	return FReply::Handled();
+}
+
+#undef LOCTEXT_NAMESPACE
diff --git a/Engine/Plugins/Runtime/ModelViewViewModel/Source/ModelViewViewModelEditor/Private/Widgets/SMVVMViewModelContextListWidget.h b/Engine/Plugins/Runtime/ModelViewViewModel/Source/ModelViewViewModelEditor/Private/Widgets/SMVVMViewModelContextListWidget.h
new file mode 100644
--- /dev/null
+++ b/Engine/Plugins/Runtime/ModelViewViewModel/Source/ModelViewViewModelEditor/Private/Widgets/SMVVMViewModelContextListWidget.h
@@ -0,0 +1,55 @@
+// Copyright Epic Games, Inc. All Rights Reserved.
+
+#pragma once
+
+#include "CoreMinimal.h"
+#include "MVVMBlueprintViewBinding.h"
+#include "MVVMBlueprintViewModelContext.h"
+#include "Widgets/SCompoundWidget.h"
+
+class ITableRow;
+template<typename T>
+class SListView;
+class STableViewBase;
+class UBlueprintGeneratedClass;
+class UWidgetBlueprint;
+
+DECLARE_DELEGATE_OneParam(FOnViewModelContextsUpdated, TArray<FMVVMBlueprintViewModelContext>);
+
+class SMVVMViewModelContextListWidget : public SCompoundWidget
+{
+public:
+	SLATE_BEGIN_ARGS(SMVVMViewModelContextListWidget) {}
+		SLATE_NAMED_SLOT(FArguments, ButtonsPanel)
+		SLATE_ARGUMENT(TSharedPtr<SWindow>, ParentWindow)
+		SLATE_ARGUMENT(UBlueprintGeneratedClass*, OwningWidget)
+		SLATE_ARGUMENT(UWidgetBlueprint*, WidgetBlueprint)
+		SLATE_ARGUMENT(TArrayView<const FMVVMBlueprintViewBinding>, Bindings)
+		SLATE_ARGUMENT(TArray<FMVVMBlueprintViewModelContext>, ExistingViewModelContexts)
+		SLATE_ARGUMENT(FOnViewModelContextsUpdated, OnViewModelContextsUpdated)
+	SLATE_END_ARGS()
+
+	void Construct(const FArguments& InArgs);
+	void AddViewModelContext(TSubclassOf<UMVVMViewModelBase> ViewModelClass);
+	TArray<FMVVMBlueprintViewModelContext> GetViewModelContexts();
+
+private:
+	TSharedRef<ITableRow> HandleGenerateRowForListView(TSharedPtr<FMVVMBlueprintViewModelContext> Item, const TSharedRef<STableViewBase>& OwnerTable);
+	bool IsContextNameAvailable(FGuid Guid, FText ContextName);
+	void RemoveViewModelContext(FGuid Guid);
+	bool ValidateRemoveViewModelContext(TSharedPtr<FMVVMBlueprintViewModelContext> ContextToRemove);
+
+	FReply HandleClicked_Finish();
+	FReply HandleClicked_Cancel();
+
+	TArray<TSharedPtr<FMVVMBlueprintViewModelContext>> ContextListSource;
+	TSharedPtr<SListView<TSharedPtr<FMVVMBlueprintViewModelContext>>> ContextListWidget;
+	TArray<FMVVMBlueprintViewModelContext> ExistingViewModelContexts;
+	UBlueprintGeneratedClass* OwningWidget = nullptr;
+	UWidgetBlueprint* WidgetBlueprint = nullptr;
+	FOnViewModelContextsUpdated OnViewModelContextsUpdated;
+	TArrayView<const FMVVMBlueprintViewBinding> Bindings;
+	TWeakPtr<SWindow> WeakParentWindow;
+
+	friend class SMVVMViewModelSelectionWidget;
+};
\ No newline at end of file
diff --git a/Engine/Plugins/Runtime/ModelViewViewModel/Source/ModelViewViewModelEditor/Public/MVVMEditorSubsystem.h b/Engine/Plugins/Runtime/ModelViewViewModel/Source/ModelViewViewModelEditor/Public/MVVMEditorSubsystem.h
new file mode 100644
--- /dev/null
+++ b/Engine/Plugins/Runtime/ModelViewViewModel/Source/ModelViewViewModelEditor/Public/MVVMEditorSubsystem.h
@@ -0,0 +1,40 @@
+// Copyright Epic Games, Inc. All Rights Reserved.
+
+#pragma once
+
+#include "EditorSubsystem.h"
+
+#include "MVVMBlueprintView.h"
+#include "Templates/SubclassOf.h"
+#include "Types/MVVMFieldVariant.h"
+
+#include "MVVMEditorSubsystem.generated.h"
+
+class UMVVMViewModelBase;
+class UWidgetBlueprint;
+
+/** */
+UCLASS()
+class MODELVIEWVIEWMODELEDITOR_API UMVVMEditorSubsystem : public UEditorSubsystem
+{
+	GENERATED_BODY()
+
+public:
+	UFUNCTION(BlueprintCallable, Category="MVVM")
+	UMVVMBlueprintView* RequestView(UWidgetBlueprint* WidgetBlueprint) const;
+
+	UFUNCTION(BlueprintCallable, Category = "MVVM")
+	UMVVMBlueprintView* GetView(UWidgetBlueprint* WidgetBlueprint) const;
+	
+	UFUNCTION(BlueprintCallable, Category = "MVVM")
+	void RemoveViewModel(UWidgetBlueprint* WidgetBlueprint, const FMVVMBlueprintViewModelContext& ViewModel);
+
+	UFUNCTION(BlueprintCallable, Category = "MVVM")
+	FMVVMBlueprintViewBinding& AddBinding(UWidgetBlueprint* WidgetBlueprint);
+
+	UFUNCTION(BlueprintCallable, Category = "MVVM")
+	void RemoveBinding(UWidgetBlueprint* WidgetBlueprint, const FMVVMBlueprintViewBinding& Binding);
+
+	TArray<UE::MVVM::FMVVMConstFieldVariant> GetChildViewModels(TSubclassOf<UMVVMViewModelBase> Class);
+	TArray<const UFunction*> GetAvailableConversionFunctions(const UE::MVVM::FMVVMConstFieldVariant& Source, const UE::MVVM::FMVVMConstFieldVariant& Dest) const;
+};
