Fix missing include.
#rnx
#rn none

diff --git a/Engine/Plugins/Runtime/ModelViewViewModel/Source/ModelViewViewModelBlueprint/Public/Bindings/MVVMCompiledBindingLibraryCompiler.h b/Engine/Plugins/Runtime/ModelViewViewModel/Source/ModelViewViewModelBlueprint/Public/Bindings/MVVMCompiledBindingLibraryCompiler.h
--- a/Engine/Plugins/Runtime/ModelViewViewModel/Source/ModelViewViewModelBlueprint/Public/Bindings/MVVMCompiledBindingLibraryCompiler.h
+++ b/Engine/Plugins/Runtime/ModelViewViewModel/Source/ModelViewViewModelBlueprint/Public/Bindings/MVVMCompiledBindingLibraryCompiler.h
@@ -6,6 +6,7 @@
 #include "Bindings/MVVMCompiledBindingLibrary.h"
 #include "MVVMSubsystem.h"
 #include "Templates/PimplPtr.h"
+#include "Templates/TypeHash.h"
 #include "Templates/ValueOrError.h"
 #include "Types/MVVMFieldVariant.h"
 
@@ -55,7 +56,7 @@ public:
 			return Id != Other.Id;
 		}
 
-		friend uint32 GetTypeHash(const FFieldPathHandle Handle)
+		friend uint32 GetTypeHash(const FFieldPathHandle& Handle)
 		{
 			return GetTypeHash(Handle.Id);
 		}
@@ -98,7 +99,7 @@ public:
 			return Id != Other.Id;
 		}
 
-		friend uint32 GetTypeHash(const FBindingHandle Handle)
+		friend uint32 GetTypeHash(const FBindingHandle& Handle)
 		{
 			return GetTypeHash(Handle.Id);
 		}
@@ -140,7 +141,7 @@ public:
 			return Id != Other.Id;
 		}
 
-		friend uint32 GetTypeHash(const FFieldIdHandle Handle)
+		friend uint32 GetTypeHash(const FFieldIdHandle& Handle)
 		{
 			return GetTypeHash(Handle.Id);
 		}
