PR #9047: Using scope resolution operator to fix MVVM compilation (Contributed by [REDACTED_PERSON])

diff --git a/Engine/Plugins/Runtime/ModelViewViewModel/Source/ModelViewViewModelBlueprint/Public/Bindings/MVVMCompiledBindingLibraryCompiler.h b/Engine/Plugins/Runtime/ModelViewViewModel/Source/ModelViewViewModelBlueprint/Public/Bindings/MVVMCompiledBindingLibraryCompiler.h
--- a/Engine/Plugins/Runtime/ModelViewViewModel/Source/ModelViewViewModelBlueprint/Public/Bindings/MVVMCompiledBindingLibraryCompiler.h
+++ b/Engine/Plugins/Runtime/ModelViewViewModel/Source/ModelViewViewModelBlueprint/Public/Bindings/MVVMCompiledBindingLibraryCompiler.h
@@ -58,7 +58,7 @@ public:
 
 		friend uint32 GetTypeHash(const FFieldPathHandle& Handle)
 		{
-			return GetTypeHash(Handle.Id);
+			return ::GetTypeHash(Handle.Id);
 		}
 
 	private:
@@ -101,7 +101,7 @@ public:
 
 		friend uint32 GetTypeHash(const FBindingHandle& Handle)
 		{
-			return GetTypeHash(Handle.Id);
+			return ::GetTypeHash(Handle.Id);
 		}
 
 	private:
@@ -143,7 +143,7 @@ public:
 
 		friend uint32 GetTypeHash(const FFieldIdHandle& Handle)
 		{
-			return GetTypeHash(Handle.Id);
+			return ::GetTypeHash(Handle.Id);
 		}
 
 	private:
