MVVM: Enable the bExposeOnSpawn flag on viewmodel property. Make sure the setter for viewmodel has the correct name.

diff --git a/Engine/Plugins/Runtime/ModelViewViewModel/Source/ModelViewViewModelBlueprint/Private/MVVMViewBlueprintCompiler.cpp b/Engine/Plugins/Runtime/ModelViewViewModel/Source/ModelViewViewModelBlueprint/Private/MVVMViewBlueprintCompiler.cpp
--- a/Engine/Plugins/Runtime/ModelViewViewModel/Source/ModelViewViewModelBlueprint/Private/MVVMViewBlueprintCompiler.cpp
+++ b/Engine/Plugins/Runtime/ModelViewViewModel/Source/ModelViewViewModelBlueprint/Private/MVVMViewBlueprintCompiler.cpp
@@ -15,17 +15,20 @@
 
 namespace UE::MVVM::Private
 {
+
 void FMVVMViewBlueprintCompiler::AddErrorForBinding(FMVVMBlueprintViewBinding& Binding, const FString& Message) const
 {
 	WidgetBlueprintCompilerContext.MessageLog.Error(*Message);
 	Binding.Errors.Add(FText::FromString(Message));
 }
 
+
 void FMVVMViewBlueprintCompiler::AddExtension(UWidgetBlueprintGeneratedClass* Class, UMVVMViewClass* ViewExtension)
 {
 	WidgetBlueprintCompilerContext.AddExtension(Class, ViewExtension);
 }
 
+
 void FMVVMViewBlueprintCompiler::CleanOldData(UWidgetBlueprintGeneratedClass* ClassToClean, UObject* OldCDO)
 {
 	// Clean old View
@@ -262,7 +265,7 @@ void FMVVMViewBlueprintCompiler::CreateSourceLists(const FWidgetBlueprintCompile
 		SourceVariable.PropertyName = ViewModelContext.GetViewModelName();
 		SourceVariable.DisplayName = ViewModelContext.GetDisplayName().ToString();
 		SourceVariable.CategoryName = TEXT("Viewmodel");
-		//SourceVariable.bExposeOnSpawn = ViewModelContext.bCreateSetterFunction;
+		SourceVariable.bExposeOnSpawn = ViewModelContext.bCreateSetterFunction;
 		SourceVariable.BlueprintSetter = SourceCreatorContexts[FoundSourceContextIndex].SetterFunctionName;
 		SourceContexts.Emplace(MoveTemp(SourceVariable));
 	}
@@ -355,7 +358,12 @@ void FMVVMViewBlueprintCompiler::CreateFunctionsDeclaration(const FWidgetBluepri
 				, TEXT("Viewmodel")
 				, false);
 
-			UE::MVVM::FunctionGraphHelper::AddFunctionArgument(SourceCreator.SetterGraph, SourceCreator.ViewModelContext.GetViewModelClass().Get(), "ViewModel");
+			if (SourceCreator.SetterGraph->GetFName() != FName(*SourceCreator.SetterFunctionName))
+			{
+				WidgetBlueprintCompilerContext.MessageLog.Error(*FString::Printf(TEXT("The setter name %s already exist and could not be autogenerated."), *SourceCreator.SetterFunctionName));
+			}
+
+			UE::MVVM::FunctionGraphHelper::AddFunctionArgument(SourceCreator.SetterGraph, SourceCreator.ViewModelContext.GetViewModelClass().Get(), "Viewmodel");
 		}
 	}
 }
@@ -385,7 +393,7 @@ bool FMVVMViewBlueprintCompiler::Compile(UWidgetBlueprintGeneratedClass* Class,
 	TValueOrError<FCompiledBindingLibraryCompiler::FCompileResult, FString> CompileResult = BindingLibraryCompiler.Compile();
 	if (CompileResult.HasError())
 	{
-		WidgetBlueprintCompilerContext.MessageLog.Error(TEXT("The binding compilation failed. %s"), *CompileResult.GetError());
+		WidgetBlueprintCompilerContext.MessageLog.Error(*FString::Printf(TEXT("The binding compilation failed. %s"), *CompileResult.GetError()));
 		return false;
 	}
 	CompileSourceCreators(CompileResult.GetValue(), Class, BlueprintView, ViewExtension);
