Fix warning for alloca in loop
#rnx

diff --git a/Engine/Plugins/Runtime/ModelViewViewModel/Source/ModelViewViewModel/Private/Bindings/MVVMCompiledBindingLibrary.cpp b/Engine/Plugins/Runtime/ModelViewViewModel/Source/ModelViewViewModel/Private/Bindings/MVVMCompiledBindingLibrary.cpp
--- a/Engine/Plugins/Runtime/ModelViewViewModel/Source/ModelViewViewModel/Private/Bindings/MVVMCompiledBindingLibrary.cpp
+++ b/Engine/Plugins/Runtime/ModelViewViewModel/Source/ModelViewViewModel/Private/Bindings/MVVMCompiledBindingLibrary.cpp
@@ -218,6 +218,20 @@ TValueOrError<void, FMVVMCompiledBindingLibrary::EExecutionFailingReason> FMVVMC
 }
 
 
+namespace UE::MVVM::Private
+{
+	UObject* EvaluateFieldPathObject(UObject* Container, UFunction* Function, const FObjectPropertyBase* ReturnObjectProperty)
+	{
+		void* DataPtr = FMemory_Alloca_Aligned(Function->ParmsSize, Function->GetMinAlignment());
+		ReturnObjectProperty->InitializeValue(DataPtr);
+		Container->ProcessEvent(Function, DataPtr);
+		UObject* NewContainer = ReturnObjectProperty->GetObjectPropertyValue(DataPtr);
+		ReturnObjectProperty->DestroyValue(DataPtr);
+		return NewContainer;
+	}
+}
+
+
 TValueOrError<UE::MVVM::FFieldContext, void> FMVVMCompiledBindingLibrary::EvaluateFieldPath(UObject* InExecutionSource, const FMVVMVCompiledFieldPath& InFieldPath) const
 {
 	if (InExecutionSource == nullptr)
@@ -302,12 +316,7 @@ TValueOrError<UE::MVVM::FFieldContext, void> FMVVMCompiledBindingLibrary::Evalua
 					return MakeError();
 				}
 
-				void* DataPtr = FMemory_Alloca_Aligned(Function->ParmsSize, Function->GetMinAlignment());
-				ReturnObjectProperty->InitializeValue(DataPtr);
-				CurrentContainer.GetUObject()->ProcessEvent(Function, DataPtr);
-				UObject* NewContainer = ReturnObjectProperty->GetObjectPropertyValue(DataPtr);
-				CurrentContainer.SetUObject(NewContainer);
-				ReturnObjectProperty->DestroyValue(DataPtr);
+				CurrentContainer.SetUObject(UE::MVVM::Private::EvaluateFieldPathObject(CurrentContainer.GetUObject(), Function, ReturnObjectProperty));
 			}
 			else if (PathIndex.bIsScriptStructProperty)
 			{
