Make MVVM PropertyPath reference getters / setters via FMemberReference.

diff --git a/Engine/Plugins/Runtime/ModelViewViewModel/Source/ModelViewViewModel/Private/MVVMSubsystem.cpp b/Engine/Plugins/Runtime/ModelViewViewModel/Source/ModelViewViewModel/Private/MVVMSubsystem.cpp
--- a/Engine/Plugins/Runtime/ModelViewViewModel/Source/ModelViewViewModel/Private/MVVMSubsystem.cpp
+++ b/Engine/Plugins/Runtime/ModelViewViewModel/Source/ModelViewViewModel/Private/MVVMSubsystem.cpp
@@ -145,6 +145,7 @@ namespace UE::MVVM::Private
 		FName NAME_BlueprintSetter = "BlueprintSetter";
 #endif
 
+		// Remove BP getters / setters and replace them with their property for bindings
 		for (TFieldIterator<const FProperty> PropertyIt(Class, EFieldIteratorFlags::IncludeSuper); PropertyIt; ++PropertyIt)
 		{
 			const FProperty* Property = *PropertyIt;
@@ -153,7 +154,6 @@ namespace UE::MVVM::Private
 			if (bIsReadable || bIsWritable)
 			{
 #if WITH_EDITOR
-				// Remove the accessor and replace it by the property
 				if (bIsReadable)
 				{
 					const FString& PropertyGetter = Property->GetMetaData(NAME_BlueprintGetter);
diff --git a/Engine/Plugins/Runtime/ModelViewViewModel/Source/ModelViewViewModel/Public/Types/MVVMFieldVariant.h b/Engine/Plugins/Runtime/ModelViewViewModel/Source/ModelViewViewModel/Public/Types/MVVMFieldVariant.h
--- a/Engine/Plugins/Runtime/ModelViewViewModel/Source/ModelViewViewModel/Public/Types/MVVMFieldVariant.h
+++ b/Engine/Plugins/Runtime/ModelViewViewModel/Source/ModelViewViewModel/Public/Types/MVVMFieldVariant.h
@@ -3,6 +3,7 @@
 #pragma once
 
 #include "CoreMinimal.h"
+#include "Engine/MemberReference.h"
 #include "Misc/TVariant.h"
 #include <type_traits>
 
@@ -12,7 +13,9 @@ class UFunction;
 namespace UE::MVVM
 {
 
-	/** */
+	/**
+	 * Represents a possibly-const binding to either a UFunction or FProperty
+	 */
 	template<bool bConst>
 	struct TMVVMFieldVariant
 	{
@@ -129,6 +132,26 @@ namespace UE::MVVM
 			return !(*this == B);
 		}
 
+		operator FMemberReference() const
+		{
+			FMemberReference BindingReference = FMemberReference();
+			if (IsProperty())
+			{
+				// @TODO: DarenC - Should self context ever be true? Seems no, but if issues occur:
+				// analyze UBlueprintVariableNodeSpawner::CreateFromMemberOrParam -> MemberVarSetupLambda
+				PropertyType* Property = GetProperty();
+				bool bSelfContext = false;
+				BindingReference.SetFromField<FProperty>(Property, bSelfContext);
+			}
+			else if (IsFunction())
+			{
+				FunctionType* Function = GetFunction();
+				bool bSelfContext = false;
+				BindingReference.SetFromField<UFunction>(Function, bSelfContext);
+			}
+			return BindingReference;
+		}
+
 		friend int32 GetTypeHash(const TMVVMFieldVariant<bConst>& Variant)
 		{
 			if (Variant.IsProperty())
diff --git a/Engine/Plugins/Runtime/ModelViewViewModel/Source/ModelViewViewModelBlueprint/Private/MVVMBlueprintView.cpp b/Engine/Plugins/Runtime/ModelViewViewModel/Source/ModelViewViewModelBlueprint/Private/MVVMBlueprintView.cpp
--- a/Engine/Plugins/Runtime/ModelViewViewModel/Source/ModelViewViewModelBlueprint/Private/MVVMBlueprintView.cpp
+++ b/Engine/Plugins/Runtime/ModelViewViewModel/Source/ModelViewViewModelBlueprint/Private/MVVMBlueprintView.cpp
@@ -122,7 +122,7 @@ FMVVMBlueprintViewBinding& UMVVMBlueprintView::AddBinding(const UWidget* Widget,
 {
 	FMVVMBlueprintViewBinding& NewBinding = Bindings.AddDefaulted_GetRef();
 	NewBinding.WidgetPath.WidgetName = Widget->GetFName();
-	NewBinding.WidgetPath.SetBindingName(Property->GetFName());
+	NewBinding.WidgetPath.SetBindingReference(UE::MVVM::FMVVMConstFieldVariant(Property));
 
 	OnBindingsUpdated.Broadcast();
 	return NewBinding;
diff --git a/Engine/Plugins/Runtime/ModelViewViewModel/Source/ModelViewViewModelBlueprint/Public/MVVMPropertyPath.h b/Engine/Plugins/Runtime/ModelViewViewModel/Source/ModelViewViewModelBlueprint/Public/MVVMPropertyPath.h
--- a/Engine/Plugins/Runtime/ModelViewViewModel/Source/ModelViewViewModelBlueprint/Public/MVVMPropertyPath.h
+++ b/Engine/Plugins/Runtime/ModelViewViewModel/Source/ModelViewViewModelBlueprint/Public/MVVMPropertyPath.h
@@ -2,12 +2,16 @@
 
 #pragma once
 
+#include "Blueprint/WidgetBlueprintGeneratedClass.h"
 #include "Types/MVVMBindingName.h"
+#include "Types/MVVMFieldVariant.h"
 
 #include "MVVMPropertyPath.generated.h"
 
 /**
- *
+ * Base path to properties for MVVM view models and widgets.
+ * 
+ * Used to associate properties within MVVM bindings in editor & during MVVM compilation
  */
 USTRUCT()
 struct FMVVMPropertyPathBase
@@ -15,9 +19,9 @@ struct FMVVMPropertyPathBase
 	GENERATED_BODY()
 
 private:
-	/** Name we are looking for. */
+	/** Reference to property for this binding. */
 	UPROPERTY(EditAnywhere, Category = "MVVM")
-	FMVVMBindingName BindingName;
+	FMemberReference BindingReference;
 
 	/** The path of the property (either the getter or the set action). */
 	UPROPERTY(EditAnywhere, Category = "MVVM")
@@ -27,30 +31,91 @@ private:
 	UPROPERTY(EditAnywhere, Category = "MVVM")
 	FString SetterPropertyPath;
 
+	/** If we are referencing a UFunction or FProperty */
+	UPROPERTY()
+	EBindingKind BindingKind;
+
 public:
+	/** Get the binding name, resolves reference deprecation / redirectors / etc before returning */
 	FMVVMBindingName GetBindingName() const
 	{
-		return BindingName;
+		// Resolve any redirectors 
+		FName Path = FName();
+		if (!BindingReference.GetMemberName().IsNone())
+		{
+			if (BindingKind == EBindingKind::Property && BindingReference.ResolveMember<FProperty>())
+			{
+				Path = BindingReference.GetMemberName();
+			}
+			else if (BindingKind == EBindingKind::Function && BindingReference.ResolveMember<UFunction>())
+			{
+				Path = BindingReference.GetMemberName();
+			}
+		}
+
+		return FMVVMBindingName(Path);
+	}
+
+	FMemberReference GetBindingReference() const
+	{
+		// Note: Prefer copy since FName & FGuid, FString should be empty based on our usage
+		return BindingReference;
+	}
+
+	void SetBindingReference(const UE::MVVM::FMVVMFieldVariant& InField)
+	{
+		BindingReference = InField;
+
+		if (InField.IsProperty())
+		{
+			BindingKind = EBindingKind::Property;
+		}
+		else if (InField.IsFunction())
+		{
+			BindingKind = EBindingKind::Function;
+		}
+		else
+		{
+			ensureAlwaysMsgf(false, TEXT("Binding to field of unknown type!"));
+		}
+	}
+
+	void SetBindingReference(const UE::MVVM::FMVVMConstFieldVariant& InField)
+	{
+		BindingReference = InField;
+
+		if (InField.IsProperty())
+		{
+			BindingKind = EBindingKind::Property;
+		}
+		else if (InField.IsFunction())
+		{
+			BindingKind = EBindingKind::Function;
+		}
+		else
+		{
+			ensureAlwaysMsgf(false, TEXT("Binding to field of unknown type!"));
+		}
 	}
 
-	void SetBindingName(FName InBindingName)
+	void Reset()
 	{
-		BindingName = FMVVMBindingName(InBindingName);
+		BindingReference = FMemberReference();
 	}
 
 	FString GetGetterPropertyPath() const
 	{
-		return !GetterPropertyPath.IsEmpty() ? GetterPropertyPath : BindingName.ToString();
+		return !GetterPropertyPath.IsEmpty() ? GetterPropertyPath : GetBindingName().ToString();
 	}
 
 	FString GetSetterPropertyPath() const
 	{
-		return !SetterPropertyPath.IsEmpty() ? SetterPropertyPath : BindingName.ToString();
+		return !SetterPropertyPath.IsEmpty() ? SetterPropertyPath : GetBindingName().ToString();
 	}
 
 	bool operator==(const FMVVMPropertyPathBase& Other) const
 	{
-		return BindingName == Other.BindingName &&
+		return BindingReference.GetMemberName() == Other.BindingReference.GetMemberName() &&
 			GetterPropertyPath == Other.GetterPropertyPath &&
 			SetterPropertyPath == Other.SetterPropertyPath;
 	}
diff --git a/Engine/Plugins/Runtime/ModelViewViewModel/Source/ModelViewViewModelEditor/Private/Customizations/MVVMPropertyBindingExtension.cpp b/Engine/Plugins/Runtime/ModelViewViewModel/Source/ModelViewViewModelEditor/Private/Customizations/MVVMPropertyBindingExtension.cpp
--- a/Engine/Plugins/Runtime/ModelViewViewModel/Source/ModelViewViewModelEditor/Private/Customizations/MVVMPropertyBindingExtension.cpp
+++ b/Engine/Plugins/Runtime/ModelViewViewModel/Source/ModelViewViewModelEditor/Private/Customizations/MVVMPropertyBindingExtension.cpp
@@ -76,15 +76,15 @@ static void ExtendBindingsMenu(FMenuBuilder& MenuBuilder, const UWidgetBlueprint
 		return HorizontalBox;
 	};
 
-	auto CreateBinding = [MVVMBlueprintView](const UWidget* Widget, const FProperty* Property, FGuid ViewModelId, const FProperty* ViewModelProperty)
+	auto CreateBinding = [MVVMBlueprintView](const UWidget* Widget, const FProperty* WidgetProperty, FGuid ViewModelId, const FProperty* ViewModelProperty)
 	{
 		FMVVMBlueprintViewBinding& NewBinding = MVVMBlueprintView->AddDefaultBinding();
 
 		NewBinding.ViewModelPath.ContextId = ViewModelId;
-		NewBinding.ViewModelPath.SetBindingName(ViewModelProperty->GetFName());
+		NewBinding.ViewModelPath.SetBindingReference(UE::MVVM::FMVVMConstFieldVariant(ViewModelProperty));
 
 		NewBinding.WidgetPath.WidgetName = Widget->GetFName();
-		NewBinding.WidgetPath.SetBindingName(Property->GetFName());
+		NewBinding.WidgetPath.SetBindingReference(UE::MVVM::FMVVMConstFieldVariant(WidgetProperty));
 
 		NewBinding.BindingType = EMVVMBindingMode::OneWayToDestination;
 	};
@@ -172,7 +172,7 @@ void FMVVMPropertyBindingExtension::ClearCurrentValue(const UWidgetBlueprint* Wi
 		{
 			if (FMVVMBlueprintViewBinding* Binding = MVVMBlueprintView->FindBinding(Widget, Property))
 			{
-				Binding->ViewModelPath.SetBindingName(FName());
+				Binding->ViewModelPath.Reset();
 			}
 		}
 	}
diff --git a/Engine/Plugins/Runtime/ModelViewViewModel/Source/ModelViewViewModelEditor/Private/Customizations/MVVMPropertyPathCustomization.cpp b/Engine/Plugins/Runtime/ModelViewViewModel/Source/ModelViewViewModelEditor/Private/Customizations/MVVMPropertyPathCustomization.cpp
--- a/Engine/Plugins/Runtime/ModelViewViewModel/Source/ModelViewViewModelEditor/Private/Customizations/MVVMPropertyPathCustomization.cpp
+++ b/Engine/Plugins/Runtime/ModelViewViewModel/Source/ModelViewViewModelEditor/Private/Customizations/MVVMPropertyPathCustomization.cpp
@@ -185,7 +185,7 @@ namespace UE::MVVM
 		for (const TSharedPtr<IFieldPathHelper>& Path : PathHelpers)
 		{
 			Path->SetSelectedSource(Selected);
-			Path->SetBindingName(FName()); // Might make sense to keep this around in case we retarget to a compatible widget or switch back.
+			Path->ResetBinding(); // Might make sense to keep this around in case we retarget to a compatible widget or switch back.
 		}
 
 		PropertyHandle->NotifyPostChange(EPropertyChangeType::ValueSet);
@@ -209,7 +209,7 @@ namespace UE::MVVM
 		TArray<IFieldPathHelper*> Helpers = GetRawPathHelpers();
 		for (const IFieldPathHelper* Helper : Helpers)
 		{
-			Helper->SetBindingName(Selected.GetName());
+			Helper->SetBindingReference(Selected);
 		}
 
 		PropertyHandle->NotifyPostChange(EPropertyChangeType::ValueSet);
diff --git a/Engine/Plugins/Runtime/ModelViewViewModel/Source/ModelViewViewModelEditor/Private/MVVMPropertyPathHelpers.cpp b/Engine/Plugins/Runtime/ModelViewViewModel/Source/ModelViewViewModelEditor/Private/MVVMPropertyPathHelpers.cpp
--- a/Engine/Plugins/Runtime/ModelViewViewModel/Source/ModelViewViewModelEditor/Private/MVVMPropertyPathHelpers.cpp
+++ b/Engine/Plugins/Runtime/ModelViewViewModel/Source/ModelViewViewModelEditor/Private/MVVMPropertyPathHelpers.cpp
@@ -164,11 +164,37 @@ namespace UE::MVVM
 		return FMVVMBindingName();
 	}
 
-	void FWidgetFieldPathHelper::SetBindingName(FName BindingName) const
+	FMemberReference FWidgetFieldPathHelper::GetBindingReference() const
 	{
 		if (FMVVMWidgetPropertyPath* Path = PathAttr.Get(nullptr))
 		{
-			Path->SetBindingName(BindingName);
+			return Path->GetBindingReference();
+		}
+
+		return FMemberReference();
+	}
+
+	void FWidgetFieldPathHelper::SetBindingReference(const UE::MVVM::FMVVMFieldVariant& InField) const
+	{
+		if (FMVVMWidgetPropertyPath* Path = PathAttr.Get(nullptr))
+		{
+			Path->SetBindingReference(InField);
+		}
+	}
+
+	void FWidgetFieldPathHelper::SetBindingReference(const UE::MVVM::FMVVMConstFieldVariant& InField) const
+	{
+		if (FMVVMWidgetPropertyPath* Path = PathAttr.Get(nullptr))
+		{
+			Path->SetBindingReference(InField);
+		}
+	}
+
+	void FWidgetFieldPathHelper::ResetBinding() const
+	{
+		if (FMVVMWidgetPropertyPath* Path = PathAttr.Get(nullptr))
+		{
+			Path->Reset();
 		}
 	}
 
@@ -284,11 +310,37 @@ namespace UE::MVVM
 		return FMVVMBindingName();
 	}
 
-	void FViewModelFieldPathHelper::SetBindingName(FName BindingName) const
+	FMemberReference FViewModelFieldPathHelper::GetBindingReference() const
+	{
+		if (FMVVMViewModelPropertyPath* Path = PathAttr.Get(nullptr))
+		{
+			return Path->GetBindingReference();
+		}
+
+		return FMemberReference();
+	}
+
+	void FViewModelFieldPathHelper::SetBindingReference(const UE::MVVM::FMVVMFieldVariant& InField) const
+	{
+		if (FMVVMViewModelPropertyPath* Path = PathAttr.Get(nullptr))
+		{
+			Path->SetBindingReference(InField);
+		}
+	}
+
+	void FViewModelFieldPathHelper::SetBindingReference(const UE::MVVM::FMVVMConstFieldVariant& InField) const
+	{
+		if (FMVVMViewModelPropertyPath* Path = PathAttr.Get(nullptr))
+		{
+			Path->SetBindingReference(InField);
+		}
+	}
+
+	void FViewModelFieldPathHelper::ResetBinding() const
 	{
 		if (FMVVMViewModelPropertyPath* Path = PathAttr.Get(nullptr))
 		{
-			Path->SetBindingName(BindingName);
+			Path->Reset();
 		}
 	}
 }
diff --git a/Engine/Plugins/Runtime/ModelViewViewModel/Source/ModelViewViewModelEditor/Private/MVVMPropertyPathHelpers.h b/Engine/Plugins/Runtime/ModelViewViewModel/Source/ModelViewViewModelEditor/Private/MVVMPropertyPathHelpers.h
--- a/Engine/Plugins/Runtime/ModelViewViewModel/Source/ModelViewViewModelEditor/Private/MVVMPropertyPathHelpers.h
+++ b/Engine/Plugins/Runtime/ModelViewViewModel/Source/ModelViewViewModelEditor/Private/MVVMPropertyPathHelpers.h
@@ -26,7 +26,10 @@ namespace UE::MVVM
 		virtual void SetSelectedSource(const TOptional<FBindingSource>& Source) const = 0;
 
 		virtual FMVVMBindingName GetBindingName() const = 0;
-		virtual void SetBindingName(FName BindingName) const = 0;
+		virtual FMemberReference GetBindingReference() const = 0;
+		virtual void SetBindingReference(const UE::MVVM::FMVVMFieldVariant& InField) const = 0;
+		virtual void SetBindingReference(const UE::MVVM::FMVVMConstFieldVariant& InField) const = 0;
+		virtual void ResetBinding() const = 0;
 	};
 
 	class FWidgetFieldPathHelper : public IFieldPathHelper
@@ -48,7 +51,10 @@ namespace UE::MVVM
 		virtual void SetSelectedSource(const TOptional<FBindingSource>& Source) const override;
 
 		virtual FMVVMBindingName GetBindingName() const override;
-		virtual void SetBindingName(FName BindingName) const override;
+		virtual FMemberReference GetBindingReference() const override;
+		virtual void SetBindingReference(const UE::MVVM::FMVVMFieldVariant& InField) const override;
+		virtual void SetBindingReference(const UE::MVVM::FMVVMConstFieldVariant& InField) const override;
+		virtual void ResetBinding() const override;
 
 	private:
 		TAttribute<FMVVMWidgetPropertyPath*> PathAttr;
@@ -74,7 +80,10 @@ namespace UE::MVVM
 		virtual void SetSelectedSource(const TOptional<FBindingSource>& Source) const override;
 
 		virtual FMVVMBindingName GetBindingName() const override;
-		virtual void SetBindingName(FName BindingName) const override;
+		virtual FMemberReference GetBindingReference() const override;
+		virtual void SetBindingReference(const UE::MVVM::FMVVMFieldVariant& InField) const override;
+		virtual void SetBindingReference(const UE::MVVM::FMVVMConstFieldVariant& InField) const override;
+		virtual void ResetBinding() const override;
 
 	private:
 		TAttribute<FMVVMViewModelPropertyPath*> PathAttr;
diff --git a/Engine/Plugins/Runtime/ModelViewViewModel/Source/ModelViewViewModelEditor/Private/Widgets/SMVVMViewBindingListView.cpp b/Engine/Plugins/Runtime/ModelViewViewModel/Source/ModelViewViewModelEditor/Private/Widgets/SMVVMViewBindingListView.cpp
--- a/Engine/Plugins/Runtime/ModelViewViewModel/Source/ModelViewViewModelEditor/Private/Widgets/SMVVMViewBindingListView.cpp
+++ b/Engine/Plugins/Runtime/ModelViewViewModel/Source/ModelViewViewModelEditor/Private/Widgets/SMVVMViewBindingListView.cpp
@@ -379,7 +379,7 @@ private:
 			BlueprintViewPtr->PreEditChange(UMVVMBlueprintView::StaticClass()->FindPropertyByName("Bindings"));
 
 			PathHelper->SetSelectedSource(SelectedSource);
-			PathHelper->SetBindingName(FName()); // Might make sense to keep this around in case we retarget to a compatible widget or switch back.
+			PathHelper->ResetBinding(); // Might make sense to keep this around in case we retarget to a compatible widget or switch back.
 
 			BlueprintViewPtr->PostEditChange();
 		}
@@ -419,7 +419,7 @@ private:
 		{
 			BlueprintViewPtr->PreEditChange(UMVVMBlueprintView::StaticClass()->FindPropertyByName("Bindings"));
 
-			PathHelper->SetBindingName(SelectedField.GetName());
+			PathHelper->SetBindingReference(SelectedField);
 
 			BlueprintViewPtr->PostEditChange();
 		}
