Removed static function limitation on conversion functions, which allows us to use functions defined on the UserWidget as conversion functions.

Formatting of errors changed to include the binding name (SourceObject.SourceProperty <-/<->/-> DestObject.DestProperty). Errors are now stored in the binding object again.

Renamed GetConversionFunctionPath to SourceToDestinationFunctionPath and SetConversionFunctionPath to DestinationToSourceFunctionPath.

diff --git a/Engine/Plugins/Runtime/ModelViewViewModel/Source/ModelViewViewModel/Private/Bindings/MVVMBindingHelper.cpp b/Engine/Plugins/Runtime/ModelViewViewModel/Source/ModelViewViewModel/Private/Bindings/MVVMBindingHelper.cpp
--- a/Engine/Plugins/Runtime/ModelViewViewModel/Source/ModelViewViewModel/Private/Bindings/MVVMBindingHelper.cpp
+++ b/Engine/Plugins/Runtime/ModelViewViewModel/Source/ModelViewViewModel/Private/Bindings/MVVMBindingHelper.cpp
@@ -381,11 +381,6 @@ namespace UE::MVVM::BindingHelper
 			return MakeError(MoveTemp(CommonResult));
 		}
 
-		if (!InFunction->HasAllFunctionFlags(FUNC_Static))
-		{
-			return MakeError(FString::Printf(TEXT("The function '%s' is not static."), *InFunction->GetName()));
-		}
-
 		if (InFunction->NumParms != 2)
 		{
 			return MakeError(FString::Printf(TEXT("The function '%s' does not have the correct number of arguments."), *InFunction->GetName()));
diff --git a/Engine/Plugins/Runtime/ModelViewViewModel/Source/ModelViewViewModel/Private/MVVMSubsystem.cpp b/Engine/Plugins/Runtime/ModelViewViewModel/Source/ModelViewViewModel/Private/MVVMSubsystem.cpp
--- a/Engine/Plugins/Runtime/ModelViewViewModel/Source/ModelViewViewModel/Private/MVVMSubsystem.cpp
+++ b/Engine/Plugins/Runtime/ModelViewViewModel/Source/ModelViewViewModel/Private/MVVMSubsystem.cpp
@@ -278,22 +278,22 @@ TValueOrError<bool, FString> UMVVMSubsystem::IsBindingValid(FConstDirectionalBin
 			return MakeError(ConversionResult.StealError());
 		}
 
-		const FProperty* ConversionReturnProperty = ConversionResult.GetValue().ReturnProperty;
-		const FProperty* ConversionArgProperty = ConversionResult.GetValue().ArgumentProperty;
-
 		// The compiled version should look like Setter(Conversion(Getter())).
+		const FProperty* ConversionReturnProperty = ConversionResult.GetValue().ReturnProperty;
 		if (!UE::MVVM::BindingHelper::ArePropertiesCompatible(ConversionReturnProperty, DestinationProperty))
 		{
-			return MakeError(FString::Printf(TEXT("The Destination '%s' property type does not match the type of the return type from the conversion function '%s'."), *DestinationProperty->GetName(), *ConversionReturnProperty->GetName()));
+			return MakeError(FString::Printf(TEXT("The destination property '%s' (%s) does not match the return type of the conversion function (%s)."), *DestinationProperty->GetName(), *DestinationProperty->GetCPPType(), *ConversionReturnProperty->GetCPPType()));
 		}
+
+		const FProperty* ConversionArgProperty = ConversionResult.GetValue().ArgumentProperty;
 		if (!UE::MVVM::BindingHelper::ArePropertiesCompatible(SourceProperty, ConversionArgProperty))
 		{
-			return MakeError(FString::Printf(TEXT("The Source '%s' property type does not match the type of the argument type of the conversion function '%s'."), *SourceProperty->GetName(), *DestinationProperty->GetName()));
+			return MakeError(FString::Printf(TEXT("The source property '%s' (%s) does not match the argument type of the conversion function (%s)."), *SourceProperty->GetName(), *SourceProperty->GetCPPType(), *DestinationProperty->GetCPPType()));
 		}
 	}
 	else if (!UE::MVVM::BindingHelper::ArePropertiesCompatible(SourceProperty, DestinationProperty))
 	{
-		return MakeError(FString::Printf(TEXT("The Source '%s' property type does not match the type of the Destination property '%s'. A conversion function is required."), *SourceProperty->GetName(), *DestinationProperty->GetName()));
+		return MakeError(FString::Printf(TEXT("The source property '%s' (%s) does not match the type of the destination property '%s' (%s). A conversion function is required."), *SourceProperty->GetName(), *SourceProperty->GetCPPType(), *DestinationProperty->GetName(), *DestinationProperty->GetCPPType()));
 	}
 
 	return MakeValue(true);
diff --git a/Engine/Plugins/Runtime/ModelViewViewModel/Source/ModelViewViewModelBlueprint/Private/MVVMViewBlueprintCompiler.cpp b/Engine/Plugins/Runtime/ModelViewViewModel/Source/ModelViewViewModelBlueprint/Private/MVVMViewBlueprintCompiler.cpp
--- a/Engine/Plugins/Runtime/ModelViewViewModel/Source/ModelViewViewModelBlueprint/Private/MVVMViewBlueprintCompiler.cpp
+++ b/Engine/Plugins/Runtime/ModelViewViewModel/Source/ModelViewViewModelBlueprint/Private/MVVMViewBlueprintCompiler.cpp
@@ -16,10 +16,11 @@
 namespace UE::MVVM::Private
 {
 
-void FMVVMViewBlueprintCompiler::AddErrorForBinding(FMVVMBlueprintViewBinding& Binding, const FString& Message) const
+void FMVVMViewBlueprintCompiler::AddErrorForBinding(FMVVMBlueprintViewBinding& Binding, const UMVVMBlueprintView* View, const FString& Message) const
 {
-	WidgetBlueprintCompilerContext.MessageLog.Error(*Message);
-	Binding.Errors.Add(FText::FromString(Message));
+	const FString BindingName = Binding.GetNameString(View);
+	WidgetBlueprintCompilerContext.MessageLog.Error(*(BindingName + TEXT(": ") + Message));
+	Binding.Errors.Add(FText::FromString(BindingName + TEXT(": ") + Message));
 }
 
 
@@ -569,14 +570,15 @@ bool FMVVMViewBlueprintCompiler::PreCompileBindings(UWidgetBlueprintGeneratedCla
 
 	for (int32 Index = 0; Index < BlueprintView->GetNumBindings(); ++Index)
 	{
-		const FMVVMBlueprintViewBinding* BindingPtr = BlueprintView->GetBindingAt(Index);
+		FMVVMBlueprintViewBinding* BindingPtr = BlueprintView->GetBindingAt(Index);
 		if (BindingPtr == nullptr)
 		{
-			WidgetBlueprintCompilerContext.MessageLog.Error(*FString::Printf(TEXT("The binding number %d is invalid."), Index));
+			WidgetBlueprintCompilerContext.MessageLog.Error(*FString::Printf(TEXT("The binding at index '%d' is invalid."), Index));
 			bAreSourceContextsValid = false;
 			continue;
 		}
-		const FMVVMBlueprintViewBinding& Binding = *BindingPtr;
+
+		FMVVMBlueprintViewBinding& Binding = *BindingPtr;
 		if (!Binding.bCompile)
 		{
 			continue;
@@ -656,9 +658,7 @@ bool FMVVMViewBlueprintCompiler::PreCompileBindings(UWidgetBlueprintGeneratedCla
 			TValueOrError<FCompiledBindingLibraryCompiler::FFieldIdHandle, FString> AddFieldResult = AddFieldId(SourceContexts[SourceContextIndex].Class, true, Binding.BindingType, Binding.ViewModelPath.GetBindingName().ToName());
 			if (AddFieldResult.HasError())
 			{
-				WidgetBlueprintCompilerContext.MessageLog.Error(*FString::Printf(TEXT("The binding '%d' could not create it's source. %s")
-					, Index
-					, *AddFieldResult.GetError()));
+				AddErrorForBinding(Binding, BlueprintView, *FString::Printf(TEXT("The binding could not create its source. %s"), *AddFieldResult.GetError()));
 				bIsBindingsValid = false;
 				continue;
 			}
@@ -701,12 +701,10 @@ bool FMVVMViewBlueprintCompiler::PreCompileBindings(UWidgetBlueprintGeneratedCla
 				}
 			}
 
-			TValueOrError<FCompilerBinding, FString> AddBindingResult = AddBinding(Class, Getters, Setter, Binding.Conversion.SetConversionFunctionPath);
+			TValueOrError<FCompilerBinding, FString> AddBindingResult = AddBinding(Class, Getters, Setter, Binding.Conversion.SourceToDestinationFunctionPath);
 			if (AddBindingResult.HasError())
 			{
-				WidgetBlueprintCompilerContext.MessageLog.Error(*FString::Printf(TEXT("The binding '%d' could not be created. %s")
-					, Index
-					, *AddBindingResult.GetError()));
+				AddErrorForBinding(Binding, BlueprintView, FString::Printf(TEXT("The binding could not be created. %s"), *AddBindingResult.GetError()));
 				bIsBindingsValid = false;
 				continue;
 			}
@@ -733,9 +731,7 @@ bool FMVVMViewBlueprintCompiler::PreCompileBindings(UWidgetBlueprintGeneratedCla
 				TValueOrError<FCompiledBindingLibraryCompiler::FFieldIdHandle, FString> AddFieldIdResult = AddFieldId(Class->ClassGeneratedBy->GetClass(), true, Binding.BindingType, Binding.WidgetPath.GetBindingName().ToName());
 				if (AddFieldIdResult.HasError())
 				{
-					WidgetBlueprintCompilerContext.MessageLog.Error(*FString::Printf(TEXT("The binding '%d' could not create it's source. %s")
-						, Index
-						, *AddFieldIdResult.GetError()));
+					AddErrorForBinding(Binding, BlueprintView, FString::Printf(TEXT("The binding '%d' could not create it's source. %s"), *AddFieldIdResult.GetError()));
 					bIsBindingsValid = false;
 					continue;
 				}
@@ -748,9 +744,7 @@ bool FMVVMViewBlueprintCompiler::PreCompileBindings(UWidgetBlueprintGeneratedCla
 				TValueOrError<FCompiledBindingLibraryCompiler::FFieldIdHandle, FString> AddFieldIdResult = AddFieldId(SourceContexts[SourceContextIndex].Class, true, Binding.BindingType, Binding.WidgetPath.GetBindingName().ToName());
 				if (AddFieldIdResult.HasError())
 				{
-					WidgetBlueprintCompilerContext.MessageLog.Error(*FString::Printf(TEXT("The binding '%d' could not create it's source. %s")
-						, Index
-						, *AddFieldIdResult.GetError()));
+					AddErrorForBinding(Binding, BlueprintView, FString::Printf(TEXT("The binding '%d' could not create it's source. %s"), *AddFieldIdResult.GetError()));
 					bIsBindingsValid = false;
 					continue;
 				}
@@ -796,12 +790,10 @@ bool FMVVMViewBlueprintCompiler::PreCompileBindings(UWidgetBlueprintGeneratedCla
 				}
 			}
 
-			TValueOrError<FCompilerBinding, FString> AddBindingResult = AddBinding(Class, Getters, Setter, Binding.Conversion.GetConversionFunctionPath);
+			TValueOrError<FCompilerBinding, FString> AddBindingResult = AddBinding(Class, Getters, Setter, Binding.Conversion.DestinationToSourceFunctionPath);
 			if (AddBindingResult.HasError())
 			{
-				WidgetBlueprintCompilerContext.MessageLog.Error(*FString::Printf(TEXT("The binding '%d' could not be created. %s")
-					, Index
-					, *AddBindingResult.GetError()));
+				AddErrorForBinding(Binding, BlueprintView, FString::Printf(TEXT("The binding '%d' could not be created. %s"), *AddBindingResult.GetError()));
 				bIsBindingsValid = false;
 				continue;
 			}
diff --git a/Engine/Plugins/Runtime/ModelViewViewModel/Source/ModelViewViewModelBlueprint/Private/MVVMViewBlueprintCompiler.h b/Engine/Plugins/Runtime/ModelViewViewModel/Source/ModelViewViewModelBlueprint/Private/MVVMViewBlueprintCompiler.h
--- a/Engine/Plugins/Runtime/ModelViewViewModel/Source/ModelViewViewModelBlueprint/Private/MVVMViewBlueprintCompiler.h
+++ b/Engine/Plugins/Runtime/ModelViewViewModel/Source/ModelViewViewModelBlueprint/Private/MVVMViewBlueprintCompiler.h
@@ -66,7 +66,7 @@ private:
 
 	const FCompilerSourceCreatorContext* FindViewModelSource(FGuid Id) const;
 
-	void AddErrorForBinding(FMVVMBlueprintViewBinding& Binding, const FString& Message) const;
+	void AddErrorForBinding(FMVVMBlueprintViewBinding& Binding, const UMVVMBlueprintView* BlueprintView, const FString& Message) const;
 
 private:
 	struct FCompilerSourceContext
diff --git a/Engine/Plugins/Runtime/ModelViewViewModel/Source/ModelViewViewModelBlueprint/Public/MVVMBlueprintViewBinding.cpp b/Engine/Plugins/Runtime/ModelViewViewModel/Source/ModelViewViewModelBlueprint/Public/MVVMBlueprintViewBinding.cpp
new file mode 100644
--- /dev/null
+++ b/Engine/Plugins/Runtime/ModelViewViewModel/Source/ModelViewViewModelBlueprint/Public/MVVMBlueprintViewBinding.cpp
@@ -0,0 +1,19 @@
+// Copyright Epic Games, Inc. All Rights Reserved.
+
+#include "MVVMBlueprintViewBinding.h"
+#include "MVVMBlueprintView.h"
+
+FString FMVVMBlueprintViewBinding::GetNameString(const UMVVMBlueprintView* View) const
+{
+	const FMVVMBlueprintViewModelContext* ViewModel = View->FindViewModel(ViewModelPath.ContextId);
+	FString BindingName = (ViewModel != nullptr ? ViewModel->GetViewModelName().ToString() : TEXT("<none>")) + TEXT(".") + ViewModelPath.GetBindingName().ToString();
+
+	BindingName += BindingType == EMVVMBindingMode::TwoWay ? TEXT(" <-> ") :
+		UE::MVVM::IsForwardBinding(BindingType) ? TEXT(" -> ") :
+		UE::MVVM::IsBackwardBinding(BindingType) ? TEXT(" <- ") :
+		TEXT("???"); // shouldn't happen
+
+	BindingName += WidgetPath.WidgetName.ToString() + TEXT(".") + WidgetPath.GetBindingName().ToString();
+
+	return BindingName;
+}
\ No newline at end of file
diff --git a/Engine/Plugins/Runtime/ModelViewViewModel/Source/ModelViewViewModelBlueprint/Public/MVVMBlueprintViewBinding.h b/Engine/Plugins/Runtime/ModelViewViewModel/Source/ModelViewViewModelBlueprint/Public/MVVMBlueprintViewBinding.h
--- a/Engine/Plugins/Runtime/ModelViewViewModel/Source/ModelViewViewModelBlueprint/Public/MVVMBlueprintViewBinding.h
+++ b/Engine/Plugins/Runtime/ModelViewViewModel/Source/ModelViewViewModelBlueprint/Public/MVVMBlueprintViewBinding.h
@@ -7,6 +7,8 @@
 
 #include "MVVMBlueprintViewBinding.generated.h"
 
+class UMVVMBlueprintView;
+
 /**
 *
 */
@@ -15,13 +17,13 @@ struct FMVVMBlueprintViewConversionPath
 {
 	GENERATED_BODY()
 
-	/** The Conversion function when setting the binding. */
+	/** The Conversion function when converting the value from the destination to the source. */
 	UPROPERTY(EditAnywhere, Category = "MVVM", AdvancedDisplay)
-	FString SetConversionFunctionPath;
+	FString DestinationToSourceFunctionPath;
 
-	/** The Conversion function when getting the binding (in 2 way). */
+	/** The Conversion function when converting the value from the source to the destination. */
 	UPROPERTY(EditAnywhere, Category = "MVVM", AdvancedDisplay)
-	FString GetConversionFunctionPath;
+	FString SourceToDestinationFunctionPath;
 };
 
 /**
@@ -60,4 +62,10 @@ struct FMVVMBlueprintViewBinding
 	/** The binding is visible in the editor, but is not compiled and cannot be used at runtime. */
 	UPROPERTY(EditAnywhere, Category = "MVVM")
 	bool bCompile = true;
+
+	/** 
+	 * Get a string that identifies this binding. 
+	 * This is of the form: ViewModel.Property -> Widget.Property
+	 */
+	FString GetNameString(const UMVVMBlueprintView* View) const;
 };
diff --git a/Engine/Plugins/Runtime/ModelViewViewModel/Source/ModelViewViewModelBlueprint/Public/MVVMPropertyPath.h b/Engine/Plugins/Runtime/ModelViewViewModel/Source/ModelViewViewModelBlueprint/Public/MVVMPropertyPath.h
--- a/Engine/Plugins/Runtime/ModelViewViewModel/Source/ModelViewViewModelBlueprint/Public/MVVMPropertyPath.h
+++ b/Engine/Plugins/Runtime/ModelViewViewModel/Source/ModelViewViewModelBlueprint/Public/MVVMPropertyPath.h
@@ -23,7 +23,7 @@ private:
 	UPROPERTY(EditAnywhere, Category = "MVVM")
 	FMemberReference BindingReference;
 
-	/** The path of the property (either the getter or the set action). */
+	/** The path of the property (either the getter or the get action). */
 	UPROPERTY(EditAnywhere, Category = "MVVM")
 	FString GetterPropertyPath;
 
diff --git a/Engine/Plugins/Runtime/ModelViewViewModel/Source/ModelViewViewModelEditor/Private/Customizations/MVVMConversionPathCustomization.cpp b/Engine/Plugins/Runtime/ModelViewViewModel/Source/ModelViewViewModelEditor/Private/Customizations/MVVMConversionPathCustomization.cpp
--- a/Engine/Plugins/Runtime/ModelViewViewModel/Source/ModelViewViewModelEditor/Private/Customizations/MVVMConversionPathCustomization.cpp
+++ b/Engine/Plugins/Runtime/ModelViewViewModel/Source/ModelViewViewModelEditor/Private/Customizations/MVVMConversionPathCustomization.cpp
@@ -19,8 +19,8 @@ UE::MVVM::FConversionPathCustomization::FConversionPathCustomization(UWidgetBlue
 
 void UE::MVVM::FConversionPathCustomization::CustomizeHeader(TSharedRef<IPropertyHandle> InPropertyHandle, FDetailWidgetRow& HeaderRow, IPropertyTypeCustomizationUtils& CustomizationUtils)
 {
-	SetterProperty = InPropertyHandle->GetChildHandle(GET_MEMBER_NAME_CHECKED(FMVVMBlueprintViewConversionPath, SetConversionFunctionPath));
-	GetterProperty = InPropertyHandle->GetChildHandle(GET_MEMBER_NAME_CHECKED(FMVVMBlueprintViewConversionPath, GetConversionFunctionPath));
+	SourceToDestinationProperty = InPropertyHandle->GetChildHandle(GET_MEMBER_NAME_CHECKED(FMVVMBlueprintViewConversionPath, SourceToDestinationFunctionPath));
+	DestinationToSourceProperty = InPropertyHandle->GetChildHandle(GET_MEMBER_NAME_CHECKED(FMVVMBlueprintViewConversionPath, DestinationToSourceFunctionPath));
 
 	HeaderRow.NameContent()
 	[
@@ -32,10 +32,10 @@ void UE::MVVM::FConversionPathCustomization::CustomizeHeader(TSharedRef<IPropert
 	];
 }
 
-FText UE::MVVM::FConversionPathCustomization::GetGetterPath() const
+FText UE::MVVM::FConversionPathCustomization::GetSourceToDestinationPath() const
 {
 	FString Value;
-	FPropertyAccess::Result Result = GetterProperty->GetValue(Value);
+	FPropertyAccess::Result Result = SourceToDestinationProperty->GetValue(Value);
 
 	if (Result == FPropertyAccess::Success)
 	{
@@ -50,10 +50,10 @@ FText UE::MVVM::FConversionPathCustomization::GetGetterPath() const
 	return FText::GetEmpty();
 }
 
-FText UE::MVVM::FConversionPathCustomization::GetSetterPath() const
+FText UE::MVVM::FConversionPathCustomization::GetDestinationToSourcePath() const
 {
 	FString Value;
-	FPropertyAccess::Result Result = SetterProperty->GetValue(Value);
+	FPropertyAccess::Result Result = DestinationToSourceProperty->GetValue(Value);
 
 	if (Result == FPropertyAccess::Success)
 	{
@@ -68,15 +68,15 @@ FText UE::MVVM::FConversionPathCustomization::GetSetterPath() const
 	return FText::GetEmpty();
 }
 
-void UE::MVVM::FConversionPathCustomization::OnTextCommitted(const FText& NewValue, ETextCommit::Type CommitType, bool bIsGetter)
+void UE::MVVM::FConversionPathCustomization::OnTextCommitted(const FText& NewValue, ETextCommit::Type CommitType, bool bSourceToDestination)
 {
 	FString NewString = NewValue.ToString();
-	OnFunctionPathChanged(NewString, bIsGetter);
+	OnFunctionPathChanged(NewString, bSourceToDestination);
 }
 
-void UE::MVVM::FConversionPathCustomization::OnFunctionPathChanged(const FString& NewPath, bool bIsGetter)
+void UE::MVVM::FConversionPathCustomization::OnFunctionPathChanged(const FString& NewPath, bool bSourceToDestination)
 {
-	TSharedPtr<IPropertyHandle> Handle = bIsGetter ? GetterProperty : SetterProperty;
+	TSharedPtr<IPropertyHandle> Handle = bSourceToDestination ? SourceToDestinationProperty : DestinationToSourceProperty;
 	Handle->SetValue(NewPath);
 }
 
@@ -93,7 +93,7 @@ void UE::MVVM::FConversionPathCustomization::CustomizeChildren(TSharedRef<IPrope
 	}
 
 	// setter
-	ChildBuilder.AddProperty(SetterProperty.ToSharedRef())
+	ChildBuilder.AddProperty(DestinationToSourceProperty.ToSharedRef())
 		.CustomWidget()
 		.NameContent()
 		[
@@ -107,7 +107,7 @@ void UE::MVVM::FConversionPathCustomization::CustomizeChildren(TSharedRef<IPrope
 			.VAlign(VAlign_Center)
 			[
 				SNew(SEditableTextBox)
-				.Text(this, &FConversionPathCustomization::GetSetterPath)
+				.Text(this, &FConversionPathCustomization::GetDestinationToSourcePath)
 				.OnTextCommitted(this, &FConversionPathCustomization::OnTextCommitted, false)
 			]
 			+ SHorizontalBox::Slot()
@@ -121,7 +121,7 @@ void UE::MVVM::FConversionPathCustomization::CustomizeChildren(TSharedRef<IPrope
 		];
 
 	// getter
-	ChildBuilder.AddProperty(GetterProperty.ToSharedRef())
+	ChildBuilder.AddProperty(SourceToDestinationProperty.ToSharedRef())
 		.CustomWidget()
 		.NameContent()
 		[
@@ -135,7 +135,7 @@ void UE::MVVM::FConversionPathCustomization::CustomizeChildren(TSharedRef<IPrope
 			.VAlign(VAlign_Center)
 			[
 				SNew(SEditableTextBox)
-				.Text(this, &FConversionPathCustomization::GetGetterPath)
+				.Text(this, &FConversionPathCustomization::GetSourceToDestinationPath)
 				.OnTextCommitted(this, &FConversionPathCustomization::OnTextCommitted, true)
 			]
 			+ SHorizontalBox::Slot()
diff --git a/Engine/Plugins/Runtime/ModelViewViewModel/Source/ModelViewViewModelEditor/Private/Customizations/MVVMConversionPathCustomization.h b/Engine/Plugins/Runtime/ModelViewViewModel/Source/ModelViewViewModelEditor/Private/Customizations/MVVMConversionPathCustomization.h
--- a/Engine/Plugins/Runtime/ModelViewViewModel/Source/ModelViewViewModelEditor/Private/Customizations/MVVMConversionPathCustomization.h
+++ b/Engine/Plugins/Runtime/ModelViewViewModel/Source/ModelViewViewModelEditor/Private/Customizations/MVVMConversionPathCustomization.h
@@ -22,15 +22,15 @@ namespace UE::MVVM
 		virtual void CustomizeChildren(TSharedRef<IPropertyHandle> InPropertyHandle, IDetailChildrenBuilder& ChildBuilder, IPropertyTypeCustomizationUtils& CustomizationUtils) override;
 
 	private:
-		FText GetGetterPath() const;
-		FText GetSetterPath() const;
+		FText GetSourceToDestinationPath() const;
+		FText GetDestinationToSourcePath() const;
 
-		void OnTextCommitted(const FText& NewValue, ETextCommit::Type CommitType, bool bIsGetter);
-		void OnFunctionPathChanged(const FString& NewPath, bool bIsGetter);
+		void OnTextCommitted(const FText& NewValue, ETextCommit::Type CommitType, bool bSourceToDestination);
+		void OnFunctionPathChanged(const FString& NewPath, bool bSourceToDestination);
 
 	private:
 		UWidgetBlueprint* WidgetBlueprint = nullptr;
-		TSharedPtr<IPropertyHandle> GetterProperty;
-		TSharedPtr<IPropertyHandle> SetterProperty;
+		TSharedPtr<IPropertyHandle> DestinationToSourceProperty;
+		TSharedPtr<IPropertyHandle> SourceToDestinationProperty;
 	};
 }
\ No newline at end of file
diff --git a/Engine/Plugins/Runtime/ModelViewViewModel/Source/ModelViewViewModelEditor/Private/Widgets/SMVVMConversionPath.cpp b/Engine/Plugins/Runtime/ModelViewViewModel/Source/ModelViewViewModelEditor/Private/Widgets/SMVVMConversionPath.cpp
--- a/Engine/Plugins/Runtime/ModelViewViewModel/Source/ModelViewViewModelEditor/Private/Widgets/SMVVMConversionPath.cpp
+++ b/Engine/Plugins/Runtime/ModelViewViewModel/Source/ModelViewViewModelEditor/Private/Widgets/SMVVMConversionPath.cpp
@@ -18,9 +18,9 @@
 
 #define LOCTEXT_NAMESPACE "MVVMConversionPath"
 
-void SMVVMConversionPath::Construct(const FArguments& InArgs, const UWidgetBlueprint* InWidgetBlueprint, bool bInIsGetter)
+void SMVVMConversionPath::Construct(const FArguments& InArgs, const UWidgetBlueprint* InWidgetBlueprint, bool bInSourceToDestination)
 {
-	bIsGetter = bInIsGetter;
+	bSourceToDestination = bInSourceToDestination;
 	WidgetBlueprint = InWidgetBlueprint;
 	OnFunctionChanged = InArgs._OnFunctionChanged;
 	Bindings = InArgs._Bindings;
@@ -43,7 +43,7 @@ void SMVVMConversionPath::Construct(const FArguments& InArgs, const UWidgetBluep
 				.OnClicked(this, &SMVVMConversionPath::OnButtonClicked)
 				[
 					SNew(SImage)
-					.Image(FUMGStyle::Get().GetBrush(bIsGetter ? "MVVM.GetterFunction" : "MVVM.SetterFunction"))
+					.Image(FUMGStyle::Get().GetBrush(bSourceToDestination ? "MVVM.GetterFunction" : "MVVM.SetterFunction"))
 					.ColorAndOpacity(this, &SMVVMConversionPath::GetFunctionColor)
 				]
 			]
@@ -80,7 +80,7 @@ FString SMVVMConversionPath::GetFunctionPath() const
 	TArray<FMVVMBlueprintViewBinding*> ViewBindings = Bindings.Get(TArray<FMVVMBlueprintViewBinding*>());
 	for (const FMVVMBlueprintViewBinding* Binding : ViewBindings)
 	{
-		FString FunctionPath = bIsGetter ? Binding->Conversion.GetConversionFunctionPath : Binding->Conversion.SetConversionFunctionPath;
+		FString FunctionPath = bSourceToDestination ? Binding->Conversion.SourceToDestinationFunctionPath : Binding->Conversion.DestinationToSourceFunctionPath;
 		return FunctionPath;
 	}
 	return FString();
@@ -94,9 +94,9 @@ FText SMVVMConversionPath::GetFunctionToolTip() const
 		return FText::FromString(FunctionPath);
 	}
 
-	return bIsGetter ?
-		LOCTEXT("AddGetterConversionFunction", "Add conversion function to be used when getting the value from the viewmodel.") :
-		LOCTEXT("AddSetterConversionFunction", "Add conversion function to be used when setting the value in the viewmodel.");
+	return bSourceToDestination ?
+		LOCTEXT("AddSourceToDestinationFunction", "Add conversion function to be used when converting the source value to the destination value.") :
+		LOCTEXT("AddDestinationToSourceFunction", "Add conversion function to be used when converting the destination value to the source value.");
 }
 
 FSlateColor SMVVMConversionPath::GetFunctionColor() const
@@ -126,13 +126,13 @@ void SMVVMConversionPath::SetConversionFunction(const UFunction* Function)
 
 	for (FMVVMBlueprintViewBinding* Binding : ViewBindings)
 	{
-		if (bIsGetter)
+		if (bSourceToDestination)
 		{
-			Binding->Conversion.GetConversionFunctionPath = Function->GetPathName();
+			Binding->Conversion.SourceToDestinationFunctionPath = Function->GetPathName();
 		}
 		else
 		{
-			Binding->Conversion.SetConversionFunctionPath = Function->GetPathName();
+			Binding->Conversion.DestinationToSourceFunctionPath = Function->GetPathName();
 		}
 	}
 
@@ -161,8 +161,8 @@ TSharedRef<SWidget> SMVVMConversionPath::GetFunctionMenuContent()
 		UE::MVVM::FMVVMConstFieldVariant WidgetField = WidgetHelper.GetSelectedField();
 
 		UMVVMSubsystem::FConstDirectionalBindingArgs Args;
-		Args.SourceBinding = bIsGetter ? ViewModelField : WidgetField;
-		Args.DestinationBinding = bIsGetter ? WidgetField : ViewModelField;
+		Args.SourceBinding = bSourceToDestination ? ViewModelField : WidgetField;
+		Args.DestinationBinding = bSourceToDestination ? WidgetField : ViewModelField;
 
 		UMVVMEditorSubsystem* EditorSubsystem = GEditor->GetEditorSubsystem<UMVVMEditorSubsystem>();
 		TArray<const UFunction*> FunctionsForThis = EditorSubsystem->GetAvailableConversionFunctions(Args.SourceBinding, Args.DestinationBinding);
@@ -191,7 +191,7 @@ TSharedRef<SWidget> SMVVMConversionPath::GetFunctionMenuContent()
 			],
 			FText::GetEmpty(),
 			true, // no indent
-			false // searchable
+			true // searchable
 		);
 	}
 
diff --git a/Engine/Plugins/Runtime/ModelViewViewModel/Source/ModelViewViewModelEditor/Private/Widgets/SMVVMConversionPath.h b/Engine/Plugins/Runtime/ModelViewViewModel/Source/ModelViewViewModelEditor/Private/Widgets/SMVVMConversionPath.h
--- a/Engine/Plugins/Runtime/ModelViewViewModel/Source/ModelViewViewModelEditor/Private/Widgets/SMVVMConversionPath.h
+++ b/Engine/Plugins/Runtime/ModelViewViewModel/Source/ModelViewViewModelEditor/Private/Widgets/SMVVMConversionPath.h
@@ -20,7 +20,7 @@ public:
 		SLATE_EVENT(FOnFunctionChanged, OnFunctionChanged)
 	SLATE_END_ARGS()
 
-	void Construct(const FArguments& InArgs, const UWidgetBlueprint* InWidgetBlueprint, bool bInIsGetter);
+	void Construct(const FArguments& InArgs, const UWidgetBlueprint* InWidgetBlueprint, bool bInSourceToDestination);
 
 private:
 	EVisibility IsFunctionVisible() const;
@@ -36,5 +36,5 @@ private:
 	FOnFunctionChanged OnFunctionChanged;
 	TSharedPtr<SMenuAnchor> Anchor;
 	const UWidgetBlueprint* WidgetBlueprint = nullptr;
-	bool bIsGetter = false;
+	bool bSourceToDestination = false;
 };
\ No newline at end of file
