Moved MVVM-specific images from UMGStyle to MVVMEditorStyle

Added Clear to conversion function button.

diff --git a/Engine/Plugins/Runtime/ModelViewViewModel/Source/ModelViewViewModelEditor/Private/Styling/MVVMEditorStyle.cpp b/Engine/Plugins/Runtime/ModelViewViewModel/Source/ModelViewViewModelEditor/Private/Styling/MVVMEditorStyle.cpp
--- a/Engine/Plugins/Runtime/ModelViewViewModel/Source/ModelViewViewModelEditor/Private/Styling/MVVMEditorStyle.cpp
+++ b/Engine/Plugins/Runtime/ModelViewViewModel/Source/ModelViewViewModelEditor/Private/Styling/MVVMEditorStyle.cpp
@@ -57,6 +57,14 @@ FMVVMEditorStyle::FMVVMEditorStyle()
 		.SetShadowColorAndOpacity(FLinearColor(0, 0, 0, 0.7f)));
 	}
 
+	Set("BindingMode.OneWay", new IMAGE_BRUSH_SVG("Slate/BindingMode_OneWay", Icon16x16));
+	Set("BindingMode.OneWayToSource", new IMAGE_BRUSH_SVG("Slate/BindingMode_OneWayToSource", Icon16x16));
+	Set("BindingMode.TwoWay", new IMAGE_BRUSH_SVG("Slate/BindingMode_TwoWay", Icon16x16));
+	Set("BindingMode.OneTime", new IMAGE_BRUSH_SVG("Slate/BindingMode_OneTime", Icon16x16));
+
+	Set("ConversionFunction.DestToSource", new IMAGE_BRUSH_SVG("Slate/ConversionFunction_DestToSource", Icon16x16));
+	Set("ConversionFunction.SourceToDest", new IMAGE_BRUSH_SVG("Slate/ConversionFunction_SourceToDest", Icon16x16));
+
 	FSlateStyleRegistry::RegisterSlateStyle(*this);
 }
 
diff --git a/Engine/Plugins/Runtime/ModelViewViewModel/Source/ModelViewViewModelEditor/Private/Widgets/SMVVMConversionPath.cpp b/Engine/Plugins/Runtime/ModelViewViewModel/Source/ModelViewViewModelEditor/Private/Widgets/SMVVMConversionPath.cpp
--- a/Engine/Plugins/Runtime/ModelViewViewModel/Source/ModelViewViewModelEditor/Private/Widgets/SMVVMConversionPath.cpp
+++ b/Engine/Plugins/Runtime/ModelViewViewModel/Source/ModelViewViewModelEditor/Private/Widgets/SMVVMConversionPath.cpp
@@ -8,9 +8,9 @@
 #include "MVVMEditorSubsystem.h"
 #include "MVVMPropertyPathHelpers.h"
 #include "MVVMSubsystem.h"
+#include "Styling/MVVMEditorStyle.h"
 #include "Styling/StyleColors.h"
 #include "Types/MVVMFieldVariant.h"
-#include "UMGStyle.h"
 #include "Widgets/Images/SImage.h"
 #include "Widgets/Input/SButton.h"
 #include "Widgets/Input/SMenuAnchor.h"
@@ -43,7 +43,7 @@ void SMVVMConversionPath::Construct(const FArguments& InArgs, const UWidgetBluep
 				.OnClicked(this, &SMVVMConversionPath::OnButtonClicked)
 				[
 					SNew(SImage)
-					.Image(FUMGStyle::Get().GetBrush(bSourceToDestination ? "MVVM.GetterFunction" : "MVVM.SetterFunction"))
+					.Image(FMVVMEditorStyle::Get().GetBrush(bSourceToDestination ? "ConversionFunction.SourceToDest" : "ConversionFunction.DestToSource"))
 					.ColorAndOpacity(this, &SMVVMConversionPath::GetFunctionColor)
 				]
 			]
@@ -56,19 +56,32 @@ EVisibility SMVVMConversionPath::IsFunctionVisible() const
 	TArray<FMVVMBlueprintViewBinding*> ViewBindings = Bindings.Get(TArray<FMVVMBlueprintViewBinding*>());
 	if (ViewBindings.Num() == 0)
 	{
-		return EVisibility::Collapsed;
+		return EVisibility::Hidden;
 	}
 
 	for (const FMVVMBlueprintViewBinding* Binding : ViewBindings)
 	{
-		bool bShouldBeVisible =
-			Binding->BindingType == EMVVMBindingMode::OneTimeToDestination ||
-			Binding->BindingType == EMVVMBindingMode::OneWayToDestination ||
-			Binding->BindingType == EMVVMBindingMode::TwoWay;
-
-		if (bShouldBeVisible)
+		if (bSourceToDestination)
+		{
+			bool bShouldBeVisible = Binding->BindingType == EMVVMBindingMode::OneTimeToDestination ||
+				Binding->BindingType == EMVVMBindingMode::OneWayToDestination ||
+				Binding->BindingType == EMVVMBindingMode::TwoWay;
+
+			if (bShouldBeVisible)
+			{
+				return EVisibility::Visible;
+			}
+		}
+		else
 		{
-			return EVisibility::Visible;
+			bool bShouldBeVisible = Binding->BindingType == EMVVMBindingMode::OneTimeToSource ||
+				Binding->BindingType == EMVVMBindingMode::OneWayToSource ||
+				Binding->BindingType == EMVVMBindingMode::TwoWay;
+
+			if (bShouldBeVisible)
+			{
+				return EVisibility::Visible;
+			}
 		}
 	}
 	
@@ -78,12 +91,23 @@ EVisibility SMVVMConversionPath::IsFunctionVisible() const
 FString SMVVMConversionPath::GetFunctionPath() const
 {
 	TArray<FMVVMBlueprintViewBinding*> ViewBindings = Bindings.Get(TArray<FMVVMBlueprintViewBinding*>());
+
+	bool bFirst = true;
+	FString FunctionPath;
 	for (const FMVVMBlueprintViewBinding* Binding : ViewBindings)
 	{
-		FString FunctionPath = bSourceToDestination ? Binding->Conversion.SourceToDestinationFunctionPath : Binding->Conversion.DestinationToSourceFunctionPath;
-		return FunctionPath;
+		FString ThisPath = bSourceToDestination ? Binding->Conversion.SourceToDestinationFunctionPath : Binding->Conversion.DestinationToSourceFunctionPath;
+		if (bFirst)
+		{
+			FunctionPath = ThisPath;
+		}
+		else if (FunctionPath != ThisPath)
+		{
+			return TEXT("Multiple Values");
+		}
 	}
-	return FString();
+
+	return FunctionPath;
 }
 
 FText SMVVMConversionPath::GetFunctionToolTip() const
@@ -91,6 +115,10 @@ FText SMVVMConversionPath::GetFunctionToolTip() const
 	FString FunctionPath = GetFunctionPath();
 	if (!FunctionPath.IsEmpty())
 	{
+		if (FunctionPath == TEXT("Multiple Values"))
+		{
+			return LOCTEXT("MultipleValues", "Multiple Values");
+		}
 		return FText::FromString(FunctionPath);
 	}
 
@@ -128,17 +156,17 @@ void SMVVMConversionPath::SetConversionFunction(const UFunction* Function)
 	{
 		if (bSourceToDestination)
 		{
-			Binding->Conversion.SourceToDestinationFunctionPath = Function->GetPathName();
+			Binding->Conversion.SourceToDestinationFunctionPath = Function != nullptr ? Function->GetPathName() : FString();
 		}
 		else
 		{
-			Binding->Conversion.DestinationToSourceFunctionPath = Function->GetPathName();
+			Binding->Conversion.DestinationToSourceFunctionPath = Function != nullptr ? Function->GetPathName() : FString();
 		}
 	}
 
 	if (OnFunctionChanged.IsBound())
 	{
-		OnFunctionChanged.Execute(Function->GetPathName());
+		OnFunctionChanged.Execute(Function != nullptr ? Function->GetPathName() : FString());
 	}
 }
 
@@ -205,6 +233,17 @@ TSharedRef<SWidget> SMVVMConversionPath::GetFunctionMenuContent()
 			Action);
 	}
 
+	const FString Path = GetFunctionPath();
+	if (!Path.IsEmpty())
+	{
+		FUIAction ClearAction(FExecuteAction::CreateSP(this, &SMVVMConversionPath::SetConversionFunction, (const UFunction*) nullptr));
+		MenuBuilder.AddMenuEntry(
+			LOCTEXT("Clear", "Clear"),
+			LOCTEXT("ClearToolTip", "Clear this conversion function."),
+			FSlateIcon(FAppStyle::GetAppStyleSetName(), "Icons.X"),
+			ClearAction);
+	}
+
 	return MenuBuilder.MakeWidget();
 }
 
diff --git a/Engine/Plugins/Runtime/ModelViewViewModel/Source/ModelViewViewModelEditor/Private/Widgets/SMVVMViewBindingListView.cpp b/Engine/Plugins/Runtime/ModelViewViewModel/Source/ModelViewViewModelEditor/Private/Widgets/SMVVMViewBindingListView.cpp
--- a/Engine/Plugins/Runtime/ModelViewViewModel/Source/ModelViewViewModelEditor/Private/Widgets/SMVVMViewBindingListView.cpp
+++ b/Engine/Plugins/Runtime/ModelViewViewModel/Source/ModelViewViewModelEditor/Private/Widgets/SMVVMViewBindingListView.cpp
@@ -18,8 +18,8 @@
 #include "MVVMWidgetBlueprintExtension_View.h"
 #include "SEnumCombo.h"
 #include "SSimpleButton.h" 
+#include "Styling/MVVMEditorStyle.h"
 #include "Styling/StyleColors.h"
-#include "UMGStyle.h"
 #include "Widgets/SMVVMConversionPath.h"
 #include "Widgets/SMVVMFieldSelector.h"
 #include "Widgets/SMVVMPropertyPath.h"
@@ -477,15 +477,15 @@ private:
 		switch (BindingMode)
 		{
 		case EMVVMBindingMode::OneTimeToDestination:
-			return FUMGStyle::Get().GetBrush("MVVM.OneTime");
+			return FMVVMEditorStyle::Get().GetBrush("BindingMode.OneTime");
 		case EMVVMBindingMode::OneWayToDestination:
-			return FUMGStyle::Get().GetBrush("MVVM.OneWay");
+			return FMVVMEditorStyle::Get().GetBrush("BindingMode.OneWay");
 		case EMVVMBindingMode::OneWayToSource:
-			return FUMGStyle::Get().GetBrush("MVVM.OneWayToSource");
+			return FMVVMEditorStyle::Get().GetBrush("BindingMode.OneWayToSource");
 		case EMVVMBindingMode::OneTimeToSource:
 			return nullptr;
 		case EMVVMBindingMode::TwoWay:
-			return FUMGStyle::Get().GetBrush("MVVM.TwoWay");
+			return FMVVMEditorStyle::Get().GetBrush("BindingMode.TwoWay");
 		default:
 			return nullptr;
 		}
diff --git a/Engine/Plugins/Runtime/ModelViewViewModel/Source/ModelViewViewModelEditor/Private/Widgets/SMVVMViewModelContextListWidget.cpp b/Engine/Plugins/Runtime/ModelViewViewModel/Source/ModelViewViewModelEditor/Private/Widgets/SMVVMViewModelContextListWidget.cpp
--- a/Engine/Plugins/Runtime/ModelViewViewModel/Source/ModelViewViewModelEditor/Private/Widgets/SMVVMViewModelContextListWidget.cpp
+++ b/Engine/Plugins/Runtime/ModelViewViewModel/Source/ModelViewViewModelEditor/Private/Widgets/SMVVMViewModelContextListWidget.cpp
@@ -15,7 +15,6 @@
 #include "SSimpleButton.h"
 #include "Styling/MVVMEditorStyle.h"
 #include "Styling/SlateIconFinder.h"
-#include "UMGStyle.h"
 #include "Widgets/Input/SComboButton.h"
 #include "Widgets/Input/SEditableTextBox.h"
 #include "Widgets/Layout/SSpacer.h"
@@ -660,15 +659,15 @@ namespace UE::MVVM::Private
 		switch (BindingMode)
 		{
 		case EMVVMBindingMode::OneTimeToDestination:
-			return FUMGStyle::Get().GetBrush("MVVM.OneTime");
+			return FMVVMEditorStyle::Get().GetBrush("BindingMode.OneTime");
 		case EMVVMBindingMode::OneWayToDestination:
-			return FUMGStyle::Get().GetBrush("MVVM.OneWay");
+			return FMVVMEditorStyle::Get().GetBrush("BindingMode.OneWay");
 		case EMVVMBindingMode::OneWayToSource:
-			return FUMGStyle::Get().GetBrush("MVVM.OneWayToSource");
+			return FMVVMEditorStyle::Get().GetBrush("BindingMode.OneWayToSource");
 		case EMVVMBindingMode::OneTimeToSource:
 			return nullptr;
 		case EMVVMBindingMode::TwoWay:
-			return FUMGStyle::Get().GetBrush("MVVM.TwoWay");
+			return FMVVMEditorStyle::Get().GetBrush("BindingMode.TwoWay");
 		default:
 			return nullptr;
 		}
diff --git a/Engine/Source/Runtime/UMG/Private/UMGStyle.cpp b/Engine/Source/Runtime/UMG/Private/UMGStyle.cpp
--- a/Engine/Source/Runtime/UMG/Private/UMGStyle.cpp
+++ b/Engine/Source/Runtime/UMG/Private/UMGStyle.cpp
@@ -160,13 +160,6 @@ TSharedRef< FSlateStyleSet > FUMGStyle::Create()
 	Style->Set("ClassThumbnail.DetailsView", new IMAGE_BRUSH_SVG(TEXT("DetailsView_64"), Icon64x64));
 	Style->Set("ClassThumbnail.SinglePropertyView", new IMAGE_BRUSH_SVG(TEXT("SinglePropertyView_64"), Icon64x64));
 
-	Style->Set("MVVM.OneWay", new IMAGE_BRUSH_SVG("ViewModel_OneWay", Icon16x16));
-	Style->Set("MVVM.OneWayToSource", new IMAGE_BRUSH_SVG("ViewModel_OneWayToSource", Icon16x16));
-	Style->Set("MVVM.TwoWay", new IMAGE_BRUSH_SVG("ViewModel_TwoWay", Icon16x16));
-	Style->Set("MVVM.OneTime", new IMAGE_BRUSH_SVG("ViewModel_OneTime", Icon16x16));
-	Style->Set("MVVM.SetterFunction", new IMAGE_BRUSH_SVG("Function_Setter", Icon16x16));
-	Style->Set("MVVM.GetterFunction", new IMAGE_BRUSH_SVG("Function_Getter", Icon16x16));
-
 	return Style;
 }
 
