MVVM: Fix some static Analysis
#rnx

diff --git a/Engine/Plugins/Runtime/ModelViewViewModel/Source/ModelViewViewModelEditor/Private/Tests/MVVMBindingExecuteTest.cpp b/Engine/Plugins/Runtime/ModelViewViewModel/Source/ModelViewViewModelEditor/Private/Tests/MVVMBindingExecuteTest.cpp
--- a/Engine/Plugins/Runtime/ModelViewViewModel/Source/ModelViewViewModelEditor/Private/Tests/MVVMBindingExecuteTest.cpp
+++ b/Engine/Plugins/Runtime/ModelViewViewModel/Source/ModelViewViewModelEditor/Private/Tests/MVVMBindingExecuteTest.cpp
@@ -139,7 +139,7 @@ bool FMVVMBindingExecuteTest::RunTest(const FString& Parameters)
 			UE::MVVM::FFieldContext Destination(UE::MVVM::FObjectVariant(DestinationObj), DestinationBinding);
 			UE::MVVM::BindingHelper::ExecuteBinding_NoCheck(Source, Destination);
 
-			if (UE::MVVM::Private::GMVVMBindingExecTextCounter != 0)
+			if (UE::MVVM::Private::GMVVMBindingExecTextCounter != 0) // -V547 
 			{
 				AddError(FString::Printf(TEXT("%s failed to release the reources."), Msg));
 			}
@@ -178,7 +178,7 @@ bool FMVVMBindingExecuteTest::RunTest(const FString& Parameters)
 			UE::MVVM::FFieldContext Destination(UE::MVVM::FObjectVariant(DestinationObj), DestinationBinding);
 			UE::MVVM::BindingHelper::ExecuteBinding_NoCheck(Source, Destination);
 
-			if (UE::MVVM::Private::GMVVMBindingExecTextCounter != 4)
+			if (UE::MVVM::Private::GMVVMBindingExecTextCounter != 4) // -V547 
 			{
 				AddError(FString::Printf(TEXT("%s failed to release the reources."), Msg));
 			}
@@ -212,7 +212,7 @@ bool FMVVMBindingExecuteTest::RunTest(const FString& Parameters)
 			UE::MVVM::FFieldContext Destination(UE::MVVM::FObjectVariant(DestinationObj), DestinationBinding);
 			UE::MVVM::BindingHelper::ExecuteBinding_NoCheck(Source, Destination);
 
-			if (UE::MVVM::Private::GMVVMBindingExecTextCounter != 0)
+			if (UE::MVVM::Private::GMVVMBindingExecTextCounter != 0) // -V547 
 			{
 				AddError(FString::Printf(TEXT("%s failed to release the reources."), Msg));
 			}
@@ -251,7 +251,7 @@ bool FMVVMBindingExecuteTest::RunTest(const FString& Parameters)
 			UE::MVVM::FFieldContext Destination(UE::MVVM::FObjectVariant(DestinationObj), DestinationBinding);
 			UE::MVVM::BindingHelper::ExecuteBinding_NoCheck(Source, Destination);
 
-			if (UE::MVVM::Private::GMVVMBindingExecTextCounter != 0)
+			if (UE::MVVM::Private::GMVVMBindingExecTextCounter != 0) // -V547
 			{
 				AddError(FString::Printf(TEXT("%s failed to release the reources."), Msg));
 			}
@@ -285,7 +285,7 @@ bool FMVVMBindingExecuteTest::RunTest(const FString& Parameters)
 			UE::MVVM::FFunctionContext Function = UE::MVVM::FFunctionContext::MakeStaticFunction(Conversion);
 			UE::MVVM::BindingHelper::ExecuteBinding_NoCheck(Source, Destination, Function);
 
-			if (UE::MVVM::Private::GMVVMBindingExecTextCounter != 0)
+			if (UE::MVVM::Private::GMVVMBindingExecTextCounter != 0) // -V547
 			{
 				AddError(FString::Printf(TEXT("%s failed to release the reources."), Msg));
 			}
@@ -326,7 +326,7 @@ bool FMVVMBindingExecuteTest::RunTest(const FString& Parameters)
 			UE::MVVM::FFunctionContext Function = UE::MVVM::FFunctionContext::MakeStaticFunction(Conversion);
 			UE::MVVM::BindingHelper::ExecuteBinding_NoCheck(Source, Destination, Function);
 
-			if (UE::MVVM::Private::GMVVMBindingExecTextCounter != 0)
+			if (UE::MVVM::Private::GMVVMBindingExecTextCounter != 0) // -V547
 			{
 				AddError(FString::Printf(TEXT("%s failed to release the reources."), Msg));
 			}
@@ -372,7 +372,7 @@ bool FMVVMBindingExecuteTest::RunTest(const FString& Parameters)
 			UE::MVVM::FFunctionContext Function = UE::MVVM::FFunctionContext::MakeStaticFunction(Conversion);
 			UE::MVVM::BindingHelper::ExecuteBinding_NoCheck(Source, Destination, Function);
 
-			if (UE::MVVM::Private::GMVVMBindingExecTextCounter != 0)
+			if (UE::MVVM::Private::GMVVMBindingExecTextCounter != 0) // -V547
 			{
 				AddError(FString::Printf(TEXT("%s failed to release the reources."), Msg));
 			}
@@ -421,7 +421,7 @@ bool FMVVMBindingExecuteTest::RunTest(const FString& Parameters)
 			UE::MVVM::FFunctionContext Function = UE::MVVM::FFunctionContext::MakeStaticFunction(Conversion);
 			UE::MVVM::BindingHelper::ExecuteBinding_NoCheck(Source, Destination, Function);
 
-			if (UE::MVVM::Private::GMVVMBindingExecTextCounter != 4)
+			if (UE::MVVM::Private::GMVVMBindingExecTextCounter != 4) // -V547
 			{
 				AddError(FString::Printf(TEXT("%s failed to release the reources."), Msg));
 			}
@@ -482,7 +482,7 @@ bool FMVVMBindingExecuteTest::RunTest(const FString& Parameters)
 					UE::MVVM::BindingHelper::ExecuteBinding_NoCheck(Source, Destination);
 				}
 
-				if (UE::MVVM::Private::GMVVMBindingExecTextCounter != 0)
+				if (UE::MVVM::Private::GMVVMBindingExecTextCounter != 0) // -V547
 				{
 					AddError(FString::Printf(TEXT("%s failed to release the reources."), Msg));
 				}
@@ -508,7 +508,7 @@ bool FMVVMBindingExecuteTest::RunTest(const FString& Parameters)
 					UE::MVVM::BindingHelper::ExecuteBinding_NoCheck(Source, Destination);
 				}
 
-				if (UE::MVVM::Private::GMVVMBindingExecTextCounter != 0)
+				if (UE::MVVM::Private::GMVVMBindingExecTextCounter != 0) // -V547
 				{
 					AddError(FString::Printf(TEXT("%s failed to release the reources."), Msg));
 				}
@@ -536,7 +536,7 @@ bool FMVVMBindingExecuteTest::RunTest(const FString& Parameters)
 					UE::MVVM::BindingHelper::ExecuteBinding_NoCheck(Source, Destination);
 				}
 
-				if (UE::MVVM::Private::GMVVMBindingExecTextCounter != 0)
+				if (UE::MVVM::Private::GMVVMBindingExecTextCounter != 0) // -V547
 				{
 					AddError(FString::Printf(TEXT("%s failed to release the reources."), Msg));
 				}
@@ -562,7 +562,7 @@ bool FMVVMBindingExecuteTest::RunTest(const FString& Parameters)
 					UE::MVVM::BindingHelper::ExecuteBinding_NoCheck(Source, Destination);
 				}
 
-				if (UE::MVVM::Private::GMVVMBindingExecTextCounter != 0)
+				if (UE::MVVM::Private::GMVVMBindingExecTextCounter != 0) // -V547
 				{
 					AddError(FString::Printf(TEXT("%s failed to release the reources."), Msg));
 				}
