Cleanup MVVM Field to MemberReference conversion

diff --git a/Engine/Plugins/Runtime/ModelViewViewModel/Source/ModelViewViewModel/Public/Types/MVVMFieldVariant.h b/Engine/Plugins/Runtime/ModelViewViewModel/Source/ModelViewViewModel/Public/Types/MVVMFieldVariant.h
--- a/Engine/Plugins/Runtime/ModelViewViewModel/Source/ModelViewViewModel/Public/Types/MVVMFieldVariant.h
+++ b/Engine/Plugins/Runtime/ModelViewViewModel/Source/ModelViewViewModel/Public/Types/MVVMFieldVariant.h
@@ -131,20 +131,27 @@ namespace UE::MVVM
 		{
 			return !(*this == B);
 		}
-
-		operator FMemberReference() const
+		
+		/**
+		 * Create a serializable member reference from this field.
+		 * 
+		 * Generally should not need to set self context to true. If you are running into reference issues
+		 * see param comment & consider reviewing UBlueprintVariableNodeSpawner::CreateFromMemberOrParam
+		 * 
+		 * @param  bInSelfContext	Pass true if a BP skeleton class owns the referenced property (Class not BP generated or otherwise).
+		 * @return A newly allocated FMemberReference representing this field variant.
+		 */
+		FMemberReference CreateMemberReference(bool bInSelfContext = false) const
 		{
 			FMemberReference BindingReference = FMemberReference();
 			if (IsProperty())
 			{
-				// @TODO: DarenC - Should self context ever be true? Seems no, but if issues occur:
-				// analyze UBlueprintVariableNodeSpawner::CreateFromMemberOrParam -> MemberVarSetupLambda
 				PropertyType* Property = GetProperty();
-				bool bSelfContext = false;
-				BindingReference.SetFromField<FProperty>(Property, bSelfContext);
+				BindingReference.SetFromField<FProperty>(Property, bInSelfContext);
 			}
 			else if (IsFunction())
 			{
+				// Functions should never be self context references
 				FunctionType* Function = GetFunction();
 				bool bSelfContext = false;
 				BindingReference.SetFromField<UFunction>(Function, bSelfContext);
diff --git a/Engine/Plugins/Runtime/ModelViewViewModel/Source/ModelViewViewModelBlueprint/Public/MVVMPropertyPath.h b/Engine/Plugins/Runtime/ModelViewViewModel/Source/ModelViewViewModelBlueprint/Public/MVVMPropertyPath.h
--- a/Engine/Plugins/Runtime/ModelViewViewModel/Source/ModelViewViewModelBlueprint/Public/MVVMPropertyPath.h
+++ b/Engine/Plugins/Runtime/ModelViewViewModel/Source/ModelViewViewModelBlueprint/Public/MVVMPropertyPath.h
@@ -64,7 +64,7 @@ public:
 
 	void SetBindingReference(const UE::MVVM::FMVVMFieldVariant& InField)
 	{
-		BindingReference = InField;
+		BindingReference = InField.CreateMemberReference();
 
 		if (InField.IsProperty())
 		{
@@ -82,7 +82,7 @@ public:
 
 	void SetBindingReference(const UE::MVVM::FMVVMConstFieldVariant& InField)
 	{
-		BindingReference = InField;
+		BindingReference = InField.CreateMemberReference();
 
 		if (InField.IsProperty())
 		{
