MVVM: Fix ModelViewViewModelEditorModule.cpp that was badly merged
#rnx
#prefligh none

diff --git a/Engine/Plugins/Runtime/ModelViewViewModel/Source/ModelViewViewModelEditor/Private/ModelViewViewModelEditorModule.cpp b/Engine/Plugins/Runtime/ModelViewViewModel/Source/ModelViewViewModelEditor/Private/ModelViewViewModelEditorModule.cpp
--- a/Engine/Plugins/Runtime/ModelViewViewModel/Source/ModelViewViewModelEditor/Private/ModelViewViewModelEditorModule.cpp
+++ b/Engine/Plugins/Runtime/ModelViewViewModel/Source/ModelViewViewModelEditor/Private/ModelViewViewModelEditorModule.cpp
@@ -26,7 +26,7 @@
 #include "WidgetBlueprint.h"
 #include "WidgetBlueprintExtension.h"
 #include "WidgetBlueprintToolMenuContext.h"
-//#include "WidgetDrawerConfig.h"
+#include "WidgetDrawerConfig.h"
 #include "Widgets/Images/SImage.h"
 #include "WorkflowOrientedApp/WorkflowTabFactory.h"
 #include "WorkflowOrientedApp/WorkflowTabManager.h"
@@ -134,7 +134,7 @@ void FModelViewViewModelEditorModule::HandleActivateMode(FWidgetBlueprintApplica
 
 		// Add MVVM Drawer
 		{
-			FStatusBarDrawer MVVMDrawer(FMVVMBindingSummoner::DrawerID);
+			FWidgetDrawerConfig MVVMDrawer(FMVVMBindingSummoner::DrawerID);
 			TWeakPtr<FWidgetBlueprintEditor> WeakBP = BP;
 			MVVMDrawer.GetDrawerContentDelegate.BindLambda([WeakBP]()
 			{
