MVVM: Fix missing include.
#rnx

diff --git a/Engine/Plugins/Runtime/ModelViewViewModel/Source/ModelViewViewModelEditor/Private/Widgets/SMVVMViewModelContextListWidget.cpp b/Engine/Plugins/Runtime/ModelViewViewModel/Source/ModelViewViewModelEditor/Private/Widgets/SMVVMViewModelContextListWidget.cpp
--- a/Engine/Plugins/Runtime/ModelViewViewModel/Source/ModelViewViewModelEditor/Private/Widgets/SMVVMViewModelContextListWidget.cpp
+++ b/Engine/Plugins/Runtime/ModelViewViewModel/Source/ModelViewViewModelEditor/Private/Widgets/SMVVMViewModelContextListWidget.cpp
@@ -17,6 +17,7 @@
 #include "Styling/SlateIconFinder.h"
 #include "Widgets/Input/SComboButton.h"
 #include "Widgets/Input/SEditableTextBox.h"
+#include "Widgets/Images/SImage.h"
 #include "Widgets/Layout/SSpacer.h"
 #include "Widgets/Layout/SUniformGridPanel.h"
 #include "Widgets/Views/SListView.h"
