MVVM: Fix missing include.
#rnx

diff --git a/Engine/Plugins/Runtime/ModelViewViewModel/Source/ModelViewViewModelEditor/Private/Widgets/SMVVMViewBindingPanel.cpp b/Engine/Plugins/Runtime/ModelViewViewModel/Source/ModelViewViewModelEditor/Private/Widgets/SMVVMViewBindingPanel.cpp
--- a/Engine/Plugins/Runtime/ModelViewViewModel/Source/ModelViewViewModelEditor/Private/Widgets/SMVVMViewBindingPanel.cpp
+++ b/Engine/Plugins/Runtime/ModelViewViewModel/Source/ModelViewViewModelEditor/Private/Widgets/SMVVMViewBindingPanel.cpp
@@ -22,6 +22,7 @@
 #include "Framework/MultiBox/MultiBoxBuilder.h"
 #include "Widgets/Input/SButton.h"
 #include "Widgets/SNullWidget.h"
+#include "Widgets/Images/SImage.h"
 #include "Widgets/Layout/SBorder.h"
 #include "Widgets/Layout/SSpacer.h"
 #include "Widgets/Layout/SSplitter.h"
