MVVM:
Compiler support convertion function.
Reset the compiled library after compilation (prevents multiple loading).
Initialize the widget generated class FieldId after a compilation.
When compiling Skel, the class is not link. Property may not be found with FindFieldByName because they were just created.
MVVM message log was not initialized.
Fix nullptr deferecing warning.
#rnx

diff --git a/Engine/Plugins/Runtime/ModelViewViewModel/Source/ModelViewViewModel/Private/Bindings/MVVMBindingHelper.cpp b/Engine/Plugins/Runtime/ModelViewViewModel/Source/ModelViewViewModel/Private/Bindings/MVVMBindingHelper.cpp
--- a/Engine/Plugins/Runtime/ModelViewViewModel/Source/ModelViewViewModel/Private/Bindings/MVVMBindingHelper.cpp
+++ b/Engine/Plugins/Runtime/ModelViewViewModel/Source/ModelViewViewModel/Private/Bindings/MVVMBindingHelper.cpp
@@ -114,7 +114,7 @@ namespace UE::MVVM::BindingHelper
 	bool IsValidForRuntimeConversion(const UFunction* InFunction)
 	{
 		if (Private::IsValidCommon(InFunction)
-			&& InFunction->HasAllFunctionFlags(FUNC_Static)
+			&& (InFunction->HasAllFunctionFlags(FUNC_Static) || InFunction->HasAnyFunctionFlags(FUNC_Const | FUNC_BlueprintPure))
 			&& InFunction->NumParms == 2)
 		{
 			const FProperty* ReturnProperty = GetReturnProperty(InFunction);
@@ -381,6 +381,11 @@ namespace UE::MVVM::BindingHelper
 			return MakeError(MoveTemp(CommonResult));
 		}
 
+		if (!InFunction->HasAllFunctionFlags(FUNC_Static) && !InFunction->HasAnyFunctionFlags(FUNC_Const | FUNC_BlueprintPure))
+		{
+			return MakeError(FString::Printf(TEXT("The function '%s' is not static or is not const and pure."), *InFunction->GetName()));
+		}
+
 		if (InFunction->NumParms != 2)
 		{
 			return MakeError(FString::Printf(TEXT("The function '%s' does not have the correct number of arguments."), *InFunction->GetName()));
diff --git a/Engine/Plugins/Runtime/ModelViewViewModel/Source/ModelViewViewModel/Private/Bindings/MVVMCompiledBindingLibrary.cpp b/Engine/Plugins/Runtime/ModelViewViewModel/Source/ModelViewViewModel/Private/Bindings/MVVMCompiledBindingLibrary.cpp
--- a/Engine/Plugins/Runtime/ModelViewViewModel/Source/ModelViewViewModel/Private/Bindings/MVVMCompiledBindingLibrary.cpp
+++ b/Engine/Plugins/Runtime/ModelViewViewModel/Source/ModelViewViewModel/Private/Bindings/MVVMCompiledBindingLibrary.cpp
@@ -49,6 +49,10 @@ FMVVMCompiledBindingLibrary::FMVVMCompiledBindingLibrary()
 
 void FMVVMCompiledBindingLibrary::Load()
 {
+	ensureAlwaysMsgf(LoadedProperties.Num() == 0, TEXT("The binding library was loaded more than once."));
+	ensureAlwaysMsgf(LoadedFunctions.Num() == 0, TEXT("The binding library was loaded more than once."));
+	ensureAlwaysMsgf(LoadedFieldIds.Num() == 0, TEXT("The binding library was loaded more than once."));
+
 	for (const FMVVMVCompiledFields& Field : CompiledFields)
 	{
 		{
diff --git a/Engine/Plugins/Runtime/ModelViewViewModel/Source/ModelViewViewModel/Private/Bindings/MVVMFieldPathHelper.cpp b/Engine/Plugins/Runtime/ModelViewViewModel/Source/ModelViewViewModel/Private/Bindings/MVVMFieldPathHelper.cpp
--- a/Engine/Plugins/Runtime/ModelViewViewModel/Source/ModelViewViewModel/Private/Bindings/MVVMFieldPathHelper.cpp
+++ b/Engine/Plugins/Runtime/ModelViewViewModel/Source/ModelViewViewModel/Private/Bindings/MVVMFieldPathHelper.cpp
@@ -197,7 +197,11 @@ TValueOrError<UObject*, void> EvaluateObjectProperty(const FFieldContext& InSour
 	else
 	{
 		check(InSource.GetObjectVariant().IsUObject());
+		check(InSource.GetObjectVariant().GetUObject());
+
 		UFunction* Function = InSource.GetFieldVariant().GetFunction();
+		check(InSource.GetObjectVariant().GetUObject()->GetClass()->IsChildOf(Function->GetOuterUClass()));
+
 		void* DataPtr = FMemory_Alloca_Aligned(Function->ParmsSize, Function->GetMinAlignment());
 		ObjectProperty->InitializeValue(DataPtr);
 		InSource.GetObjectVariant().GetUObject()->ProcessEvent(Function, DataPtr);
diff --git a/Engine/Plugins/Runtime/ModelViewViewModel/Source/ModelViewViewModel/Private/View/MVVMViewClass.cpp b/Engine/Plugins/Runtime/ModelViewViewModel/Source/ModelViewViewModel/Private/View/MVVMViewClass.cpp
--- a/Engine/Plugins/Runtime/ModelViewViewModel/Source/ModelViewViewModel/Private/View/MVVMViewClass.cpp
+++ b/Engine/Plugins/Runtime/ModelViewViewModel/Source/ModelViewViewModel/Private/View/MVVMViewClass.cpp
@@ -99,7 +99,7 @@ UObject* FMVVMViewClass_SourceCreator::CreateInstance(const UMVVMViewClass* InVi
 		else if (GlobalViewModelInstance.IsValid())
 		{
 			UMVVMViewModelBase* FoundViewModelInstance = GEngine->GetEngineSubsystem<UMVVMSubsystem>()->GetGlobalViewModelCollection()->FindViewModelInstance(GlobalViewModelInstance);
-			if (FoundViewModelInstance)
+			if (FoundViewModelInstance != nullptr)
 			{
 				ensureMsgf(FoundViewModelInstance->IsA(GlobalViewModelInstance.ContextClass), TEXT("The Global View Model Instance is not of the expected type."));
 				Result = FoundViewModelInstance;
@@ -117,7 +117,7 @@ UObject* FMVVMViewClass_SourceCreator::CreateInstance(const UMVVMViewClass* InVi
 			if (FieldPathResult.HasValue())
 			{
 				TValueOrError<UObject*, void> ObjectResult = UE::MVVM::FieldPathHelper::EvaluateObjectProperty(FieldPathResult.GetValue());
-				if (ObjectResult.HasValue())
+				if (ObjectResult.HasValue() && ObjectResult.GetValue() != nullptr)
 				{
 					Result = ObjectResult.GetValue();
 					AssignProperty(ObjectResult.GetValue());
@@ -160,6 +160,7 @@ void UMVVMViewClass::Initialize(UUserWidget* UserWidget)
 		if (!bLoaded)
 		{
 			BindingLibrary.Load();
+			bLoaded = true;
 		}
 
 		View->ConstructView(this);
diff --git a/Engine/Plugins/Runtime/ModelViewViewModel/Source/ModelViewViewModel/Public/Types/MVVMFunctionContext.h b/Engine/Plugins/Runtime/ModelViewViewModel/Source/ModelViewViewModel/Public/Types/MVVMFunctionContext.h
--- a/Engine/Plugins/Runtime/ModelViewViewModel/Source/ModelViewViewModel/Public/Types/MVVMFunctionContext.h
+++ b/Engine/Plugins/Runtime/ModelViewViewModel/Source/ModelViewViewModel/Public/Types/MVVMFunctionContext.h
@@ -91,8 +91,8 @@ namespace UE::MVVM
 
 		static FConstFunctionContext MakeStaticFunction(FunctionType InFunction)
 		{
-			check(InFunction->HasAnyFunctionFlags(FUNC_Static));
-			return InFunction ? FConstFunctionContext() : FConstFunctionContext(InFunction->GetOuterUClass()->GetDefaultObject(), InFunction);
+			check(InFunction && InFunction->HasAnyFunctionFlags(FUNC_Static));
+			return FConstFunctionContext(InFunction->GetOuterUClass()->GetDefaultObject(), InFunction);
 		}
 
 		FConstFunctionContext& operator=(const FFunctionContext& Other)
diff --git a/Engine/Plugins/Runtime/ModelViewViewModel/Source/ModelViewViewModelBlueprint/Private/Bindings/MVVMCompiledBindingLibraryCompiler.cpp b/Engine/Plugins/Runtime/ModelViewViewModel/Source/ModelViewViewModelBlueprint/Private/Bindings/MVVMCompiledBindingLibraryCompiler.cpp
--- a/Engine/Plugins/Runtime/ModelViewViewModel/Source/ModelViewViewModelBlueprint/Private/Bindings/MVVMCompiledBindingLibraryCompiler.cpp
+++ b/Engine/Plugins/Runtime/ModelViewViewModel/Source/ModelViewViewModelBlueprint/Private/Bindings/MVVMCompiledBindingLibraryCompiler.cpp
@@ -357,7 +357,60 @@ TValueOrError<FCompiledBindingLibraryCompiler::FFieldPathHandle, FString> FCompi
 
 TValueOrError<FCompiledBindingLibraryCompiler::FFieldPathHandle, FString> FCompiledBindingLibraryCompiler::AddConversionFunctionFieldPath(TSubclassOf<UObject> SourceClass, FStringView FieldPath)
 {
-	return MakeError(TEXT("Not implemented yet"));
+	Impl->bCompiled = false;
+
+	if (SourceClass == nullptr)
+	{
+		return MakeError(TEXT("The source class is invalid."));
+	}
+	if (FieldPath.Len() == 0)
+	{
+		return MakeError(TEXT("The function path is empty."));
+	}
+
+	FMVVMFieldVariant FoundFunction = FMVVMFieldVariant(FindObject<UFunction>(nullptr, FieldPath.GetData(), true));
+	if (FoundFunction.IsEmpty() || !FoundFunction.IsFunction() || FoundFunction.GetFunction() == nullptr)
+	{
+		FoundFunction = UE::MVVM::BindingHelper::FindFieldByName(SourceClass.Get(), FMVVMBindingName(FieldPath.GetData()));
+	}
+
+	if (FoundFunction.IsEmpty() || !FoundFunction.IsFunction() || FoundFunction.GetFunction() == nullptr)
+	{
+		return MakeError(FString::Printf(TEXT("The function %s could not be found."), FieldPath.GetData()));
+	}
+
+	if (!BindingHelper::IsValidForRuntimeConversion(FoundFunction.GetFunction()))
+	{
+		return MakeError(FString::Printf(TEXT("Function %s cannot be used as a runtime conversion function."), *FoundFunction.GetFunction()->GetName()));
+	}
+
+	if (!FoundFunction.GetFunction()->HasAllFunctionFlags(FUNC_Static))
+	{
+		if (!SourceClass->IsChildOf(FoundFunction.GetFunction()->GetOuterUClass()))
+		{
+			return MakeError(FString::Printf(TEXT("Function %s is going to be executed with an invalid self."), *FoundFunction.GetFunction()->GetName()));
+		}
+	}
+
+	TArray<int32> RawFieldIndexes;
+	RawFieldIndexes.Add(Impl->AddUniqueField(FoundFunction));
+	const int32 FoundFieldPath = Impl->FieldPaths.IndexOfByPredicate([&RawFieldIndexes](const Private::FRawFieldPath& Other)
+		{
+			return Other.RawFieldIndexes == RawFieldIndexes;
+		});
+	if (FoundFieldPath != INDEX_NONE)
+	{
+		return MakeValue(Impl->FieldPaths[FoundFieldPath].PathHandle);
+	}
+
+	Private::FRawFieldPath RawFieldPath;
+	RawFieldPath.RawFieldIndexes = MoveTemp(RawFieldIndexes);
+	RawFieldPath.PathHandle = FFieldPathHandle::MakeHandle();
+	RawFieldPath.bIsReadable = false;
+	RawFieldPath.bIsWritable = false;
+
+	const int32 NewFieldPathIndex = Impl->FieldPaths.Add(MoveTemp(RawFieldPath));
+	return MakeValue(Impl->FieldPaths[NewFieldPathIndex].PathHandle);
 }
 
 
@@ -727,6 +780,10 @@ TValueOrError<FCompiledBindingLibraryCompiler::FCompileResult, FString> FCompile
 		Result.Bindings.Add(Binding.BindingHandle, Binding.CompiledBinding);
 	}
 
+	Result.Library.LoadedProperties.Reset();
+	Result.Library.LoadedFunctions.Reset();
+	Result.Library.LoadedFieldIds.Reset();
+
 	Impl->bCompiled = true;
 	return MakeValue(MoveTemp(Result));
 }
diff --git a/Engine/Plugins/Runtime/ModelViewViewModel/Source/ModelViewViewModelBlueprint/Private/MVVMViewBlueprintCompiler.cpp b/Engine/Plugins/Runtime/ModelViewViewModel/Source/ModelViewViewModelBlueprint/Private/MVVMViewBlueprintCompiler.cpp
--- a/Engine/Plugins/Runtime/ModelViewViewModel/Source/ModelViewViewModelBlueprint/Private/MVVMViewBlueprintCompiler.cpp
+++ b/Engine/Plugins/Runtime/ModelViewViewModel/Source/ModelViewViewModelBlueprint/Private/MVVMViewBlueprintCompiler.cpp
@@ -92,7 +92,7 @@ void FMVVMViewBlueprintCompiler::CreateVariables(const FWidgetBlueprintCompilerC
 		return;
 	}
 
-	if (!bAreSourcesCreatorValid || !bIsBindingsValid || !bIsBindingsValid)
+	if (!bAreSourcesCreatorValid || !bAreSourceContextsValid || !bIsBindingsValid)
 	{
 		return;
 	}
@@ -140,6 +140,30 @@ void FMVVMViewBlueprintCompiler::CreateVariables(const FWidgetBlueprintCompilerC
 	{
 		SourceContext.Field = BindingHelper::FindFieldByName(Context.GetSkeletonGeneratedClass(), FMVVMBindingName(SourceContext.PropertyName));
 
+		// The class is not linked yet. It may not be available yet.
+		if (SourceContext.Field.IsEmpty())
+		{
+			for (FField* Field = Context.GetSkeletonGeneratedClass()->ChildProperties; Field != nullptr; Field = Field->Next)
+			{
+				if (Field->GetFName() == SourceContext.PropertyName)
+				{
+					if (FProperty* Property = CastField<FProperty>(Field))
+					{
+						SourceContext.Field = FMVVMFieldVariant(Property);
+						break;
+					}
+					else
+					{
+						WidgetBlueprintCompilerContext.MessageLog.Error(*FString::Printf(TEXT("The field for source '%s' exist but is not a property.")
+							, *SourceContext.DisplayName));
+						bAreSourcesCreatorValid = false;
+						continue;
+					}
+				}
+			}
+		}
+
+		// Reuse the property if found
 		if (!SourceContext.Field.IsEmpty())
 		{
 			if (!BindingHelper::IsValidForSourceBinding(SourceContext.Field))
@@ -359,7 +383,7 @@ void FMVVMViewBlueprintCompiler::CreateFunctionsDeclaration(const FWidgetBluepri
 				, TEXT("Viewmodel")
 				, false);
 
-			if (SourceCreator.SetterGraph->GetFName() != FName(*SourceCreator.SetterFunctionName))
+			if (SourceCreator.SetterGraph == nullptr || SourceCreator.SetterGraph->GetFName() != FName(*SourceCreator.SetterFunctionName))
 			{
 				WidgetBlueprintCompilerContext.MessageLog.Error(*FString::Printf(TEXT("The setter name %s already exist and could not be autogenerated."), *SourceCreator.SetterFunctionName));
 			}
@@ -372,7 +396,7 @@ void FMVVMViewBlueprintCompiler::CreateFunctionsDeclaration(const FWidgetBluepri
 
 bool FMVVMViewBlueprintCompiler::PreCompile(UWidgetBlueprintGeneratedClass* Class, UMVVMBlueprintView* BlueprintView)
 {
-	if (!bAreSourcesCreatorValid || !bIsBindingsValid || !bIsBindingsValid)
+	if (!bAreSourcesCreatorValid || !bAreSourceContextsValid || !bIsBindingsValid)
 	{
 		return false;
 	}
@@ -386,7 +410,7 @@ bool FMVVMViewBlueprintCompiler::PreCompile(UWidgetBlueprintGeneratedClass* Clas
 
 bool FMVVMViewBlueprintCompiler::Compile(UWidgetBlueprintGeneratedClass* Class, UMVVMBlueprintView* BlueprintView, UMVVMViewClass* ViewExtension)
 {
-	if (!bAreSourcesCreatorValid || !bIsBindingsValid || !bIsBindingsValid)
+	if (!bAreSourcesCreatorValid || !bAreSourceContextsValid || !bIsBindingsValid)
 	{
 		return false;
 	}
@@ -400,7 +424,7 @@ bool FMVVMViewBlueprintCompiler::Compile(UWidgetBlueprintGeneratedClass* Class,
 	CompileSourceCreators(CompileResult.GetValue(), Class, BlueprintView, ViewExtension);
 	CompileBindings(CompileResult.GetValue(), Class, BlueprintView, ViewExtension);
 
-	bool bResult = bAreSourcesCreatorValid && bIsBindingsValid && bIsBindingsValid;
+	bool bResult = bAreSourcesCreatorValid && bAreSourceContextsValid && bIsBindingsValid;
 	if (bResult)
 	{
 		ViewExtension->BindingLibrary = MoveTemp(CompileResult.GetValue().Library);
diff --git a/Engine/Plugins/Runtime/ModelViewViewModel/Source/ModelViewViewModelEditor/ModelViewViewModelEditor.Build.cs b/Engine/Plugins/Runtime/ModelViewViewModel/Source/ModelViewViewModelEditor/ModelViewViewModelEditor.Build.cs
--- a/Engine/Plugins/Runtime/ModelViewViewModel/Source/ModelViewViewModelEditor/ModelViewViewModelEditor.Build.cs
+++ b/Engine/Plugins/Runtime/ModelViewViewModel/Source/ModelViewViewModelEditor/ModelViewViewModelEditor.Build.cs
@@ -28,6 +28,7 @@ public class ModelViewViewModelEditor : ModuleRules
 				"BlueprintGraph",
 				"InputCore",
 				"Kismet",
+				"MessageLog",
 				"Projects",
 				"PropertyEditor",
 				"StatusBar",
diff --git a/Engine/Plugins/Runtime/ModelViewViewModel/Source/ModelViewViewModelEditor/Private/ModelViewViewModelEditorModule.cpp b/Engine/Plugins/Runtime/ModelViewViewModel/Source/ModelViewViewModelEditor/Private/ModelViewViewModelEditorModule.cpp
--- a/Engine/Plugins/Runtime/ModelViewViewModel/Source/ModelViewViewModelEditor/Private/ModelViewViewModelEditorModule.cpp
+++ b/Engine/Plugins/Runtime/ModelViewViewModel/Source/ModelViewViewModelEditor/Private/ModelViewViewModelEditorModule.cpp
@@ -12,6 +12,8 @@
 #include "Framework/MultiBox/MultiBoxBuilder.h"
 #include "IHasPropertyBindingExtensibility.h"
 #include "Kismet2/BlueprintEditorUtils.h"
+#include "MessageLogInitializationOptions.h"
+#include "MessageLogModule.h"
 #include "MVVMBlueprintView.h"
 #include "MVVMBlueprintViewModelContext.h"
 #include "MVVMEditorCommands.h"
@@ -52,12 +54,26 @@ void FModelViewViewModelEditorModule::StartupModule()
 
 	FBlueprintEditorUtils::OnRenameVariableReferencesEvent.AddRaw(this, &FModelViewViewModelEditorModule::HandleRenameVariableReferences);
 
+	{
+		FMessageLogModule& MessageLogModule = FModuleManager::LoadModuleChecked<FMessageLogModule>("MessageLog");
+		FMessageLogInitializationOptions InitOptions;
+		InitOptions.bShowFilters = true;
+		InitOptions.bShowPages = false;
+		InitOptions.bAllowClear = true;
+		MessageLogModule.RegisterLogListing("Model View Viewmodel", LOCTEXT("MVVMLog", "Model View Viewmodel"), InitOptions);
+	}
+
 	FMVVMEditorCommands::Register();
 }
 
 
 void FModelViewViewModelEditorModule::ShutdownModule()
 {
+	if (FMessageLogModule* MessageLogModule = FModuleManager::GetModulePtr<FMessageLogModule>("MessageLog"))
+	{
+		MessageLogModule->UnregisterLogListing("Model View Viewmodel");
+	}
+
 	FBlueprintEditorUtils::OnRenameVariableReferencesEvent.RemoveAll(this);
 
 	if (IUMGEditorModule* UMGEditorModule = FModuleManager::GetModulePtr<IUMGEditorModule>("UMGEditor"))
diff --git a/Engine/Source/Editor/UMGEditor/Private/WidgetBlueprintCompiler.cpp b/Engine/Source/Editor/UMGEditor/Private/WidgetBlueprintCompiler.cpp
--- a/Engine/Source/Editor/UMGEditor/Private/WidgetBlueprintCompiler.cpp
+++ b/Engine/Source/Editor/UMGEditor/Private/WidgetBlueprintCompiler.cpp
@@ -998,6 +998,14 @@ void FWidgetBlueprintCompilerContext::FinishCompilingClass(UClass* Class)
 
 	Super::FinishCompilingClass(Class);
 
+	if (!bIsSkeletonOnly)
+	{
+		if (UUserWidget* UserWidget = Cast<UUserWidget>(BPGClass->GetDefaultObject()))
+		{
+			BPGClass->InitializeFieldNotification(UserWidget);
+		}
+	}
+
 	UWidgetBlueprintExtension::ForEachExtension(WidgetBlueprint(), [BPGClass](UWidgetBlueprintExtension* InExtension)
 		{
 			InExtension->FinishCompilingClass(BPGClass);
diff --git a/Engine/Source/Runtime/UMG/Private/WidgetBlueprintGeneratedClass.cpp b/Engine/Source/Runtime/UMG/Private/WidgetBlueprintGeneratedClass.cpp
--- a/Engine/Source/Runtime/UMG/Private/WidgetBlueprintGeneratedClass.cpp
+++ b/Engine/Source/Runtime/UMG/Private/WidgetBlueprintGeneratedClass.cpp
@@ -223,6 +223,7 @@ void UWidgetBlueprintGeneratedClass::InitializeWidgetStatic(UUserWidget* UserWid
 		for (TFieldIterator<FObjectPropertyBase>It(WidgetBlueprintClass, EFieldIterationFlags::Default); It; ++It)
 		{
 			check(*It);
+			ensureMsgf(!ObjectPropertiesMap.Contains(It->GetFName()), TEXT("There are properties with the same names. %s"), *It->GetName());
 			ObjectPropertiesMap.Add(It->GetFName(), *It);
 		}
 
