MVVM: TryGetPropertyTypeForConversionFunction use the correct ALL flags. Fix typo. Fix CIS warning about duplicate test.

diff --git a/Engine/Plugins/Runtime/ModelViewViewModel/Source/ModelViewViewModel/Private/Bindings/MVVMBindingHelper.cpp b/Engine/Plugins/Runtime/ModelViewViewModel/Source/ModelViewViewModel/Private/Bindings/MVVMBindingHelper.cpp
--- a/Engine/Plugins/Runtime/ModelViewViewModel/Source/ModelViewViewModel/Private/Bindings/MVVMBindingHelper.cpp
+++ b/Engine/Plugins/Runtime/ModelViewViewModel/Source/ModelViewViewModel/Private/Bindings/MVVMBindingHelper.cpp
@@ -114,7 +114,7 @@ namespace UE::MVVM::BindingHelper
 	bool IsValidForRuntimeConversion(const UFunction* InFunction)
 	{
 		if (Private::IsValidCommon(InFunction)
-			&& (InFunction->HasAllFunctionFlags(FUNC_Static) || InFunction->HasAnyFunctionFlags(FUNC_Const | FUNC_BlueprintPure))
+			&& (InFunction->HasAllFunctionFlags(FUNC_Static) || InFunction->HasAllFunctionFlags(FUNC_Const | FUNC_BlueprintPure))
 			&& InFunction->NumParms == 2)
 		{
 			const FProperty* ReturnProperty = GetReturnProperty(InFunction);
@@ -381,7 +381,7 @@ namespace UE::MVVM::BindingHelper
 			return MakeError(MoveTemp(CommonResult));
 		}
 
-		if (!InFunction->HasAllFunctionFlags(FUNC_Static) && !InFunction->HasAnyFunctionFlags(FUNC_Const | FUNC_BlueprintPure))
+		if (!InFunction->HasAllFunctionFlags(FUNC_Static) && !InFunction->HasAllFunctionFlags(FUNC_Const | FUNC_BlueprintPure))
 		{
 			return MakeError(FString::Printf(TEXT("The function '%s' is not static or is not const and pure."), *InFunction->GetName()));
 		}
diff --git a/Engine/Plugins/Runtime/ModelViewViewModel/Source/ModelViewViewModelBlueprint/Private/MVVMViewBlueprintCompiler.cpp b/Engine/Plugins/Runtime/ModelViewViewModel/Source/ModelViewViewModelBlueprint/Private/MVVMViewBlueprintCompiler.cpp
--- a/Engine/Plugins/Runtime/ModelViewViewModel/Source/ModelViewViewModelBlueprint/Private/MVVMViewBlueprintCompiler.cpp
+++ b/Engine/Plugins/Runtime/ModelViewViewModel/Source/ModelViewViewModelBlueprint/Private/MVVMViewBlueprintCompiler.cpp
@@ -154,7 +154,7 @@ void FMVVMViewBlueprintCompiler::CreateVariables(const FWidgetBlueprintCompilerC
 					}
 					else
 					{
-						WidgetBlueprintCompilerContext.MessageLog.Error(*FString::Printf(TEXT("The field for source '%s' exist but is not a property.")
+						WidgetBlueprintCompilerContext.MessageLog.Error(*FString::Printf(TEXT("The field for source '%s' exists but is not a property.")
 							, *SourceContext.DisplayName));
 						bAreSourcesCreatorValid = false;
 						continue;
@@ -404,7 +404,7 @@ bool FMVVMViewBlueprintCompiler::PreCompile(UWidgetBlueprintGeneratedClass* Clas
 	PreCompileSourceCreators(Class, BlueprintView);
 	PreCompileBindings(Class, BlueprintView);
 
-	return bAreSourcesCreatorValid && bIsBindingsValid && bIsBindingsValid;
+	return bAreSourcesCreatorValid && bAreSourceContextsValid && bIsBindingsValid;
 }
 
 
diff --git a/Engine/Source/Runtime/UMG/Private/WidgetBlueprintGeneratedClass.cpp b/Engine/Source/Runtime/UMG/Private/WidgetBlueprintGeneratedClass.cpp
--- a/Engine/Source/Runtime/UMG/Private/WidgetBlueprintGeneratedClass.cpp
+++ b/Engine/Source/Runtime/UMG/Private/WidgetBlueprintGeneratedClass.cpp
@@ -223,7 +223,7 @@ void UWidgetBlueprintGeneratedClass::InitializeWidgetStatic(UUserWidget* UserWid
 		for (TFieldIterator<FObjectPropertyBase>It(WidgetBlueprintClass, EFieldIterationFlags::Default); It; ++It)
 		{
 			check(*It);
-			ensureMsgf(!ObjectPropertiesMap.Contains(It->GetFName()), TEXT("There are properties with the same names. %s"), *It->GetName());
+			ensureMsgf(!ObjectPropertiesMap.Contains(It->GetFName()), TEXT("There are properties with the same names: '%s'"), *It->GetName());
 			ObjectPropertiesMap.Add(It->GetFName(), *It);
 		}
 
