MVVM: First version of the viewmodel editor

diff --git a/Engine/Plugins/Runtime/ModelViewViewModel/Source/ModelViewViewModel/Private/MVVMViewModelBase.cpp b/Engine/Plugins/Runtime/ModelViewViewModel/Source/ModelViewViewModel/Private/MVVMViewModelBase.cpp
--- a/Engine/Plugins/Runtime/ModelViewViewModel/Source/ModelViewViewModel/Private/MVVMViewModelBase.cpp
+++ b/Engine/Plugins/Runtime/ModelViewViewModel/Source/ModelViewViewModel/Private/MVVMViewModelBase.cpp
@@ -2,6 +2,7 @@
 
 #include "MVVMViewModelBase.h"
 
+#include "ViewModel/MVVMViewModelBlueprintGeneratedClass.h"
 #include "Bindings/MVVMBindingHelper.h"
 
 
@@ -67,6 +68,15 @@ const UE::FieldNotification::IClassDescriptor& UMVVMViewModelBase::GetFieldNotif
 }
 
 
+void UMVVMViewModelBase::FFieldNotificationClassDescriptor::ForEachField(const UClass* Class, TFunctionRef<bool(::UE::FieldNotification::FFieldId FieldId)> Callback) const
+{
+	if (const UMVVMViewModelBlueprintGeneratedClass* ViewModelBPClass = Cast<const UMVVMViewModelBlueprintGeneratedClass>(Class))
+	{
+		ViewModelBPClass->ForEachField(Callback);
+	}
+}
+
+
 void UMVVMViewModelBase::BindingFieldValueChanged(UE::FieldNotification::FFieldId InFieldId)
 {
 	if (InFieldId.IsValid() && EnabledFieldNotifications.IsValidIndex(InFieldId.GetIndex()) && EnabledFieldNotifications[InFieldId.GetIndex()])
@@ -76,9 +86,8 @@ void UMVVMViewModelBase::BindingFieldValueChanged(UE::FieldNotification::FFieldI
 }
 
 
-void UMVVMViewModelBase::BindingValueChanged(FMVVMBindingName ViewModelPropertyName)
+void UMVVMViewModelBase::K2_BroadcastFieldValueChanged(FFieldNotificationId ViewModelPropertyName)
 {
-	UE::FieldNotification::FFieldId FieldId = GetFieldNotificationDescriptor().GetField(GetClass(), ViewModelPropertyName.ToName());
+	UE::FieldNotification::FFieldId FieldId = GetFieldNotificationDescriptor().GetField(GetClass(), ViewModelPropertyName.GetFieldName());
 	BindingFieldValueChanged(FieldId);
 }
-
diff --git a/Engine/Plugins/Runtime/ModelViewViewModel/Source/ModelViewViewModel/Private/ViewModel/MVVMViewModelBlueprintGeneratedClass.cpp b/Engine/Plugins/Runtime/ModelViewViewModel/Source/ModelViewViewModel/Private/ViewModel/MVVMViewModelBlueprintGeneratedClass.cpp
new file mode 100644
--- /dev/null
+++ b/Engine/Plugins/Runtime/ModelViewViewModel/Source/ModelViewViewModel/Private/ViewModel/MVVMViewModelBlueprintGeneratedClass.cpp
@@ -0,0 +1,52 @@
+// Copyright Epic Games, Inc. All Rights Reserved.
+
+#include "ViewModel/MVVMViewModelBlueprintGeneratedClass.h"
+#include "MVVMViewModelBase.h"
+
+UMVVMViewModelBlueprintGeneratedClass::UMVVMViewModelBlueprintGeneratedClass()
+	: FieldNotifyStartBitNumber(INDEX_NONE)
+{}
+
+
+void UMVVMViewModelBlueprintGeneratedClass::PostLoadDefaultObject(UObject* Object)
+{
+	Super::PostLoadDefaultObject(Object);
+	InitializeFieldNotification(Cast<UMVVMViewModelBase>(Object));
+}
+
+
+void UMVVMViewModelBlueprintGeneratedClass::PurgeClass(bool bRecompilingOnLoad)
+{
+	Super::PurgeClass(bRecompilingOnLoad);
+	FieldNotifyNames.Empty();
+}
+
+
+void UMVVMViewModelBlueprintGeneratedClass::InitializeFieldNotification(const UMVVMViewModelBase* ViewModel)
+{
+	FieldNotifyStartBitNumber = 0;
+	if (ViewModel && FieldNotifyNames.Num())
+	{
+		int32 NumberOfField = 0;
+		ViewModel->GetFieldNotificationDescriptor().ForEachField(this, [&NumberOfField](::UE::FieldNotification::FFieldId FielId)
+			{
+				++NumberOfField;
+				return true;
+			});
+		FieldNotifyStartBitNumber = NumberOfField - FieldNotifyNames.Num();
+		ensureMsgf(FieldNotifyStartBitNumber >= 0, TEXT("The FieldNotifyStartIndex is negative. The number of field should be positive."));
+	}
+}
+
+
+void UMVVMViewModelBlueprintGeneratedClass::ForEachField(TFunctionRef<bool(::UE::FieldNotification::FFieldId FieldId)> Callback) const
+{
+	ensureMsgf(FieldNotifyStartBitNumber >= 0, TEXT("The FieldNotifyStartIndex is negative. The number of field should be positive."));
+	for (int32 Index = 0; Index < FieldNotifyNames.Num(); ++Index)
+	{
+		if (!Callback(UE::FieldNotification::FFieldId(FieldNotifyNames[Index].GetFieldName(), Index + FieldNotifyStartBitNumber)))
+		{
+			break;
+		}
+	}
+}
diff --git a/Engine/Plugins/Runtime/ModelViewViewModel/Source/ModelViewViewModel/Public/MVVMViewModelBase.h b/Engine/Plugins/Runtime/ModelViewViewModel/Source/ModelViewViewModel/Public/MVVMViewModelBase.h
--- a/Engine/Plugins/Runtime/ModelViewViewModel/Source/ModelViewViewModel/Public/MVVMViewModelBase.h
+++ b/Engine/Plugins/Runtime/ModelViewViewModel/Source/ModelViewViewModel/Public/MVVMViewModelBase.h
@@ -5,6 +5,7 @@
 #include "CoreMinimal.h"
 #include "UObject/Object.h"
 
+#include "FieldNotification/FieldId.h"
 #include "FieldNotification/FieldNotificationDeclaration.h"
 #include "FieldNotification/FieldMulticastDelegate.h"
 #include "FieldNotification/IFieldValueChanged.h"
@@ -46,8 +47,9 @@ class MODELVIEWVIEWMODEL_API UMVVMViewModelBase : public UObject, public INotify
 	GENERATED_BODY()
 
 public:
-	struct FFieldNotificationClassDescriptor : public UE::FieldNotification::IClassDescriptor
+	struct MODELVIEWVIEWMODEL_API FFieldNotificationClassDescriptor : public ::UE::FieldNotification::IClassDescriptor
 	{
+		virtual void ForEachField(const UClass* Class, TFunctionRef<bool(UE::FieldNotification::FFieldId FielId)> Callback) const override;
 	};
 
 public:
@@ -60,16 +62,8 @@ public:
 	//~ End INotifyFieldValueChanged Interface
 
 protected:
-	// todo make a 2k node to create a binding base on a variable or a function
-	UFUNCTION(BlueprintCallable, Category="MVVM")
-	static FMVVMBindingName MakeBindingName_Temporary(FName BindingName)
-	{
-		return FMVVMBindingName(BindingName);
-	}
-
-protected:
-	UFUNCTION(BlueprintCallable, Category = "MVVM")
-	void BindingValueChanged(FMVVMBindingName PropertyName);
+	UFUNCTION(BlueprintCallable, Category="FieldNotify", meta=(DisplayName="Broadcast Field Value Changed", ScriptName="BroadcastFieldValueChanged"))
+	void K2_BroadcastFieldValueChanged(FFieldNotificationId FieldId);
 	void BindingFieldValueChanged(UE::FieldNotification::FFieldId InFieldId);
 
 protected:
diff --git a/Engine/Plugins/Runtime/ModelViewViewModel/Source/ModelViewViewModel/Public/ViewModel/MVVMViewModelBlueprintGeneratedClass.h b/Engine/Plugins/Runtime/ModelViewViewModel/Source/ModelViewViewModel/Public/ViewModel/MVVMViewModelBlueprintGeneratedClass.h
new file mode 100644
--- /dev/null
+++ b/Engine/Plugins/Runtime/ModelViewViewModel/Source/ModelViewViewModel/Public/ViewModel/MVVMViewModelBlueprintGeneratedClass.h
@@ -0,0 +1,43 @@
+// Copyright Epic Games, Inc. All Rights Reserved.
+
+#pragma once
+
+#include "CoreMinimal.h"
+#include "UObject/ObjectMacros.h"
+#include "Engine/BlueprintGeneratedClass.h"
+#include "FieldNotification/FieldId.h"
+
+#include "MVVMViewModelBlueprintGeneratedClass.generated.h"
+
+class UMVVMViewModelBase;
+
+namespace UE::MVVM
+{
+	class FViewModelBlueprintCompilerContext;
+}//namespace
+
+UCLASS()
+class MODELVIEWVIEWMODEL_API UMVVMViewModelBlueprintGeneratedClass : public UBlueprintGeneratedClass
+{
+	GENERATED_BODY()
+
+	friend UE::MVVM::FViewModelBlueprintCompilerContext;
+
+public:
+	UMVVMViewModelBlueprintGeneratedClass();
+
+	//~ Begin UBlueprintGeneratedClass interface
+	virtual void PostLoadDefaultObject(UObject* Object) override;
+	virtual void PurgeClass(bool bRecompilingOnLoad) override;
+	//~ End UBlueprintGeneratedClass interface
+
+	void InitializeFieldNotification(const UMVVMViewModelBase* ViewModel);
+	void ForEachField(TFunctionRef<bool(::UE::FieldNotification::FFieldId FieldId)> Callback) const;
+
+private:
+	/** List Field Notifies. No index here on purpose to prevent saving them. */
+	UPROPERTY()
+	TArray<FFieldNotificationId> FieldNotifyNames;
+
+	int32 FieldNotifyStartBitNumber;
+};
diff --git a/Engine/Plugins/Runtime/ModelViewViewModel/Source/ModelViewViewModelBlueprint/Private/MVVMFunctionGraphHelper.cpp b/Engine/Plugins/Runtime/ModelViewViewModel/Source/ModelViewViewModelBlueprint/Private/MVVMFunctionGraphHelper.cpp
--- a/Engine/Plugins/Runtime/ModelViewViewModel/Source/ModelViewViewModelBlueprint/Private/MVVMFunctionGraphHelper.cpp
+++ b/Engine/Plugins/Runtime/ModelViewViewModel/Source/ModelViewViewModelBlueprint/Private/MVVMFunctionGraphHelper.cpp
@@ -2,18 +2,24 @@
 
 #include "MVVMFunctionGraphHelper.h"
 
+#include "BlueprintTypePromotion.h"
 #include "EngineLogs.h"
 #include "EdGraph/EdGraph.h"
 #include "EdGraph/EdGraphSchema.h"
 #include "EdGraphSchema_K2.h"
 #include "K2Node_CallFunction.h"
 #include "K2Node_DynamicCast.h"
-#include "K2Node_FunctionEntry.h"
 #include "K2Node_EditablePinBase.h"
+#include "K2Node_FunctionEntry.h"
+#include "K2Node_IfThenElse.h"
+#include "K2Node_PromotableOperator.h"
+#include "K2Node_VariableGet.h"
+#include "K2Node_VariableSet.h"
 #include "KismetCompiler.h"
 #include "Kismet2/BlueprintEditorUtils.h"
 
 #include "View/MVVMView.h"
+#include "MVVMViewModelBase.h"
 
 
 namespace UE::MVVM::FunctionGraphHelper
@@ -63,7 +69,7 @@ UEdGraph* CreateIntermediateFunctionGraph(FKismetCompilerContext& InContext, FSt
 
 
 bool AddFunctionArgument(UEdGraph* InFunctionGraph, TSubclassOf<UObject> InArgument, FName InArgumentName)
-{	
+{
 	for (UEdGraphNode* Node : InFunctionGraph->Nodes)
 	{
 		if (UK2Node_FunctionEntry* FunctionEntry = Cast<UK2Node_FunctionEntry>(Node))
@@ -83,6 +89,28 @@ bool AddFunctionArgument(UEdGraph* InFunctionGraph, TSubclassOf<UObject> InArgum
 }
 
 
+bool AddFunctionArgument(UEdGraph* InFunctionGraph, FProperty* InArgument, FName InArgumentName)
+{
+	for (UEdGraphNode* Node : InFunctionGraph->Nodes)
+	{
+		if (UK2Node_FunctionEntry* FunctionEntry = Cast<UK2Node_FunctionEntry>(Node))
+		{
+			const UEdGraphSchema_K2* K2Schema = GetDefault<UEdGraphSchema_K2>();
+
+			TSharedPtr<FUserPinInfo> PinInfo = MakeShared<FUserPinInfo>();
+			K2Schema->ConvertPropertyToPinType(InArgument, PinInfo->PinType);
+			PinInfo->PinName = InArgumentName;
+			PinInfo->DesiredPinDirection = EGPD_Output;
+			FunctionEntry->UserDefinedPins.Add(PinInfo);
+
+			FunctionEntry->ReconstructNode();
+			return true;
+		}
+	}
+	return false;
+}
+
+
 bool GenerateViewModelSetter(FKismetCompilerContext& InContext, UEdGraph* InFunctionGraph, FName InViewModelName)
 {
 	check(InFunctionGraph);
@@ -181,4 +209,256 @@ bool GenerateViewModelSetter(FKismetCompilerContext& InContext, UEdGraph* InFunc
 	return bResult;
 }
 
+
+bool GenerateViewModelFieldNotifySetter(FKismetCompilerContext& InContext, UEdGraph* InFunctionGraph, FProperty* InProperty, FName InInputPinName)
+{
+	check(InFunctionGraph);
+
+	UK2Node_FunctionEntry* FunctionEntry = nullptr;
+	for (UEdGraphNode* Node : InFunctionGraph->Nodes)
+	{
+		FunctionEntry = Cast<UK2Node_FunctionEntry>(Node);
+		if (FunctionEntry)
+		{
+			break;
+		}
+	}
+
+	if (FunctionEntry == nullptr)
+	{
+		return false;
+	}
+
+	const UEdGraphSchema* Schema = InFunctionGraph->GetSchema();
+	if (Schema == nullptr)
+	{
+		return false;
+	}
+
+	UK2Node_VariableSet* SetVariableNode = InContext.SpawnIntermediateNode<UK2Node_VariableSet>(FunctionEntry, InFunctionGraph);
+	{
+		//SetVariableNode->SetFromProperty(InProperty, true, nullptr);
+		SetVariableNode->VariableReference.SetSelfMember(InProperty->GetFName());
+		SetVariableNode->AllocateDefaultPins();
+	}
+
+	UK2Node_CallFunction* CallBroadcast = InContext.SpawnIntermediateNode<UK2Node_CallFunction>(FunctionEntry, InFunctionGraph);
+	{
+		CallBroadcast->FunctionReference.SetExternalMember("K2_BroadcastFieldValueChanged", UMVVMViewModelBase::StaticClass());
+		CallBroadcast->AllocateDefaultPins();
+
+		UEdGraphPin* NamePin = CallBroadcast->FindPin(TEXT("FieldId"), EGPD_Input);
+		ensure(NamePin);
+		if (!NamePin)
+		{
+			return false;
+		}
+
+		FFieldNotificationId NewPropertyNotifyId;
+		NewPropertyNotifyId.FieldName = InProperty->GetFName();
+
+		FString ValueString;
+		FFieldNotificationId::StaticStruct()->ExportText(ValueString, &NewPropertyNotifyId, nullptr, nullptr, EPropertyPortFlags::PPF_None, nullptr);
+
+		bool bMarkAsModified = false;
+		Schema->TrySetDefaultValue(*NamePin, ValueString, bMarkAsModified);
+	}
+
+	UK2Node_VariableGet* GetVariableNode = InContext.SpawnIntermediateNode<UK2Node_VariableGet>(FunctionEntry, InFunctionGraph);
+	{
+		//GetVariableNode->SetFromProperty(InProperty, true, nullptr);
+		GetVariableNode->VariableReference.SetSelfMember(InProperty->GetFName());
+		GetVariableNode->AllocateDefaultPins();
+	}
+
+	UK2Node_IfThenElse* BranchNode = InContext.SpawnIntermediateNode<UK2Node_IfThenElse>(FunctionEntry, InFunctionGraph);
+	{
+		BranchNode->AllocateDefaultPins();
+	}
+
+	UFunction* CompareFunction = nullptr;
+	{
+		TArray<UEdGraphPin*> PromotableOperatorPinsToConsider;
+		UEdGraphPin* EntryValuePin = FunctionEntry->FindPin(InInputPinName, EGPD_Output);
+		if (ensure(EntryValuePin))
+		{
+			PromotableOperatorPinsToConsider.Add(EntryValuePin);
+		}
+		UEdGraphPin* GetVariableNode_Property = GetVariableNode->FindPin(InProperty->GetFName(), EGPD_Output);
+		if (ensure(GetVariableNode_Property))
+		{
+			PromotableOperatorPinsToConsider.Add(GetVariableNode_Property);
+		}
+		UEdGraphPin* BranchCondition = BranchNode->FindPin(UEdGraphSchema_K2::PN_Condition, EGPD_Input);
+		if (ensure(BranchCondition))
+		{
+			PromotableOperatorPinsToConsider.Add(BranchCondition);
+		}
+
+		CompareFunction = FTypePromotion::FindBestMatchingFunc(TEXT("NotEqual"), PromotableOperatorPinsToConsider);
+	}
+
+	bool bResult = true;
+
+	bool bDoBranchTest = CompareFunction != nullptr;
+	if (bDoBranchTest)
+	{
+		UK2Node_PromotableOperator* CompareOperator = InContext.SpawnIntermediateNode<UK2Node_PromotableOperator>(FunctionEntry, InFunctionGraph);
+		{
+			CompareOperator->SetFromFunction(CompareFunction);
+			CompareOperator->AllocateDefaultPins();
+		}
+
+		// Entry -> Branch
+		{
+			UEdGraphPin* ThenPin = FunctionEntry->FindPin(UEdGraphSchema_K2::PN_Then, EGPD_Output);
+			UEdGraphPin* ExecPin = BranchNode->FindPin(UEdGraphSchema_K2::PN_Execute, EGPD_Input);
+			if (ensure(ThenPin && ExecPin))
+			{
+				ensure(Schema->TryCreateConnection(ThenPin, ExecPin));
+			}
+			bResult = bResult && ThenPin && ExecPin;
+		}
+
+		// Entry -> Compare, Get -> Compare, Compare -> Branch
+		{
+			UEdGraphPin* EntryValuePin = FunctionEntry->FindPin(InInputPinName, EGPD_Output);
+			UEdGraphPin* TopInputPin = CompareOperator->FindPin(TEXT("A"), EGPD_Input);
+			if (ensure(EntryValuePin && TopInputPin))
+			{
+				ensure(Schema->TryCreateConnection(EntryValuePin, TopInputPin));
+			}
+			bResult = bResult && EntryValuePin && EntryValuePin;
+
+			UEdGraphPin* GetPin = GetVariableNode->FindPin(InProperty->GetFName(), EGPD_Output);
+			UEdGraphPin* BottomInputPin = CompareOperator->FindPin(TEXT("B"), EGPD_Input);
+			if (ensure(GetPin && BottomInputPin))
+			{
+				ensure(Schema->TryCreateConnection(GetPin, BottomInputPin));
+			}
+			bResult = bResult && GetPin && BottomInputPin;
+
+			UEdGraphPin* OutputPin = CompareOperator->GetOutputPin();
+			UEdGraphPin* BranchCondition = BranchNode->FindPin(UEdGraphSchema_K2::PN_Condition, EGPD_Input);
+			if (ensure(BranchCondition && OutputPin))
+			{
+				ensure(Schema->TryCreateConnection(OutputPin, BranchCondition));
+			}
+			bResult = bResult && BranchCondition && OutputPin;
+		}
+
+		// Branch -> Set
+		{
+			UEdGraphPin* ThenPin = BranchNode->FindPin(UEdGraphSchema_K2::PN_Then, EGPD_Output);
+			UEdGraphPin* ExecPin = SetVariableNode->FindPin(UEdGraphSchema_K2::PN_Execute, EGPD_Input);
+			if (ensure(ThenPin && ExecPin))
+			{
+				ensure(Schema->TryCreateConnection(ThenPin, ExecPin));
+			}
+			bResult = bResult && ThenPin && ExecPin;
+		}
+	}
+	else
+	{
+		// Entry -> Set
+		{
+			UEdGraphPin* ThenPin = FunctionEntry->FindPin(UEdGraphSchema_K2::PN_Then, EGPD_Output);
+			UEdGraphPin* ExecPin = SetVariableNode->FindPin(UEdGraphSchema_K2::PN_Execute, EGPD_Input);
+			if (ensure(ThenPin && ExecPin))
+			{
+				ensure(Schema->TryCreateConnection(ThenPin, ExecPin));
+			}
+			bResult = bResult && ThenPin && ExecPin;
+		}
+	}
+
+	// Entry -> Set {Set value}
+	{
+		UEdGraphPin* EntryValuePin = FunctionEntry->FindPin(InInputPinName, EGPD_Output);
+		UEdGraphPin* SetValuePin = SetVariableNode->FindPin(InProperty->GetFName(), EGPD_Input);
+		if (ensure(EntryValuePin && SetValuePin))
+		{
+			ensure(Schema->TryCreateConnection(EntryValuePin, SetValuePin));
+		}
+		bResult = bResult && EntryValuePin && SetValuePin;
+	}
+
+	// Set -> Broadcast
+	{
+		UEdGraphPin* ThenPin = SetVariableNode->FindPin(UEdGraphSchema_K2::PN_Then, EGPD_Output);
+		UEdGraphPin* ExecPin = CallBroadcast->FindPin(UEdGraphSchema_K2::PN_Execute, EGPD_Input);
+		if (ensure(ThenPin && ExecPin))
+		{
+			ensure(Schema->TryCreateConnection(ThenPin, ExecPin));
+		}
+		bResult = bResult && ThenPin && ExecPin;
+	}
+
+	return bResult;
+}
+
+
+bool GenerateViewModelFielNotifyBroadcast(FKismetCompilerContext& InContext, UEdGraph* InFunctionGraph, FProperty* InProperty)
+{
+	check(InFunctionGraph);
+
+	UK2Node_FunctionEntry* FunctionEntry = nullptr;
+	for (UEdGraphNode* Node : InFunctionGraph->Nodes)
+	{
+		FunctionEntry = Cast<UK2Node_FunctionEntry>(Node);
+		if (FunctionEntry)
+		{
+			break;
+		}
+	}
+
+	if (FunctionEntry == nullptr)
+	{
+		return false;
+	}
+
+	const UEdGraphSchema* Schema = InFunctionGraph->GetSchema();
+	if (Schema == nullptr)
+	{
+		return false;
+	}
+
+	UK2Node_CallFunction* CallBroadcast = InContext.SpawnIntermediateNode<UK2Node_CallFunction>(FunctionEntry, InFunctionGraph);
+	{
+		CallBroadcast->FunctionReference.SetExternalMember("K2_BroadcastFieldValueChanged", UMVVMViewModelBase::StaticClass());
+		CallBroadcast->AllocateDefaultPins();
+
+		UEdGraphPin* NamePin = CallBroadcast->FindPin(TEXT("FieldId"), EGPD_Input);
+		ensure(NamePin);
+		if (!NamePin)
+		{
+			return false;
+		}
+
+		FFieldNotificationId NewPropertyNotifyId;
+		NewPropertyNotifyId.FieldName = InProperty->GetFName();
+
+		FString ValueString;
+		FFieldNotificationId::StaticStruct()->ExportText(ValueString, &NewPropertyNotifyId, nullptr, nullptr, EPropertyPortFlags::PPF_None, nullptr);
+
+		bool bMarkAsModified = false;
+		Schema->TrySetDefaultValue(*NamePin, ValueString, bMarkAsModified);
+	}
+
+	bool bResult = true;
+
+	// Entry -> Broadcast
+	{
+		UEdGraphPin* ThenPin = FunctionEntry->FindPin(UEdGraphSchema_K2::PN_Then, EGPD_Output);
+		UEdGraphPin* ExecPin = CallBroadcast->FindPin(UEdGraphSchema_K2::PN_Execute, EGPD_Input);
+		if (ensure(ThenPin && ExecPin))
+		{
+			ensure(Schema->TryCreateConnection(ThenPin, ExecPin));
+		}
+		bResult = bResult && ThenPin && ExecPin;
+	}
+
+	return bResult;
+}
+
 } //namespace
\ No newline at end of file
diff --git a/Engine/Plugins/Runtime/ModelViewViewModel/Source/ModelViewViewModelBlueprint/Private/MVVMFunctionGraphHelper.h b/Engine/Plugins/Runtime/ModelViewViewModel/Source/ModelViewViewModelBlueprint/Private/MVVMFunctionGraphHelper.h
--- a/Engine/Plugins/Runtime/ModelViewViewModel/Source/ModelViewViewModelBlueprint/Private/MVVMFunctionGraphHelper.h
+++ b/Engine/Plugins/Runtime/ModelViewViewModel/Source/ModelViewViewModelBlueprint/Private/MVVMFunctionGraphHelper.h
@@ -21,7 +21,16 @@ namespace UE::MVVM::FunctionGraphHelper
 	/** Add an input argument to an existing function graph. */
 	MODELVIEWVIEWMODELBLUEPRINT_API bool AddFunctionArgument(UEdGraph* FunctionGraph, TSubclassOf<UObject> Argument, FName ArgumentName);
 
+	/** Add an input argument to an existing function graph. */
+	MODELVIEWVIEWMODELBLUEPRINT_API bool AddFunctionArgument(UEdGraph* FunctionGraph, FProperty* Argument, FName ArgumentName);
+
 	/** Generate the nodes needed for to call the MVVMView::SetViewModel function. */
-	MODELVIEWVIEWMODELBLUEPRINT_API bool GenerateViewModelSetter(FKismetCompilerContext& InContext, UEdGraph* InFunctionGraph, FName InViewModelName);
+	MODELVIEWVIEWMODELBLUEPRINT_API bool GenerateViewModelSetter(FKismetCompilerContext& Context, UEdGraph* FunctionGraph, FName IiewModelName);
+
+	/** Generate the nodes needed for to set the property and to call the broadcast field notify function. */
+	MODELVIEWVIEWMODELBLUEPRINT_API bool GenerateViewModelFieldNotifySetter(FKismetCompilerContext& Context, UEdGraph* FunctionGraph, FProperty* Property, FName InputPinName);
+
+	/** Generate the nodes needed to broadcast field notify function. */
+	MODELVIEWVIEWMODELBLUEPRINT_API bool GenerateViewModelFielNotifyBroadcast(FKismetCompilerContext& Context, UEdGraph* FunctionGraph, FProperty* Property);
 
 } //namespace
\ No newline at end of file
diff --git a/Engine/Plugins/Runtime/ModelViewViewModel/Source/ModelViewViewModelBlueprint/Private/ModelViewViewModelBlueprintModule.cpp b/Engine/Plugins/Runtime/ModelViewViewModel/Source/ModelViewViewModelBlueprint/Private/ModelViewViewModelBlueprintModule.cpp
--- a/Engine/Plugins/Runtime/ModelViewViewModel/Source/ModelViewViewModelBlueprint/Private/ModelViewViewModelBlueprintModule.cpp
+++ b/Engine/Plugins/Runtime/ModelViewViewModel/Source/ModelViewViewModelBlueprint/Private/ModelViewViewModelBlueprintModule.cpp
@@ -4,5 +4,33 @@
 #include "Modules/ModuleInterface.h"
 #include "Modules/ModuleManager.h"
 
+#include "KismetCompilerModule.h"
+#include "ViewModel/MVVMViewModelBlueprint.h"
+#include "ViewModel/MVVMViewModelBlueprintCompiler.h"
 
-IMPLEMENT_MODULE(FDefaultModuleImpl, ModelViewViewModelBlueprintModule);
+
+class FModelViewViewModelBlueprintModule : public IModuleInterface
+{
+public:
+	virtual void StartupModule() override
+	{
+		IKismetCompilerInterface& KismetCompilerModule = FModuleManager::GetModuleChecked<IKismetCompilerInterface>("KismetCompiler");
+		KismetCompilerModule.GetCompilers().Add(&ViewModelBlueprintCompiler);
+
+		FKismetCompilerContext::RegisterCompilerForBP(UMVVMViewModelBlueprint::StaticClass(), &UMVVMViewModelBlueprint::GetCompilerForViewModelBlueprint);
+	}
+
+	virtual void ShutdownModule() override
+	{
+		if (IKismetCompilerInterface* KismetCompilerModule = FModuleManager::GetModulePtr<IKismetCompilerInterface>("KismetCompiler"))
+		{
+			KismetCompilerModule->GetCompilers().Remove(&ViewModelBlueprintCompiler);
+		}
+	}
+
+private:
+	UE::MVVM::FViewModelBlueprintCompiler ViewModelBlueprintCompiler;
+};
+
+
+IMPLEMENT_MODULE(FModelViewViewModelBlueprintModule, ModelViewViewModelBlueprintModule);
diff --git a/Engine/Plugins/Runtime/ModelViewViewModel/Source/ModelViewViewModelBlueprint/Private/ViewModel/MVVMViewModelBlueprint.cpp b/Engine/Plugins/Runtime/ModelViewViewModel/Source/ModelViewViewModelBlueprint/Private/ViewModel/MVVMViewModelBlueprint.cpp
new file mode 100644
--- /dev/null
+++ b/Engine/Plugins/Runtime/ModelViewViewModel/Source/ModelViewViewModelBlueprint/Private/ViewModel/MVVMViewModelBlueprint.cpp
@@ -0,0 +1,49 @@
+// Copyright Epic Games, Inc. All Rights Reserved.
+
+#include "ViewModel/MVVMViewModelBlueprint.h"
+#include "ViewModel/MVVMViewModelBlueprintCompiler.h"
+#include "ViewModel/MVVMViewModelBlueprintGeneratedClass.h"
+
+#include "Kismet2/CompilerResultsLog.h"
+
+
+TSharedPtr<FKismetCompilerContext> UMVVMViewModelBlueprint::GetCompilerForViewModelBlueprint(UBlueprint* BP, FCompilerResultsLog& InMessageLog, const FKismetCompilerOptions& InCompileOptions)
+{
+	return TSharedPtr<FKismetCompilerContext>(new UE::MVVM::FViewModelBlueprintCompilerContext(CastChecked<UMVVMViewModelBlueprint>(BP), InMessageLog, InCompileOptions));
+}
+
+
+UMVVMViewModelBlueprintGeneratedClass* UMVVMViewModelBlueprint::GetViewModelBlueprintGeneratedClass() const
+{
+	return Cast<UMVVMViewModelBlueprintGeneratedClass>(*GeneratedClass);
+}
+
+
+UMVVMViewModelBlueprintGeneratedClass* UMVVMViewModelBlueprint::GetViewModelBlueprintSkeletonClass() const
+{
+	return Cast<UMVVMViewModelBlueprintGeneratedClass>(*SkeletonGeneratedClass);
+}
+
+
+bool UMVVMViewModelBlueprint::SupportsMacros() const
+{
+	return false;
+}
+
+
+bool UMVVMViewModelBlueprint::SupportsEventGraphs() const
+{
+	return false;
+}
+
+
+bool UMVVMViewModelBlueprint::SupportsDelegates() const
+{
+	return false;
+}
+
+
+bool UMVVMViewModelBlueprint::SupportsAnimLayers() const
+{
+	return false;
+}
diff --git a/Engine/Plugins/Runtime/ModelViewViewModel/Source/ModelViewViewModelBlueprint/Private/ViewModel/MVVMViewModelBlueprintCompiler.cpp b/Engine/Plugins/Runtime/ModelViewViewModel/Source/ModelViewViewModelBlueprint/Private/ViewModel/MVVMViewModelBlueprintCompiler.cpp
new file mode 100644
--- /dev/null
+++ b/Engine/Plugins/Runtime/ModelViewViewModel/Source/ModelViewViewModelBlueprint/Private/ViewModel/MVVMViewModelBlueprintCompiler.cpp
@@ -0,0 +1,327 @@
+// Copyright Epic Games, Inc. All Rights Reserved.
+
+#include "ViewModel/MVVMViewModelBlueprintCompiler.h"
+#include "ViewModel/MVVMViewModelBlueprint.h"
+#include "ViewModel/MVVMViewModelBlueprintGeneratedClass.h"
+#include "MVVMViewModelBase.h"
+
+#include "FieldNotification/CustomizationHelper.h"
+#include "FieldNotification/FieldNotificationHelpers.h"
+#include "K2Node_FunctionEntry.h"
+#include "K2Node_VariableSet.h"
+#include "Kismet2/CompilerResultsLog.h"
+#include "Kismet2/KismetReinstanceUtilities.h"
+#include "MVVMFunctionGraphHelper.h"
+
+#define LOCTEXT_NAMESPACE "ViewModelBlueprintCompiler"
+
+namespace UE::MVVM
+{
+
+static TAutoConsoleVariable<bool> CVarAutogenerateSetter(
+	TEXT("MVVM.Viewmodel.GenerateSetter"),
+	true,
+	TEXT("If true, the setter function will be always be autogenerated"),
+	ECVF_Default);
+
+static TAutoConsoleVariable<bool> CVarAutogenerateOnRep(
+	TEXT("MVVM.Viewmodel.GenerateOnRep"),
+	true,
+	TEXT("If true, the OnRep function will be always be autogenerated"),
+	ECVF_Default);
+
+//////////////////////////////////////////////////////////////////////////
+// FViewModelBlueprintCompiler 
+
+bool FViewModelBlueprintCompiler::CanCompile(const UBlueprint* Blueprint)
+{
+	return Cast<UMVVMViewModelBlueprint>(Blueprint) != nullptr;
+}
+
+
+void FViewModelBlueprintCompiler::Compile(UBlueprint* Blueprint, const FKismetCompilerOptions& CompileOptions, FCompilerResultsLog& Results)
+{
+	if (UMVVMViewModelBlueprint* WidgetBlueprint = CastChecked<UMVVMViewModelBlueprint>(Blueprint))
+	{
+		FViewModelBlueprintCompilerContext Compiler(WidgetBlueprint, Results, CompileOptions);
+		Compiler.Compile();
+		check(Compiler.NewClass);
+	}
+}
+
+
+bool FViewModelBlueprintCompiler::GetBlueprintTypesForClass(UClass* ParentClass, UClass*& OutBlueprintClass, UClass*& OutBlueprintGeneratedClass) const
+{
+	if (ParentClass == UMVVMViewModelBase::StaticClass() || ParentClass->IsChildOf(UMVVMViewModelBase::StaticClass()))
+	{
+		OutBlueprintClass = UMVVMViewModelBlueprint::StaticClass();
+		OutBlueprintGeneratedClass = UMVVMViewModelBlueprintGeneratedClass::StaticClass();
+		return true;
+	}
+
+	return false;
+}
+
+
+//////////////////////////////////////////////////////////////////////////
+// FViewModelBlueprintCompilerContext
+
+
+UMVVMViewModelBlueprint* FViewModelBlueprintCompilerContext::GetViewModelBlueprint() const
+{
+	return Cast<UMVVMViewModelBlueprint>(Blueprint);
+}
+
+
+FProperty* FViewModelBlueprintCompilerContext::FindPropertyByNameOnNewClass(FName PropertyName) const
+{
+	for (FField* CurrentField = NewClass->ChildProperties; CurrentField != nullptr; CurrentField = CurrentField->Next)
+	{
+		if (CurrentField->GetFName() == PropertyName)
+		{
+			return  CastField<FProperty>(CurrentField);
+			break;
+		}
+	}
+	return nullptr;
+}
+
+
+void FViewModelBlueprintCompilerContext::SaveSubObjectsFromCleanAndSanitizeClass(FSubobjectCollection& SubObjectsToSave, UBlueprintGeneratedClass* ClassToClean)
+{
+	Super::SaveSubObjectsFromCleanAndSanitizeClass(SubObjectsToSave, ClassToClean);
+
+	// Make sure our typed pointer is set
+	check(ClassToClean == NewClass);
+	NewViewModelBlueprintClass = CastChecked<UMVVMViewModelBlueprintGeneratedClass>((UObject*)NewClass);
+}
+
+
+void FViewModelBlueprintCompilerContext::CleanAndSanitizeClass(UBlueprintGeneratedClass* ClassToClean, UObject*& InOutOldCDO)
+{
+	Super::CleanAndSanitizeClass(ClassToClean, InOutOldCDO);
+
+	NewViewModelBlueprintClass->FieldNotifyNames.Empty();
+}
+
+
+void FViewModelBlueprintCompilerContext::CreateFunctionList()
+{
+	for (const FGeneratedFunction& Function : GeneratedFunctions)
+	{
+		if (Function.Property == nullptr || Function.SkelProperty == nullptr)
+		{
+			MessageLog.Error(TEXT("The property for the setter function for '%s' could not be generated."), *Function.PropertyName.ToString());
+			continue;
+		}
+
+		if (Function.SetterFunction)
+		{
+			if (!FunctionGraphHelper::GenerateViewModelFieldNotifySetter(*this, Function.SetterFunction, Function.Property, TEXT("NewValue")))
+			{
+				MessageLog.Error(TEXT("The setter function for '%s' could not be generated."), *Function.PropertyName.ToString());
+				continue;
+			}
+		}
+		if (Function.NetRepFunction)
+		{
+			if (!FunctionGraphHelper::GenerateViewModelFielNotifyBroadcast(*this, Function.NetRepFunction, Function.Property))
+			{
+				MessageLog.Error(TEXT("The net rep function for '%s' could not be generated."), *Function.PropertyName.ToString());
+				continue;
+			}
+		}
+
+		// HACK: Since the VariableSetter doesn't work because of the BlueprintReadOnly, clear it before processing the functions
+		Function.SkelProperty->ClearPropertyFlags(CPF_BlueprintReadOnly);
+	}
+
+	Super::CreateFunctionList();
+
+	// HACK: Reset the BlueprintReadOnly flag
+	for (const FGeneratedFunction& Function : GeneratedFunctions)
+	{
+		if (Function.SkelProperty == nullptr)
+		{
+			continue;
+		}
+		Function.SkelProperty->SetPropertyFlags(CPF_BlueprintReadOnly);
+	}
+}
+
+
+void FViewModelBlueprintCompilerContext::CreateClassVariablesFromBlueprint()
+{
+	Super::CreateClassVariablesFromBlueprint();
+
+	const FName NAME_MetaDataBlueprintSetter = "BlueprintSetter";
+	const FName NAME_Category = "Category";
+	for (TFieldIterator<FProperty> PropertyIt(NewClass, EFieldIterationFlags::None); PropertyIt; ++PropertyIt)
+	{
+		FProperty* Property = *PropertyIt;
+		FName PropertyName = Property->GetFName();
+
+		bool bGenerateSetterFunction = CVarAutogenerateSetter.GetValueOnGameThread() || Property->HasMetaData(UE::FieldNotification::FCustomizationHelper::MetaData_FieldNotify);
+		bool bGenerateReplicatedFunction = CVarAutogenerateOnRep.GetValueOnGameThread(); // Property->RepNotifyFunc.IsNone();
+		if (bGenerateSetterFunction || bGenerateReplicatedFunction)
+		{
+			Property->SetPropertyFlags(CPF_BlueprintReadOnly | CPF_DisableEditOnInstance);
+
+			FGeneratedFunction* FoundGeneratedFunction = GeneratedFunctions.FindByPredicate([PropertyName](const FGeneratedFunction& Other)
+				{
+					return Other.PropertyName == PropertyName;
+				});
+
+
+			if (FoundGeneratedFunction)
+			{
+				if (CompileOptions.CompileType == EKismetCompileType::SkeletonOnly)
+				{
+					FoundGeneratedFunction->SkelProperty = Property;
+					ensureMsgf(false, TEXT("The generated function should already exist"));
+				}
+				else
+				{
+					FoundGeneratedFunction->Property = Property;
+					if (FoundGeneratedFunction->NetRepFunction)
+					{
+						Property->SetPropertyFlags(CPF_Net);
+						Property->RepNotifyFunc = FoundGeneratedFunction->NetRepFunction->GetFName();
+					}
+				}
+			}
+			else if (CompileOptions.CompileType == EKismetCompileType::SkeletonOnly)
+			{
+				FGeneratedFunction GeneratedFunction;
+				GeneratedFunction.PropertyName = PropertyName;
+				GeneratedFunction.SkelProperty = Property;
+
+				if (bGenerateSetterFunction)
+				{
+					const FString SetterName = FString::Printf(TEXT("Set%s"), *Property->GetName());
+					GeneratedFunction.SetterFunction = FunctionGraphHelper::CreateIntermediateFunctionGraph(
+						*this
+						, *SetterName
+						, (FUNC_BlueprintCallable | FUNC_Public)
+						, Property->GetMetaData(NAME_Category)
+						, false);
+
+					if (GeneratedFunction.SetterFunction == nullptr || GeneratedFunction.SetterFunction->GetFName() != FName(*SetterName))
+					{
+						MessageLog.Error(*FString::Printf(TEXT("The setter name %s already exist and could not be autogenerated."), *SetterName));
+					}
+
+					FunctionGraphHelper::AddFunctionArgument(GeneratedFunction.SetterFunction, Property, TEXT("NewValue"));
+				}
+
+				if (bGenerateReplicatedFunction)
+				{
+					const FString NetRepName = FString::Printf(TEXT("__RepNotify_%s"), *Property->GetName());
+					GeneratedFunction.NetRepFunction = FunctionGraphHelper::CreateIntermediateFunctionGraph(
+						*this
+						, *NetRepName
+						, (FUNC_Private)
+						, Property->GetMetaData(NAME_Category)
+						, false);
+
+					if (GeneratedFunction.NetRepFunction == nullptr || GeneratedFunction.NetRepFunction->GetFName() != FName(*NetRepName))
+					{
+						MessageLog.Error(*FString::Printf(TEXT("The rep name %s already exist and could not be autogenerated."), *NetRepName));
+						continue;
+					}
+					else
+					{
+						Property->SetPropertyFlags(CPF_Net);
+						Property->RepNotifyFunc = GeneratedFunction.NetRepFunction->GetFName();
+					}
+
+					FunctionGraphHelper::AddFunctionArgument(GeneratedFunction.SetterFunction, Property, TEXT("NewValue"));
+				}
+
+				GeneratedFunctions.Add(GeneratedFunction);
+			}
+
+			NewViewModelBlueprintClass->FieldNotifyNames.Emplace(Property->GetFName());
+		}
+	}
+}
+
+
+void FViewModelBlueprintCompilerContext::SpawnNewClass(const FString& NewClassName)
+{
+	NewViewModelBlueprintClass = FindObject<UMVVMViewModelBlueprintGeneratedClass>(Blueprint->GetOutermost(), *NewClassName);
+
+	if (NewViewModelBlueprintClass == nullptr)
+	{
+		NewViewModelBlueprintClass = NewObject<UMVVMViewModelBlueprintGeneratedClass>(Blueprint->GetOutermost(), FName(*NewClassName), RF_Public | RF_Transactional);
+	}
+	else
+	{
+		// Already existed, but wasn't linked in the Blueprint yet due to load ordering issues
+		FBlueprintCompileReinstancer::Create(NewViewModelBlueprintClass);
+	}
+	NewClass = NewViewModelBlueprintClass;
+}
+
+
+void FViewModelBlueprintCompilerContext::OnNewClassSet(UBlueprintGeneratedClass* ClassToUse)
+{
+	NewViewModelBlueprintClass = CastChecked<UMVVMViewModelBlueprintGeneratedClass>(ClassToUse);
+}
+
+
+void FViewModelBlueprintCompilerContext::PrecompileFunction(FKismetFunctionContext& Context, EInternalCompilerFlags InternalFlags)
+{
+	// HACK: Since the VariableSet doesn't work because of the BlueprintReadOnly, clear it before processing the functions
+	for (const FGeneratedFunction& Function : GeneratedFunctions)
+	{
+		check(Function.Property);
+		Function.Property->ClearPropertyFlags(CPF_BlueprintReadOnly);
+	}
+
+	Super::PrecompileFunction(Context, InternalFlags);
+}
+
+
+void FViewModelBlueprintCompilerContext::PostcompileFunction(FKismetFunctionContext& Context)
+{
+	Super::PostcompileFunction(Context);
+
+	if (Context.Function && Context.Function->HasMetaData(UE::FieldNotification::FCustomizationHelper::MetaData_FieldNotify))
+	{
+		if (!UE::FieldNotification::Helpers::IsValidAsField(Context.Function))
+		{
+			MessageLog.Error(*LOCTEXT("FieldNotify_IsEventGraph", "Function @@ cannot be a FieldNotify. A function needs to be const, returns a single properties, has no inputs, not be an event or a net function.").ToString(), Context.EntryPoint);
+		}
+		NewViewModelBlueprintClass->FieldNotifyNames.Emplace(Context.Function->GetFName());
+	}
+
+	// HACK: Reset the BlueprintReadOnly flag
+	for (const FGeneratedFunction& Function : GeneratedFunctions)
+	{
+		check(Function.Property);
+		Function.Property->SetPropertyFlags(CPF_BlueprintReadOnly);
+	}
+}
+
+
+void FViewModelBlueprintCompilerContext::FinishCompilingClass(UClass* Class)
+{
+	Super::FinishCompilingClass(Class);
+
+	const bool bIsSkeletonOnly = CompileOptions.CompileType == EKismetCompileType::SkeletonOnly;
+	if (!bIsSkeletonOnly)
+	{
+		if (UMVVMViewModelBlueprintGeneratedClass* BPGClass = CastChecked<UMVVMViewModelBlueprintGeneratedClass>(Class))
+		{
+			if (UMVVMViewModelBase* ViewModelBase = Cast<UMVVMViewModelBase>(BPGClass->GetDefaultObject()))
+			{
+				BPGClass->InitializeFieldNotification(ViewModelBase);
+			}
+		}
+	}
+}
+
+} //namespace
+
+#undef LOCTEXT_NAMESPACE
diff --git a/Engine/Plugins/Runtime/ModelViewViewModel/Source/ModelViewViewModelBlueprint/Private/ViewModel/MVVMViewModelBlueprintCompiler.h b/Engine/Plugins/Runtime/ModelViewViewModel/Source/ModelViewViewModelBlueprint/Private/ViewModel/MVVMViewModelBlueprintCompiler.h
new file mode 100644
--- /dev/null
+++ b/Engine/Plugins/Runtime/ModelViewViewModel/Source/ModelViewViewModelBlueprint/Private/ViewModel/MVVMViewModelBlueprintCompiler.h
@@ -0,0 +1,77 @@
+// Copyright Epic Games, Inc. All Rights Reserved.
+
+#pragma once
+
+#include "CoreMinimal.h"
+#include "EngineLogs.h"
+#include "Engine/Blueprint.h"
+#include "WidgetBlueprint.h"
+#include "KismetCompiler.h"
+#include "KismetCompilerModule.h"
+
+class UEdGraph;
+class UMVVMViewModelBlueprint;
+class UMVVMViewModelBlueprintGeneratedClass;
+
+//////////////////////////////////////////////////////////////////////////
+// FViewModelBlueprintCompiler 
+
+namespace UE::MVVM
+{
+
+class FViewModelBlueprintCompiler : public IBlueprintCompiler
+{
+public:
+	bool CanCompile(const UBlueprint* Blueprint) override;
+	void Compile(UBlueprint* Blueprint, const FKismetCompilerOptions& CompileOptions, FCompilerResultsLog& Results) override;
+	bool GetBlueprintTypesForClass(UClass* ParentClass, UClass*& OutBlueprintClass, UClass*& OutBlueprintGeneratedClass) const override;
+};
+
+} //namespace
+
+
+//////////////////////////////////////////////////////////////////////////
+// FViewModelBlueprintCompilerContext
+
+namespace UE::MVVM
+{
+
+class FViewModelBlueprintCompilerContext : public FKismetCompilerContext
+{
+	typedef FKismetCompilerContext Super;
+
+public:
+	using FKismetCompilerContext::FKismetCompilerContext;
+
+private:
+	UMVVMViewModelBlueprint* GetViewModelBlueprint() const;
+	FProperty* FindPropertyByNameOnNewClass(FName PropertyName) const;
+
+	//~ Begin FKismetCompilerContext Interface
+	virtual void SaveSubObjectsFromCleanAndSanitizeClass(FSubobjectCollection& SubObjectsToSave, UBlueprintGeneratedClass* ClassToClean) override;
+	virtual void CleanAndSanitizeClass(UBlueprintGeneratedClass* ClassToClean, UObject*& InOutOldCDO) override;
+	virtual void CreateFunctionList() override;
+	virtual void CreateClassVariablesFromBlueprint() override;
+	virtual void SpawnNewClass(const FString& NewClassName) override;
+	virtual void OnNewClassSet(UBlueprintGeneratedClass* ClassToUse) override;
+	virtual void PrecompileFunction(FKismetFunctionContext& Context, EInternalCompilerFlags InternalFlags) override;
+	virtual void PostcompileFunction(FKismetFunctionContext& Context) override;
+	virtual void FinishCompilingClass(UClass* Class) override;
+	//~ End FKismetCompilerContext Interface
+
+private:
+	struct FGeneratedFunction
+	{
+		FName PropertyName;
+		UEdGraph* SetterFunction = nullptr;
+		UEdGraph* NetRepFunction = nullptr;
+		FProperty* Property = nullptr;
+		FProperty* SkelProperty = nullptr;
+	};
+
+	UMVVMViewModelBlueprintGeneratedClass* NewViewModelBlueprintClass;
+	TArray<FGeneratedFunction> GeneratedFunctions;
+};
+
+} //namespace
+
diff --git a/Engine/Plugins/Runtime/ModelViewViewModel/Source/ModelViewViewModelBlueprint/Public/ViewModel/MVVMViewModelBlueprint.h b/Engine/Plugins/Runtime/ModelViewViewModel/Source/ModelViewViewModelBlueprint/Public/ViewModel/MVVMViewModelBlueprint.h
new file mode 100644
--- /dev/null
+++ b/Engine/Plugins/Runtime/ModelViewViewModel/Source/ModelViewViewModelBlueprint/Public/ViewModel/MVVMViewModelBlueprint.h
@@ -0,0 +1,38 @@
+// Copyright Epic Games, Inc. All Rights Reserved.
+
+#pragma once
+
+#include "CoreMinimal.h"
+#include "UObject/ObjectMacros.h"
+#include "Engine/Blueprint.h"
+
+#include "MVVMViewModelBlueprint.generated.h"
+
+class FCompilerResultsLog;
+class UMVVMViewModelBlueprintGeneratedClass;
+
+
+/**
+ * 
+ */
+UCLASS(BlueprintType, meta=(IgnoreClassThumbnail))
+class MODELVIEWVIEWMODELBLUEPRINT_API UMVVMViewModelBlueprint : public UBlueprint
+{
+	GENERATED_BODY()
+
+public:
+	/** */
+	static TSharedPtr<FKismetCompilerContext> GetCompilerForViewModelBlueprint(UBlueprint* BP, FCompilerResultsLog& InMessageLog, const FKismetCompilerOptions& InCompileOptions);
+
+	/** Get the (full) generated class for this viewmodel blueprint */
+	UMVVMViewModelBlueprintGeneratedClass* GetViewModelBlueprintGeneratedClass() const;
+
+	/** Get the (skeleton) generated class for this viewmodel blueprint */
+	UMVVMViewModelBlueprintGeneratedClass* GetViewModelBlueprintSkeletonClass() const;
+
+protected:
+	virtual bool SupportsMacros() const override;
+	virtual bool SupportsEventGraphs() const override;
+	virtual bool SupportsDelegates() const override;
+	virtual bool SupportsAnimLayers() const override;
+};
diff --git a/Engine/Plugins/Runtime/ModelViewViewModel/Source/ModelViewViewModelEditor/ModelViewViewModelEditor.Build.cs b/Engine/Plugins/Runtime/ModelViewViewModel/Source/ModelViewViewModelEditor/ModelViewViewModelEditor.Build.cs
--- a/Engine/Plugins/Runtime/ModelViewViewModel/Source/ModelViewViewModelEditor/ModelViewViewModelEditor.Build.cs
+++ b/Engine/Plugins/Runtime/ModelViewViewModel/Source/ModelViewViewModelEditor/ModelViewViewModelEditor.Build.cs
@@ -23,6 +23,7 @@ public class ModelViewViewModelEditor : ModuleRules
 		PrivateDependencyModuleNames.AddRange(
 			new string[]
 			{
+				"AssetTools",
 				"EditorStyle",
 				"EditorWidgets",
 				"BlueprintGraph",
diff --git a/Engine/Plugins/Runtime/ModelViewViewModel/Source/ModelViewViewModelEditor/Private/ModelViewViewModelEditorModule.cpp b/Engine/Plugins/Runtime/ModelViewViewModel/Source/ModelViewViewModelEditor/Private/ModelViewViewModelEditorModule.cpp
--- a/Engine/Plugins/Runtime/ModelViewViewModel/Source/ModelViewViewModelEditor/Private/ModelViewViewModelEditorModule.cpp
+++ b/Engine/Plugins/Runtime/ModelViewViewModel/Source/ModelViewViewModelEditor/Private/ModelViewViewModelEditorModule.cpp
@@ -2,6 +2,7 @@
 
 #include "ModelViewViewModelEditorModule.h"
 
+#include "AssetToolsModule.h"
 #include "BlueprintModes/WidgetBlueprintApplicationMode.h"
 #include "BlueprintModes/WidgetBlueprintApplicationModes.h"
 #include "Customizations/MVVMBindPropertiesDetailView.h"
@@ -25,6 +26,7 @@
 #include "Styling/StyleColors.h"
 #include "Tabs/MVVMBindingSummoner.h"
 #include "UMGEditorModule.h"
+#include "ViewModel/AssetTypeActions_ViewModelBlueprint.h"
 #include "WidgetBlueprint.h"
 #include "WidgetBlueprintExtension.h"
 #include "WidgetBlueprintToolMenuContext.h"
@@ -46,6 +48,13 @@ void FModelViewViewModelEditorModule::StartupModule()
 	IUMGEditorModule& UMGEditorModule = FModuleManager::LoadModuleChecked<IUMGEditorModule>("UMGEditor");
 	UMGEditorModule.OnRegisterTabsForEditor().AddRaw(this, &FModelViewViewModelEditorModule::HandleRegisterBlueprintEditorTab);
 
+	// Register asset types
+	{
+		IAssetTools& AssetTools = FModuleManager::LoadModuleChecked<FAssetToolsModule>("AssetTools").Get();
+		ViewModelBlueprintActions = MakeShared<UE::MVVM::FAssetTypeActions_ViewModelBlueprint>();
+		AssetTools.RegisterAssetTypeActions(ViewModelBlueprintActions.ToSharedRef());
+	}
+
 	PropertyBindingExtension = MakeShared<FMVVMPropertyBindingExtension>();
 	UMGEditorModule.GetPropertyBindingExtensibilityManager()->AddExtension(PropertyBindingExtension.ToSharedRef());
 
@@ -83,6 +92,12 @@ void FModelViewViewModelEditorModule::ShutdownModule()
 	}
 	PropertyBindingExtension.Reset();
 
+	// Unregister all the asset types that we registered
+	if (FAssetToolsModule* AssetTools = FModuleManager::GetModulePtr<FAssetToolsModule>("AssetTools"))
+	{
+		AssetTools->Get().UnregisterAssetTypeActions(ViewModelBlueprintActions.ToSharedRef());
+	}
+
 	if (FPropertyEditorModule* PropertyEditorModule = FModuleManager::GetModulePtr<FPropertyEditorModule>("PropertyEditor"))
 	{
 		PropertyEditorModule->GetGlobalRowExtensionDelegate().RemoveAll(BindDetailView.Get());
diff --git a/Engine/Plugins/Runtime/ModelViewViewModel/Source/ModelViewViewModelEditor/Private/ModelViewViewModelEditorModule.h b/Engine/Plugins/Runtime/ModelViewViewModel/Source/ModelViewViewModelEditor/Private/ModelViewViewModelEditorModule.h
--- a/Engine/Plugins/Runtime/ModelViewViewModel/Source/ModelViewViewModelEditor/Private/ModelViewViewModelEditorModule.h
+++ b/Engine/Plugins/Runtime/ModelViewViewModel/Source/ModelViewViewModelEditor/Private/ModelViewViewModelEditorModule.h
@@ -6,8 +6,13 @@
 #include "Modules/ModuleInterface.h"
 #include "Modules/ModuleManager.h"
 
-class FMVVMPropertyBindingExtension;
+namespace UE::MVVM
+{
+	class FAssetTypeActions_ViewModelBlueprint;
+}
+
 class FMVVMBindPropertiesDetailView;
+class FMVVMPropertyBindingExtension;
 class FWidgetBlueprintApplicationMode;
 class FWorkflowAllowedTabSet;
 class UBlueprint;
@@ -34,4 +39,5 @@ private:
 private:
 	TUniquePtr<FMVVMBindPropertiesDetailView> BindDetailView;
 	TSharedPtr<FMVVMPropertyBindingExtension> PropertyBindingExtension;
+	TSharedPtr<UE::MVVM::FAssetTypeActions_ViewModelBlueprint> ViewModelBlueprintActions;
 };
diff --git a/Engine/Plugins/Runtime/ModelViewViewModel/Source/ModelViewViewModelEditor/Private/ViewModel/AssetTypeActions_ViewModelBlueprint.cpp b/Engine/Plugins/Runtime/ModelViewViewModel/Source/ModelViewViewModelEditor/Private/ViewModel/AssetTypeActions_ViewModelBlueprint.cpp
new file mode 100644
--- /dev/null
+++ b/Engine/Plugins/Runtime/ModelViewViewModel/Source/ModelViewViewModelEditor/Private/ViewModel/AssetTypeActions_ViewModelBlueprint.cpp
@@ -0,0 +1,83 @@
+// Copyright Epic Games, Inc. All Rights Reserved.
+
+#include "AssetTypeActions_ViewModelBlueprint.h"
+#include "Misc/MessageDialog.h"
+#include "Engine/Blueprint.h"
+#include "Kismet2/BlueprintEditorUtils.h"
+#include "SBlueprintDiff.h"
+
+#include "ViewModel/MVVMViewModelBlueprint.h"
+#include "ViewModel/MVVMViewModelBlueprintEditor.h"
+
+#define LOCTEXT_NAMESPACE "AssetTypeActions"
+
+namespace UE::MVVM
+{
+
+void FAssetTypeActions_ViewModelBlueprint::OpenAssetEditor( const TArray<UObject*>& InObjects, TSharedPtr<IToolkitHost> EditWithinLevelEditor )
+{
+	EToolkitMode::Type Mode = EditWithinLevelEditor.IsValid() ? EToolkitMode::WorldCentric : EToolkitMode::Standalone;
+
+	for (auto ObjIt = InObjects.CreateConstIterator(); ObjIt; ++ObjIt)
+	{
+		UBlueprint* Blueprint = Cast<UBlueprint>(*ObjIt);
+		if (Blueprint && Blueprint->SkeletonGeneratedClass && Blueprint->GeneratedClass )
+		{
+			TSharedRef< FMVVMViewModelBlueprintEditor > NewBlueprintEditor(new FMVVMViewModelBlueprintEditor());
+
+			bool IsDataOnlyBlueprint = FBlueprintEditorUtils::IsDataOnlyBlueprint(Blueprint)
+				&& !Blueprint->bForceFullEditor
+				&& !Blueprint->bIsNewlyCreated;
+			TArray<UBlueprint*> Blueprints;
+			Blueprints.Add(Blueprint);
+			NewBlueprintEditor->InitViewModelBlueprintEditor(Mode, EditWithinLevelEditor, Blueprints, IsDataOnlyBlueprint);
+		}
+		else
+		{
+			FMessageDialog::Open( EAppMsgType::Ok, LOCTEXT("FailedToLoadViewModelBlueprint", "Viewmodel Blueprint could not be loaded because it derives from an invalid class.\nCheck to make sure the parent class for this blueprint hasn't been removed!"));
+		}
+	}
+}
+
+
+UClass* FAssetTypeActions_ViewModelBlueprint::GetSupportedClass() const
+{
+	return UMVVMViewModelBlueprint::StaticClass();
+}
+
+
+FText FAssetTypeActions_ViewModelBlueprint::GetAssetDescription( const FAssetData& AssetData ) const
+{
+	FString Description = AssetData.GetTagValueRef<FString>( GET_MEMBER_NAME_CHECKED( UBlueprint, BlueprintDescription ) );
+	if ( !Description.IsEmpty() )
+	{
+		Description.ReplaceInline( TEXT( "\\n" ), TEXT( "\n" ) );
+		return FText::FromString( MoveTemp(Description) );
+	}
+
+	return FText::GetEmpty();
+}
+
+
+void FAssetTypeActions_ViewModelBlueprint::PerformAssetDiff(UObject* Asset1, UObject* Asset2, const struct FRevisionInfo& OldRevision, const struct FRevisionInfo& NewRevision) const
+{
+	UBlueprint* OldBlueprint = CastChecked<UBlueprint>(Asset1);
+	UBlueprint* NewBlueprint = CastChecked<UBlueprint>(Asset2);
+
+	// sometimes we're comparing different revisions of one single asset (other 
+	//times we're comparing two completely separate assets altogether)
+	bool bIsSingleAsset = (NewBlueprint->GetName() == OldBlueprint->GetName());
+
+	FText WindowTitle = LOCTEXT("NamelessViewmodelBlueprintDiff", "Viewmodel Blueprint Diff");
+	if (bIsSingleAsset)
+	{
+		// identify the assumed single asset in the window's title
+		WindowTitle = FText::Format(LOCTEXT("WidgetBlueprintDiff", "{0} - Viewmodel Blueprint Diff"), FText::FromString(NewBlueprint->GetName()));
+	}
+
+	SBlueprintDiff::CreateDiffWindow(WindowTitle, OldBlueprint, NewBlueprint, OldRevision, NewRevision);
+}
+
+} //namespace
+
+#undef LOCTEXT_NAMESPACE
diff --git a/Engine/Plugins/Runtime/ModelViewViewModel/Source/ModelViewViewModelEditor/Private/ViewModel/AssetTypeActions_ViewModelBlueprint.h b/Engine/Plugins/Runtime/ModelViewViewModel/Source/ModelViewViewModelEditor/Private/ViewModel/AssetTypeActions_ViewModelBlueprint.h
new file mode 100644
--- /dev/null
+++ b/Engine/Plugins/Runtime/ModelViewViewModel/Source/ModelViewViewModelEditor/Private/ViewModel/AssetTypeActions_ViewModelBlueprint.h
@@ -0,0 +1,29 @@
+// Copyright Epic Games, Inc. All Rights Reserved.
+
+#pragma once
+
+#include "CoreMinimal.h"
+#include "Toolkits/IToolkitHost.h"
+#include "AssetTypeActions_Base.h"
+
+struct FAssetData;
+
+namespace UE::MVVM
+{
+	
+class FAssetTypeActions_ViewModelBlueprint : public FAssetTypeActions_Base
+{
+public:
+	// IAssetTypeActions Implementation
+	virtual FText GetName() const override { return NSLOCTEXT("AssetTypeActions", "AssetTypeActions_ViewModelBlueprint", "Viewmodel Blueprint"); }
+	virtual FColor GetTypeColor() const override { return FColor(44, 180, 180); }
+	virtual UClass* GetSupportedClass() const override;
+	virtual void OpenAssetEditor(const TArray<UObject*>& InObjects, TSharedPtr<class IToolkitHost> EditWithinLevelEditor = TSharedPtr<IToolkitHost>()) override;
+	virtual uint32 GetCategories() override { return EAssetTypeCategories::UI; }
+	virtual bool CanLocalize() const override { return false; }
+	virtual void PerformAssetDiff(UObject* Asset1, UObject* Asset2, const struct FRevisionInfo& OldRevision, const struct FRevisionInfo& NewRevision) const override;
+	virtual FText GetAssetDescription( const FAssetData& AssetData ) const override;
+	// End IAssetTypeActions Implementation
+};
+
+} //namespace
diff --git a/Engine/Plugins/Runtime/ModelViewViewModel/Source/ModelViewViewModelEditor/Private/ViewModel/MVVMViewModelBlueprintEditor.cpp b/Engine/Plugins/Runtime/ModelViewViewModel/Source/ModelViewViewModelEditor/Private/ViewModel/MVVMViewModelBlueprintEditor.cpp
new file mode 100644
--- /dev/null
+++ b/Engine/Plugins/Runtime/ModelViewViewModel/Source/ModelViewViewModelEditor/Private/ViewModel/MVVMViewModelBlueprintEditor.cpp
@@ -0,0 +1,102 @@
+// Copyright Epic Games, Inc. All Rights Reserved.
+
+#include "ViewModel/MVVMViewModelBlueprintEditor.h"
+#include "ViewModel/MVVMViewModelBlueprint.h"
+#include "ViewModel/MVVMViewModelBlueprintToolMenuContext.h"
+
+#define LOCTEXT_NAMESPACE "MVVMViewModelBlueprintEditor"
+
+void FMVVMViewModelBlueprintEditor::InitViewModelBlueprintEditor(const EToolkitMode::Type Mode, const TSharedPtr< IToolkitHost >& InitToolkitHost, const TArray<UBlueprint*>& InBlueprints, bool bShouldOpenInDefaultsMode)
+{
+	Super::InitBlueprintEditor(Mode, InitToolkitHost, InBlueprints, bShouldOpenInDefaultsMode);
+}
+
+
+UMVVMViewModelBlueprint* FMVVMViewModelBlueprintEditor::GetViewModelBlueprintObj() const
+{
+	return Cast<UMVVMViewModelBlueprint>(GetBlueprintObj());
+}
+
+
+FName FMVVMViewModelBlueprintEditor::GetToolkitFName() const
+{
+	return FName("ViewModelBlueprintEditor");
+}
+
+
+FText FMVVMViewModelBlueprintEditor::GetBaseToolkitName() const
+{
+	return LOCTEXT("AppLabel", "Viewmodel Editor");
+}
+
+
+FString FMVVMViewModelBlueprintEditor::GetWorldCentricTabPrefix() const
+{
+	return LOCTEXT("WorldCentricTabPrefix", "Viewmodel Editor ").ToString();
+}
+
+
+FLinearColor FMVVMViewModelBlueprintEditor::GetWorldCentricTabColorScale() const
+{
+	return FLinearColor(0.3f, 0.25f, 0.35f, 0.5f);
+}
+
+
+void FMVVMViewModelBlueprintEditor::InitToolMenuContext(FToolMenuContext& MenuContext)
+{
+	Super::InitToolMenuContext(MenuContext);
+
+	UMVVMViewModelBlueprintToolMenuContext* Context = NewObject<UMVVMViewModelBlueprintToolMenuContext>();
+	Context->ViewModelBlueprintEditor = SharedThis(this);
+	MenuContext.AddObject(Context);
+}
+
+
+void FMVVMViewModelBlueprintEditor::RegisterApplicationModes(const TArray<UBlueprint*>& InBlueprints, bool bShouldOpenInDefaultsMode, bool bNewlyCreated/* = false*/)
+{
+	Super::RegisterApplicationModes(InBlueprints, bShouldOpenInDefaultsMode);
+}
+
+
+FGraphAppearanceInfo FMVVMViewModelBlueprintEditor::GetGraphAppearance(UEdGraph* InGraph) const
+{
+	FGraphAppearanceInfo AppearanceInfo = Super::GetGraphAppearance(InGraph);
+
+	if (GetBlueprintObj()->IsA(UMVVMViewModelBlueprint::StaticClass()))
+	{
+		AppearanceInfo.CornerText = LOCTEXT("AppearanceCornerText", "VIEWMODEL BLUEPRINT");
+	}
+
+	return AppearanceInfo;
+}
+
+
+bool FMVVMViewModelBlueprintEditor::NewDocument_IsVisibleForType(ECreatedDocumentType InGraphType) const
+{
+	switch (InGraphType)
+	{
+	case CGT_NewVariable:
+	case CGT_NewFunctionGraph:
+		return true;
+	default:
+		break;
+	}
+	return false;
+}
+
+
+bool FMVVMViewModelBlueprintEditor::IsSectionVisible(NodeSectionID::Type InSectionID) const
+{
+	switch (InSectionID)
+	{
+	case NodeSectionID::VARIABLE:
+	case NodeSectionID::FUNCTION:
+	case NodeSectionID::LOCAL_VARIABLE:
+		return true;
+	default:
+		break;
+	}
+	return false;
+}
+
+#undef LOCTEXT_NAMESPACE
diff --git a/Engine/Plugins/Runtime/ModelViewViewModel/Source/ModelViewViewModelEditor/Private/ViewModel/MVVMViewModelBlueprintFactory.cpp b/Engine/Plugins/Runtime/ModelViewViewModel/Source/ModelViewViewModelEditor/Private/ViewModel/MVVMViewModelBlueprintFactory.cpp
new file mode 100644
--- /dev/null
+++ b/Engine/Plugins/Runtime/ModelViewViewModel/Source/ModelViewViewModelEditor/Private/ViewModel/MVVMViewModelBlueprintFactory.cpp
@@ -0,0 +1,119 @@
+// Copyright Epic Games, Inc. All Rights Reserved.
+
+#include "ViewModel/MVVMViewModelBlueprintFactory.h"
+#include "ViewModel/MVVMViewModelBlueprint.h"
+#include "ViewModel/MVVMViewModelBlueprintGeneratedClass.h"
+
+#include "Blueprint/UserWidget.h"
+#include "ClassViewerModule.h"
+#include "ClassViewerFilter.h"
+#include "Kismet2/KismetEditorUtilities.h"
+#include "Kismet2/SClassPickerDialog.h"
+#include "Misc/MessageDialog.h"
+#include "Modules/ModuleManager.h"
+
+#define LOCTEXT_NAMESPACE "ViewModelBlueprintFactory"
+
+/*------------------------------------------------------------------------------
+	UMVVMViewModelBlueprintFactory implementation.
+------------------------------------------------------------------------------*/
+
+namespace UE::MVVM::Private
+{
+
+class FViewModelClassFilter : public IClassViewerFilter
+{
+public:
+	/** All children of these classes will be included unless filtered out by another setting. */
+	TSet <const UClass*> AllowedChildrenOfClasses;
+
+	/** Disallowed class flags. */
+	EClassFlags DisallowedClassFlags;
+
+	virtual bool IsClassAllowed(const FClassViewerInitializationOptions& InInitOptions, const UClass* InClass, TSharedRef< FClassViewerFilterFuncs > InFilterFuncs) override
+	{
+		return !InClass->HasAnyClassFlags(DisallowedClassFlags)
+			&& InFilterFuncs->IfInChildOfClassesSet(AllowedChildrenOfClasses, InClass) != EFilterReturn::Failed;
+	}
+
+	virtual bool IsUnloadedClassAllowed(const FClassViewerInitializationOptions& InInitOptions, const TSharedRef< const IUnloadedBlueprintData > InUnloadedClassData, TSharedRef< FClassViewerFilterFuncs > InFilterFuncs) override
+	{
+		return !InUnloadedClassData->HasAnyClassFlags(DisallowedClassFlags)
+			&& InFilterFuncs->IfInChildOfClassesSet(AllowedChildrenOfClasses, InUnloadedClassData) != EFilterReturn::Failed;
+	}
+};
+
+}//namespace
+
+UMVVMViewModelBlueprintFactory::UMVVMViewModelBlueprintFactory()
+{
+	bCreateNew = true;
+	bEditAfterNew = true;
+	SupportedClass = UMVVMViewModelBlueprint::StaticClass();
+	ParentClass = UMVVMViewModelBase::StaticClass();
+}
+
+bool UMVVMViewModelBlueprintFactory::ConfigureProperties()
+{
+	FClassViewerModule& ClassViewerModule = FModuleManager::LoadModuleChecked<FClassViewerModule>("ClassViewer");
+
+	// Fill in options
+	FClassViewerInitializationOptions Options;
+	Options.DisplayMode = EClassViewerDisplayMode::Type::TreeView;
+	Options.Mode = EClassViewerMode::ClassPicker;
+	Options.bShowNoneOption = false;
+	Options.bExpandAllNodes = true;
+
+	TSharedRef<UE::MVVM::Private::FViewModelClassFilter> Filter = MakeShared<UE::MVVM::Private::FViewModelClassFilter>();
+	Filter->DisallowedClassFlags = CLASS_Deprecated | CLASS_NewerVersionExists;
+	Filter->AllowedChildrenOfClasses.Add(UMVVMViewModelBase::StaticClass());
+
+	Options.ClassFilters.Add(Filter);
+	Options.ExtraPickerCommonClasses.Add(UMVVMViewModelBase::StaticClass());
+
+	const FText TitleText = LOCTEXT("CreateViewmodelBlueprint", "Pick Root Viewmodel for New ViewModel Blueprint");
+	UClass* ChosenParentClass = nullptr;
+	if (SClassPickerDialog::PickClass(TitleText, Options, ChosenParentClass, UMVVMViewModelBase::StaticClass()))
+	{
+		ParentClass = ChosenParentClass ? ChosenParentClass : UMVVMViewModelBase::StaticClass();
+	}
+
+	return true;
+}
+
+bool UMVVMViewModelBlueprintFactory::ShouldShowInNewMenu() const
+{
+	return true;
+}
+
+UObject* UMVVMViewModelBlueprintFactory::FactoryCreateNew(UClass* Class, UObject* InParent, FName Name, EObjectFlags Flags, UObject* Context, FFeedbackContext* Warn, FName CallingContext)
+{
+	// Make sure we are trying to factory a Anim Blueprint, then create and init one
+	check(Class->IsChildOf(UMVVMViewModelBlueprint::StaticClass()));
+
+	UClass* CurrentParentClass = ParentClass;
+	if (CurrentParentClass == nullptr)
+	{
+		CurrentParentClass = UMVVMViewModelBase::StaticClass();
+	}
+
+	if ((CurrentParentClass == nullptr) || !FKismetEditorUtilities::CanCreateBlueprintOfClass(CurrentParentClass) || !CurrentParentClass->IsChildOf(UMVVMViewModelBase::StaticClass()))
+	{
+		FFormatNamedArguments Args;
+		Args.Add(TEXT("ClassName"), CurrentParentClass ? FText::FromString(CurrentParentClass->GetName()) : LOCTEXT("Null", "(null)"));
+		FMessageDialog::Open(EAppMsgType::Ok, FText::Format(LOCTEXT("CannotCreateViewmodelBlueprint", "Cannot create a Viewmodel Blueprint based on the class '{ClassName}'."), Args));
+		return nullptr;
+	}
+	else
+	{
+
+		return CastChecked<UMVVMViewModelBlueprint>(FKismetEditorUtilities::CreateBlueprint(CurrentParentClass, InParent, Name, EBlueprintType::BPTYPE_Normal, UMVVMViewModelBlueprint::StaticClass(), UMVVMViewModelBlueprintGeneratedClass::StaticClass(), CallingContext));
+	}
+}
+
+UObject* UMVVMViewModelBlueprintFactory::FactoryCreateNew(UClass* Class, UObject* InParent, FName Name, EObjectFlags Flags, UObject* Context, FFeedbackContext* Warn)
+{
+	return FactoryCreateNew(Class, InParent, Name, Flags, Context, Warn, NAME_None);
+}
+
+#undef LOCTEXT_NAMESPACE
diff --git a/Engine/Plugins/Runtime/ModelViewViewModel/Source/ModelViewViewModelEditor/Private/ViewModel/MVVMViewModelBlueprintFactory.h b/Engine/Plugins/Runtime/ModelViewViewModel/Source/ModelViewViewModelEditor/Private/ViewModel/MVVMViewModelBlueprintFactory.h
new file mode 100644
--- /dev/null
+++ b/Engine/Plugins/Runtime/ModelViewViewModel/Source/ModelViewViewModelEditor/Private/ViewModel/MVVMViewModelBlueprintFactory.h
@@ -0,0 +1,33 @@
+// Copyright Epic Games, Inc. All Rights Reserved.
+
+#pragma once
+
+#include "CoreMinimal.h"
+#include "UObject/ObjectMacros.h"
+#include "Templates/SubclassOf.h"
+#include "Factories/Factory.h"
+#include "Engine/Blueprint.h"
+#include "MVVMViewModelBase.h"
+
+#include "MVVMViewModelBlueprintFactory.generated.h"
+
+UCLASS(HideCategories=Object, MinimalAPI)
+class UMVVMViewModelBlueprintFactory : public UFactory
+{
+	GENERATED_BODY()
+
+public:
+	UMVVMViewModelBlueprintFactory();
+
+	//~ Begin UFactory Interface
+	virtual bool ConfigureProperties() override;
+	virtual bool ShouldShowInNewMenu() const override;
+	virtual UObject* FactoryCreateNew(UClass* Class, UObject* InParent, FName Name, EObjectFlags Flags, UObject* Context, FFeedbackContext* Warn, FName CallingContext) override;
+	virtual UObject* FactoryCreateNew(UClass* Class, UObject* InParent, FName Name, EObjectFlags Flags, UObject* Context, FFeedbackContext* Warn) override;
+	//~ End UFactory Interface	
+
+private:
+	// The parent class of the created blueprint
+	UPROPERTY(EditAnywhere, Category="MVVM", meta=(AllowAbstract="true"))
+	TSubclassOf<class UMVVMViewModelBase> ParentClass;
+};
diff --git a/Engine/Plugins/Runtime/ModelViewViewModel/Source/ModelViewViewModelEditor/Public/ViewModel/MVVMViewModelBlueprintEditor.h b/Engine/Plugins/Runtime/ModelViewViewModel/Source/ModelViewViewModelEditor/Public/ViewModel/MVVMViewModelBlueprintEditor.h
new file mode 100644
--- /dev/null
+++ b/Engine/Plugins/Runtime/ModelViewViewModel/Source/ModelViewViewModelEditor/Public/ViewModel/MVVMViewModelBlueprintEditor.h
@@ -0,0 +1,50 @@
+// Copyright Epic Games, Inc. All Rights Reserved.
+
+#pragma once
+
+#include "CoreMinimal.h"
+#include "BlueprintEditor.h"
+
+class UMVVMViewModelBlueprint;
+
+/**
+ * Viewmodel blueprint editor (extends Blueprint editor)
+ */
+class MODELVIEWVIEWMODELEDITOR_API FMVVMViewModelBlueprintEditor : public FBlueprintEditor
+{
+private:
+	using Super = FBlueprintEditor;
+
+
+public:
+	FMVVMViewModelBlueprintEditor() = default;
+	//virtual ~FMVVMViewModelBlueprintEditor();
+
+	void InitViewModelBlueprintEditor(const EToolkitMode::Type Mode, const TSharedPtr< IToolkitHost >& InitToolkitHost, const TArray<UBlueprint*>& InBlueprints, bool bShouldOpenInDefaultsMode);
+
+	//~ Begin FBlueprintEditor interface
+	//virtual void Tick(float DeltaTime) override;
+	//virtual void PostUndo(bool bSuccessful) override;
+	//virtual void PostRedo(bool bSuccessful) override;
+	//virtual void Compile() override;
+	//~ End FBlueprintEditor interface
+
+	// ~ Begin IToolkit interface
+	virtual FName GetToolkitFName() const override;
+	virtual FText GetBaseToolkitName() const override;
+	virtual FString GetWorldCentricTabPrefix() const override;
+	virtual FLinearColor GetWorldCentricTabColorScale() const override;
+	virtual void InitToolMenuContext(FToolMenuContext& MenuContext) override;
+	// ~ End IToolkit interface
+
+	/** @return The viewmodel blueprint currently being edited in this editor */
+	UMVVMViewModelBlueprint* GetViewModelBlueprintObj() const;
+
+protected:
+	//~ Begin FBlueprintEditor interface
+	virtual void RegisterApplicationModes(const TArray<UBlueprint*>& Blueprints, bool bShouldOpenInDefaultsMode, bool bNewlyCreated = false) override;
+	virtual FGraphAppearanceInfo GetGraphAppearance(class UEdGraph* Graph) const override;
+	virtual bool NewDocument_IsVisibleForType(ECreatedDocumentType GraphType) const override;
+	virtual bool IsSectionVisible(NodeSectionID::Type SectionID) const override;
+	//~ End FBlueprintEditor interface
+};
diff --git a/Engine/Plugins/Runtime/ModelViewViewModel/Source/ModelViewViewModelEditor/Public/ViewModel/MVVMViewModelBlueprintToolMenuContext.h b/Engine/Plugins/Runtime/ModelViewViewModel/Source/ModelViewViewModelEditor/Public/ViewModel/MVVMViewModelBlueprintToolMenuContext.h
new file mode 100644
--- /dev/null
+++ b/Engine/Plugins/Runtime/ModelViewViewModel/Source/ModelViewViewModelEditor/Public/ViewModel/MVVMViewModelBlueprintToolMenuContext.h
@@ -0,0 +1,18 @@
+// Copyright Epic Games, Inc. All Rights Reserved.
+
+#pragma once
+
+#include "CoreMinimal.h"
+
+#include "MVVMViewModelBlueprintToolMenuContext.generated.h"
+
+class FMVVMViewModelBlueprintEditor;
+
+UCLASS()
+class MODELVIEWVIEWMODELEDITOR_API UMVVMViewModelBlueprintToolMenuContext : public UObject
+{
+	GENERATED_BODY()
+
+public:
+	TWeakPtr<FMVVMViewModelBlueprintEditor> ViewModelBlueprintEditor;
+};
