MVVM: Hide private and protected function from the MVVM editor.
#rnx

diff --git a/Engine/Plugins/Runtime/ModelViewViewModel/Source/ModelViewViewModel/Private/MVVMSubsystem.cpp b/Engine/Plugins/Runtime/ModelViewViewModel/Source/ModelViewViewModel/Private/MVVMSubsystem.cpp
--- a/Engine/Plugins/Runtime/ModelViewViewModel/Source/ModelViewViewModel/Private/MVVMSubsystem.cpp
+++ b/Engine/Plugins/Runtime/ModelViewViewModel/Source/ModelViewViewModel/Private/MVVMSubsystem.cpp
@@ -132,11 +132,14 @@ namespace UE::MVVM::Private
 		{
 			const UFunction* Function = *FunctionItt;
 			check(Function);
-			const bool bIsReadable = FieldDescriptors.Contains(Function->GetFName()) && BindingHelper::IsValidForSourceBinding(Function);
-			const bool bIsWritable = BindingHelper::IsValidForDestinationBinding(Function);
-			if (bIsReadable || bIsWritable)
+			if (!Function->HasAnyFunctionFlags(FUNC_Private | FUNC_Protected))
 			{
-				Result.Add(FMVVMAvailableBinding(FMVVMBindingName(Function->GetFName()), bIsReadable, bIsWritable));
+				const bool bIsReadable = FieldDescriptors.Contains(Function->GetFName()) && BindingHelper::IsValidForSourceBinding(Function);
+				const bool bIsWritable = BindingHelper::IsValidForDestinationBinding(Function);
+				if (bIsReadable || bIsWritable)
+				{
+					Result.Add(FMVVMAvailableBinding(FMVVMBindingName(Function->GetFName()), bIsReadable, bIsWritable));
+				}
 			}
 		}
 
