Fixed duplicated condition detected by static analysis in UE::MVVM::FunctionGraphHelper::GenerateViewModelFieldNotifySetter().

diff --git a/Engine/Plugins/Runtime/ModelViewViewModel/Source/ModelViewViewModelBlueprint/Private/MVVMFunctionGraphHelper.cpp b/Engine/Plugins/Runtime/ModelViewViewModel/Source/ModelViewViewModelBlueprint/Private/MVVMFunctionGraphHelper.cpp
--- a/Engine/Plugins/Runtime/ModelViewViewModel/Source/ModelViewViewModelBlueprint/Private/MVVMFunctionGraphHelper.cpp
+++ b/Engine/Plugins/Runtime/ModelViewViewModel/Source/ModelViewViewModelBlueprint/Private/MVVMFunctionGraphHelper.cpp
@@ -328,7 +328,7 @@ bool GenerateViewModelFieldNotifySetter(FKismetCompilerContext& InContext, UEdGr
 			{
 				ensure(Schema->TryCreateConnection(EntryValuePin, TopInputPin));
 			}
-			bResult = bResult && EntryValuePin && EntryValuePin;
+			bResult = bResult && EntryValuePin && TopInputPin;
 
 			UEdGraphPin* GetPin = GetVariableNode->FindPin(InProperty->GetFName(), EGPD_Output);
 			UEdGraphPin* BottomInputPin = CompareOperator->FindPin(TEXT("B"), EGPD_Input);
