Fixed missing icons in MVVMViewBindingListView.

TOptional replaced with plain FBindingSource, which is nullable.

diff --git a/Engine/Plugins/Runtime/ModelViewViewModel/Source/ModelViewViewModelEditor/Private/Customizations/MVVMPropertyPathCustomization.cpp b/Engine/Plugins/Runtime/ModelViewViewModel/Source/ModelViewViewModelEditor/Private/Customizations/MVVMPropertyPathCustomization.cpp
--- a/Engine/Plugins/Runtime/ModelViewViewModel/Source/ModelViewViewModelEditor/Private/Customizations/MVVMPropertyPathCustomization.cpp
+++ b/Engine/Plugins/Runtime/ModelViewViewModel/Source/ModelViewViewModelEditor/Private/Customizations/MVVMPropertyPathCustomization.cpp
@@ -176,7 +176,7 @@ namespace UE::MVVM
 		return EMVVMBindingMode::OneWayToDestination;
 	}
 
-	void FPropertyPathCustomizationBase::OnSourceSelectionChanged(TOptional<FBindingSource> Selected)
+	void FPropertyPathCustomizationBase::OnSourceSelectionChanged(FBindingSource Selected)
 	{
 		FScopedTransaction Transaction(LOCTEXT("ChangeSource", "Change binding source."));
 
diff --git a/Engine/Plugins/Runtime/ModelViewViewModel/Source/ModelViewViewModelEditor/Private/Customizations/MVVMPropertyPathCustomization.h b/Engine/Plugins/Runtime/ModelViewViewModel/Source/ModelViewViewModelEditor/Private/Customizations/MVVMPropertyPathCustomization.h
--- a/Engine/Plugins/Runtime/ModelViewViewModel/Source/ModelViewViewModelEditor/Private/Customizations/MVVMPropertyPathCustomization.h
+++ b/Engine/Plugins/Runtime/ModelViewViewModel/Source/ModelViewViewModelEditor/Private/Customizations/MVVMPropertyPathCustomization.h
@@ -32,7 +32,7 @@ namespace UE::MVVM
 		virtual void CustomizeChildren(TSharedRef<IPropertyHandle> InPropertyHandle, IDetailChildrenBuilder& ChildBuilder, IPropertyTypeCustomizationUtils& CustomizationUtils) override {}
 
 	private:
-		void OnSourceSelectionChanged(TOptional<FBindingSource> Selected);
+		void OnSourceSelectionChanged(FBindingSource Selected);
 		void OnPropertySelectionChanged(FMVVMConstFieldVariant Selected);
 
 		void OnOtherPropertyChanged();
diff --git a/Engine/Plugins/Runtime/ModelViewViewModel/Source/ModelViewViewModelEditor/Private/MVVMPropertyPathHelpers.cpp b/Engine/Plugins/Runtime/ModelViewViewModel/Source/ModelViewViewModelEditor/Private/MVVMPropertyPathHelpers.cpp
--- a/Engine/Plugins/Runtime/ModelViewViewModel/Source/ModelViewViewModelEditor/Private/MVVMPropertyPathHelpers.cpp
+++ b/Engine/Plugins/Runtime/ModelViewViewModel/Source/ModelViewViewModelEditor/Private/MVVMPropertyPathHelpers.cpp
@@ -66,7 +66,7 @@ namespace UE::MVVM
 		}
 	}
 
-	TOptional<FBindingSource> FWidgetFieldPathHelper::GetSelectedSource() const
+	FBindingSource FWidgetFieldPathHelper::GetSelectedSource() const
 	{
 		TSet<FBindingSource> Sources;
 		GetAvailableSources(Sources);
@@ -79,7 +79,7 @@ namespace UE::MVVM
 			}
 		}
 
-		return TOptional<FBindingSource>();
+		return FBindingSource();
 	}
 
 	void FWidgetFieldPathHelper::GetAvailableFields(TSet<FMVVMConstFieldVariant>& OutFields) const
@@ -136,7 +136,7 @@ namespace UE::MVVM
 		return FMVVMConstFieldVariant();
 	}
 
-	void FWidgetFieldPathHelper::SetSelectedSource(const TOptional<FBindingSource>& Source) const
+	void FWidgetFieldPathHelper::SetSelectedSource(const FBindingSource& Source) const
 	{
 		FMVVMWidgetPropertyPath* Path = PathAttr.Get(nullptr);
 		if (Path == nullptr)
@@ -144,9 +144,9 @@ namespace UE::MVVM
 			return;
 		}
 
-		if (Source.IsSet())
+		if (Source.IsValid())
 		{
-			Path->WidgetName = Source.GetValue().Name;
+			Path->WidgetName = Source.Name;
 		}
 		else
 		{
@@ -224,7 +224,7 @@ namespace UE::MVVM
 		}
 	}
 
-	TOptional<FBindingSource> FViewModelFieldPathHelper::GetSelectedSource() const
+	FBindingSource FViewModelFieldPathHelper::GetSelectedSource() const
 	{
 		TSet<FBindingSource> Sources;
 		GetAvailableSources(Sources);
@@ -237,7 +237,7 @@ namespace UE::MVVM
 			}
 		}
 
-		return TOptional<FBindingSource>();
+		return FBindingSource();
 	}
 
 	void FViewModelFieldPathHelper::GetAvailableFields(TSet<FMVVMConstFieldVariant>& OutFields) const
@@ -282,7 +282,7 @@ namespace UE::MVVM
 		return FMVVMConstFieldVariant();
 	}
 
-	void FViewModelFieldPathHelper::SetSelectedSource(const TOptional<FBindingSource>& Source) const
+	void FViewModelFieldPathHelper::SetSelectedSource(const FBindingSource& Source) const
 	{
 		FMVVMViewModelPropertyPath* Path = PathAttr.Get(nullptr);
 		if (Path == nullptr)
@@ -290,9 +290,9 @@ namespace UE::MVVM
 			return;
 		}
 
-		if (Source.IsSet())
+		if (Source.IsValid())
 		{
-			Path->ContextId = Source.GetValue().SourceGuid;
+			Path->ContextId = Source.SourceGuid;
 		}
 		else
 		{
diff --git a/Engine/Plugins/Runtime/ModelViewViewModel/Source/ModelViewViewModelEditor/Private/MVVMPropertyPathHelpers.h b/Engine/Plugins/Runtime/ModelViewViewModel/Source/ModelViewViewModelEditor/Private/MVVMPropertyPathHelpers.h
--- a/Engine/Plugins/Runtime/ModelViewViewModel/Source/ModelViewViewModelEditor/Private/MVVMPropertyPathHelpers.h
+++ b/Engine/Plugins/Runtime/ModelViewViewModel/Source/ModelViewViewModelEditor/Private/MVVMPropertyPathHelpers.h
@@ -5,7 +5,6 @@
 #include "Internationalization/Text.h"
 #include "Misc/Attribute.h"
 #include "Misc/Guid.h"
-#include "Misc/Optional.h"
 #include "MVVMPropertyPath.h"
 #include "Types/MVVMBindingName.h"
 #include "Types/MVVMBindingSource.h"
@@ -22,8 +21,8 @@ namespace UE::MVVM
 		virtual FMVVMConstFieldVariant GetSelectedField() const = 0;
 
 		virtual void GetAvailableSources(TSet<FBindingSource>& OutSources) const = 0;
-		virtual TOptional<FBindingSource> GetSelectedSource() const = 0;
-		virtual void SetSelectedSource(const TOptional<FBindingSource>& Source) const = 0;
+		virtual FBindingSource GetSelectedSource() const = 0;
+		virtual void SetSelectedSource(const FBindingSource& Source) const = 0;
 
 		virtual FMVVMBindingName GetBindingName() const = 0;
 		virtual FMemberReference GetBindingReference() const = 0;
@@ -47,8 +46,8 @@ namespace UE::MVVM
 		virtual FMVVMConstFieldVariant GetSelectedField() const override;
 
 		virtual void GetAvailableSources(TSet<FBindingSource>& OutSources) const override;
-		virtual TOptional<FBindingSource> GetSelectedSource() const override;
-		virtual void SetSelectedSource(const TOptional<FBindingSource>& Source) const override;
+		virtual FBindingSource GetSelectedSource() const override;
+		virtual void SetSelectedSource(const FBindingSource& Source) const override;
 
 		virtual FMVVMBindingName GetBindingName() const override;
 		virtual FMemberReference GetBindingReference() const override;
@@ -76,8 +75,8 @@ namespace UE::MVVM
 		virtual FMVVMConstFieldVariant GetSelectedField() const override;
 
 		virtual void GetAvailableSources(TSet<FBindingSource>& OutSources) const override;
-		virtual TOptional<FBindingSource> GetSelectedSource() const override;
-		virtual void SetSelectedSource(const TOptional<FBindingSource>& Source) const override;
+		virtual FBindingSource GetSelectedSource() const override;
+		virtual void SetSelectedSource(const FBindingSource& Source) const override;
 
 		virtual FMVVMBindingName GetBindingName() const override;
 		virtual FMemberReference GetBindingReference() const override;
diff --git a/Engine/Plugins/Runtime/ModelViewViewModel/Source/ModelViewViewModelEditor/Private/Styling/MVVMEditorStyle.cpp b/Engine/Plugins/Runtime/ModelViewViewModel/Source/ModelViewViewModelEditor/Private/Styling/MVVMEditorStyle.cpp
--- a/Engine/Plugins/Runtime/ModelViewViewModel/Source/ModelViewViewModelEditor/Private/Styling/MVVMEditorStyle.cpp
+++ b/Engine/Plugins/Runtime/ModelViewViewModel/Source/ModelViewViewModelEditor/Private/Styling/MVVMEditorStyle.cpp
@@ -36,9 +36,11 @@ FMVVMEditorStyle::FMVVMEditorStyle()
 		Set("ClassIcon.MVVMBlueprintViewModel", new IMAGE_BRUSH("Slate/PH_MVVMBlueprintViewModel_16x", Icon16x16));
 	}
 
-	{
-		Set("BlueprintView.TabIcon", new IMAGE_BRUSH_SVG("Slate/ViewModel", Icon16x16));
-	}
+	Set("BlueprintView.TabIcon", new IMAGE_BRUSH_SVG("Slate/ViewModel", Icon16x16));
+
+	Set("BindingView.ManageViewModels", new IMAGE_BRUSH_SVG("Slate/ViewModel", Icon16x16));
+	Set("BindingView.AddBinding", new IMAGE_BRUSH_SVG("Slate/ViewModel_AddBinding", Icon16x16));
+
 
 	// ViewModelSelectionWidget Icons
 	{
diff --git a/Engine/Plugins/Runtime/ModelViewViewModel/Source/ModelViewViewModelEditor/Private/Types/MVVMBindingSource.h b/Engine/Plugins/Runtime/ModelViewViewModel/Source/ModelViewViewModelEditor/Private/Types/MVVMBindingSource.h
--- a/Engine/Plugins/Runtime/ModelViewViewModel/Source/ModelViewViewModelEditor/Private/Types/MVVMBindingSource.h
+++ b/Engine/Plugins/Runtime/ModelViewViewModel/Source/ModelViewViewModelEditor/Private/Types/MVVMBindingSource.h
@@ -27,7 +27,14 @@ namespace UE::MVVM
 
 		friend int32 GetTypeHash(const FBindingSource& Source)
 		{
-			return GetTypeHash(Source.SourceGuid) ^ GetTypeHash(Source.Name) ^ GetTypeHash(Source.Class);
+			uint32 Hash = HashCombine(GetTypeHash(Source.SourceGuid), GetTypeHash(Source.Name));
+			Hash = HashCombine(Hash, GetTypeHash(Source.Class));
+			return Hash;
+		}
+
+		bool IsValid() const
+		{
+			return SourceGuid.IsValid() || !Name.IsNone();
 		}
 
 		void Reset()
@@ -69,7 +76,7 @@ struct TListTypeTraits<UE::MVVM::FBindingSource>
 
 	static bool IsPtrValid(const UE::MVVM::FBindingSource& InPtr)
 	{
-		return !InPtr.Name.IsNone() || InPtr.SourceGuid.IsValid();
+		return InPtr.IsValid();
 	}
 
 	static void ResetPtr(UE::MVVM::FBindingSource& InPtr)
diff --git a/Engine/Plugins/Runtime/ModelViewViewModel/Source/ModelViewViewModelEditor/Private/Widgets/SMVVMSourceEntry.cpp b/Engine/Plugins/Runtime/ModelViewViewModel/Source/ModelViewViewModelEditor/Private/Widgets/SMVVMSourceEntry.cpp
--- a/Engine/Plugins/Runtime/ModelViewViewModel/Source/ModelViewViewModelEditor/Private/Widgets/SMVVMSourceEntry.cpp
+++ b/Engine/Plugins/Runtime/ModelViewViewModel/Source/ModelViewViewModelEditor/Private/Widgets/SMVVMSourceEntry.cpp
@@ -9,24 +9,24 @@
 
 #define LOCTEXT_NAMESPACE "SMVVMSource"
 
-static const FSlateBrush* GetSourceIcon(const TOptional<UE::MVVM::FBindingSource>& Source)
+static const FSlateBrush* GetSourceIcon(const UE::MVVM::FBindingSource& Source)
 {
-	if (!Source.IsSet())
+	if (!Source.IsValid())
 	{
 		return nullptr;
 	}
 
-	return FSlateIconFinder::FindIconBrushForClass(Source.GetValue().Class);
+	return FSlateIconFinder::FindIconBrushForClass(Source.Class);
 }
 
-static FLinearColor GetSourceColor(const TOptional<UE::MVVM::FBindingSource>& Source)
+static FLinearColor GetSourceColor(const UE::MVVM::FBindingSource& Source)
 {
-	if (!Source.IsSet() || Source.GetValue().Class == nullptr || Source.GetValue().Class->GetName().IsEmpty())
+	if (!Source.IsValid())
 	{
 		return FLinearColor::White;
 	}
 
-	uint32 Hash = GetTypeHash(Source.GetValue().Class->GetName());
+	uint32 Hash = GetTypeHash(Source.Class->GetName());
 	FLinearColor Color = FLinearColor::White;
 	Color.R = ((Hash * 1 % 96) + 32) / 256.f;
 	Color.G = ((Hash * 2 % 96) + 32) / 256.f;
@@ -34,14 +34,14 @@ static FLinearColor GetSourceColor(const TOptional<UE::MVVM::FBindingSource>& So
 	return Color;
 }
 
-static FText GetSourceDisplayName(const TOptional<UE::MVVM::FBindingSource>& Source)
+static FText GetSourceDisplayName(const UE::MVVM::FBindingSource& Source)
 {
-	if (!Source.IsSet())
+	if (!Source.IsValid())
 	{
 		return LOCTEXT("None", "<None>");
 	}
 
-	return Source.GetValue().DisplayName;
+	return !Source.DisplayName.IsEmpty() ? Source.DisplayName : FText::FromString(Source.Name.ToString());
 }
 
 void SMVVMSourceEntry::Construct(const FArguments& InArgs)
@@ -69,7 +69,7 @@ void SMVVMSourceEntry::Construct(const FArguments& InArgs)
 	RefreshSource(InArgs._Source);
 }
 
-void SMVVMSourceEntry::RefreshSource(const TOptional<UE::MVVM::FBindingSource>& Source)
+void SMVVMSourceEntry::RefreshSource(const UE::MVVM::FBindingSource& Source)
 {
 	Image->SetImage(GetSourceIcon(Source));
 	Image->SetColorAndOpacity(GetSourceColor(Source));
diff --git a/Engine/Plugins/Runtime/ModelViewViewModel/Source/ModelViewViewModelEditor/Private/Widgets/SMVVMSourceEntry.h b/Engine/Plugins/Runtime/ModelViewViewModel/Source/ModelViewViewModelEditor/Private/Widgets/SMVVMSourceEntry.h
--- a/Engine/Plugins/Runtime/ModelViewViewModel/Source/ModelViewViewModelEditor/Private/Widgets/SMVVMSourceEntry.h
+++ b/Engine/Plugins/Runtime/ModelViewViewModel/Source/ModelViewViewModelEditor/Private/Widgets/SMVVMSourceEntry.h
@@ -22,11 +22,11 @@ public:
 		{
 		}
 		SLATE_STYLE_ARGUMENT(FTextBlockStyle, TextStyle)
-		SLATE_ARGUMENT(TOptional<UE::MVVM::FBindingSource>, Source)
+		SLATE_ARGUMENT(UE::MVVM::FBindingSource, Source)
 	SLATE_END_ARGS()
 
 	void Construct(const FArguments& Args);
-	void RefreshSource(const TOptional<UE::MVVM::FBindingSource>& Source);
+	void RefreshSource(const UE::MVVM::FBindingSource& Source);
 
 private:
 	TSharedPtr<STextBlock> Label;
diff --git a/Engine/Plugins/Runtime/ModelViewViewModel/Source/ModelViewViewModelEditor/Private/Widgets/SMVVMSourceSelector.cpp b/Engine/Plugins/Runtime/ModelViewViewModel/Source/ModelViewViewModelEditor/Private/Widgets/SMVVMSourceSelector.cpp
--- a/Engine/Plugins/Runtime/ModelViewViewModel/Source/ModelViewViewModelEditor/Private/Widgets/SMVVMSourceSelector.cpp
+++ b/Engine/Plugins/Runtime/ModelViewViewModel/Source/ModelViewViewModelEditor/Private/Widgets/SMVVMSourceSelector.cpp
@@ -65,11 +65,9 @@ void SMVVMSourceSelector::OnComboBoxSelectionChanged(FBindingSource Selected, ES
 {
 	SelectedSource = Selected;
 
-	TOptional<FBindingSource> OptionalSource = SelectedSource;
+	SelectedSourceWidget->RefreshSource(Selected);
 
-	SelectedSourceWidget->RefreshSource(OptionalSource);
-
-	OnSelectionChanged.ExecuteIfBound(OptionalSource);
+	OnSelectionChanged.ExecuteIfBound(Selected);
 }
 
 void SMVVMSourceSelector::Refresh()
diff --git a/Engine/Plugins/Runtime/ModelViewViewModel/Source/ModelViewViewModelEditor/Private/Widgets/SMVVMSourceSelector.h b/Engine/Plugins/Runtime/ModelViewViewModel/Source/ModelViewViewModelEditor/Private/Widgets/SMVVMSourceSelector.h
--- a/Engine/Plugins/Runtime/ModelViewViewModel/Source/ModelViewViewModelEditor/Private/Widgets/SMVVMSourceSelector.h
+++ b/Engine/Plugins/Runtime/ModelViewViewModel/Source/ModelViewViewModelEditor/Private/Widgets/SMVVMSourceSelector.h
@@ -14,7 +14,7 @@ class SMVVMSourceEntry;
 class SMVVMSourceSelector : public SCompoundWidget
 {
 public:
-	DECLARE_DELEGATE_OneParam(FSelectionChanged, TOptional<UE::MVVM::FBindingSource>);
+	DECLARE_DELEGATE_OneParam(FSelectionChanged, UE::MVVM::FBindingSource);
 
 	SLATE_BEGIN_ARGS(SMVVMSourceSelector) :
 		_TextStyle( &FCoreStyle::Get().GetWidgetStyle<FTextBlockStyle>( "NormalText" ) )
diff --git a/Engine/Plugins/Runtime/ModelViewViewModel/Source/ModelViewViewModelEditor/Private/Widgets/SMVVMViewBindingListView.cpp b/Engine/Plugins/Runtime/ModelViewViewModel/Source/ModelViewViewModelEditor/Private/Widgets/SMVVMViewBindingListView.cpp
--- a/Engine/Plugins/Runtime/ModelViewViewModel/Source/ModelViewViewModelEditor/Private/Widgets/SMVVMViewBindingListView.cpp
+++ b/Engine/Plugins/Runtime/ModelViewViewModel/Source/ModelViewViewModelEditor/Private/Widgets/SMVVMViewBindingListView.cpp
@@ -166,21 +166,27 @@ public:
 				}
 			}
 
-			return SNew(SComboBox<FName>)
-				.OptionsSource(&ModeNames)
-				.InitiallySelectedItem(ModeEnum->GetNameByValue((int64) ViewModelBinding->BindingType))
-				.OnSelectionChanged(this, &SMVVMViewBindingListEntryRow::OnModeSelectionChanged)
-				.OnGenerateWidget(this, &SMVVMViewBindingListEntryRow::GenerateModeWidget)
-				.Content()
+			return 
+				SNew(SBox)
+				.Padding(FMargin(2, 0))
+				.VAlign(VAlign_Center)
 				[
-					SNew(SBox)
-					.HAlign(HAlign_Center)
-					.VAlign(VAlign_Center)
-					.WidthOverride(16)
-					.HeightOverride(16)
+					SNew(SComboBox<FName>)
+					.OptionsSource(&ModeNames)
+					.InitiallySelectedItem(ModeEnum->GetNameByValue((int64) ViewModelBinding->BindingType))
+					.OnSelectionChanged(this, &SMVVMViewBindingListEntryRow::OnModeSelectionChanged)
+					.OnGenerateWidget(this, &SMVVMViewBindingListEntryRow::GenerateModeWidget)
+					.Content()
 					[
-						SNew(SImage)
-						.Image(this, &SMVVMViewBindingListEntryRow::GetCurrentModeBrush)
+						SNew(SBox)
+						.HAlign(HAlign_Center)
+						.VAlign(VAlign_Center)
+						.WidthOverride(16)
+						.HeightOverride(16)
+						[
+							SNew(SImage)
+							.Image(this, &SMVVMViewBindingListEntryRow::GetCurrentModeBrush)
+						]
 					]
 				];
 		}
@@ -217,6 +223,7 @@ public:
 				.VAlign(VAlign_Center)
 				[
 					SNew(SEnumComboBox, UpdateModeEnum)
+					.ContentPadding(FMargin(4, 0))
 					.OnEnumSelectionChanged(this, &SMVVMViewBindingListEntryRow::OnUpdateModeSelectionChanged)
 					.CurrentValue(this, &SMVVMViewBindingListEntryRow::GetUpdateModeValue)
 				];
@@ -354,7 +361,7 @@ private:
 			];
 	}
 
-	void OnViewModelSelectionChanged(TOptional<UE::MVVM::FBindingSource> Source)
+	void OnViewModelSelectionChanged(UE::MVVM::FBindingSource Source)
 	{
 		OnSourceSelectionChanged(Source, ViewModelHelper.Get());
 		if (ViewModelFieldSelector.IsValid())
@@ -363,7 +370,7 @@ private:
 		}
 	}
 
-	void OnWidgetSelectionChanged(TOptional<UE::MVVM::FBindingSource> Source)
+	void OnWidgetSelectionChanged(UE::MVVM::FBindingSource Source)
 	{
 		OnSourceSelectionChanged(Source, WidgetHelper.Get());
 		if (WidgetFieldSelector.IsValid())
@@ -372,7 +379,7 @@ private:
 		}
 	}
 
-	void OnSourceSelectionChanged(TOptional<UE::MVVM::FBindingSource> SelectedSource, UE::MVVM::IFieldPathHelper* PathHelper)
+	void OnSourceSelectionChanged(UE::MVVM::FBindingSource SelectedSource, UE::MVVM::IFieldPathHelper* PathHelper)
 	{
 		if (UMVVMBlueprintView* BlueprintViewPtr = BlueprintView.Get())
 		{
diff --git a/Engine/Plugins/Runtime/ModelViewViewModel/Source/ModelViewViewModelEditor/Private/Widgets/SMVVMViewBindingPanel.cpp b/Engine/Plugins/Runtime/ModelViewViewModel/Source/ModelViewViewModelEditor/Private/Widgets/SMVVMViewBindingPanel.cpp
--- a/Engine/Plugins/Runtime/ModelViewViewModel/Source/ModelViewViewModelEditor/Private/Widgets/SMVVMViewBindingPanel.cpp
+++ b/Engine/Plugins/Runtime/ModelViewViewModel/Source/ModelViewViewModelEditor/Private/Widgets/SMVVMViewBindingPanel.cpp
@@ -20,6 +20,7 @@
 #include "ToolMenus.h"
 
 #include "Framework/MultiBox/MultiBoxBuilder.h"
+#include "Styling/MVVMEditorStyle.h"
 #include "Widgets/Input/SButton.h"
 #include "Widgets/SNullWidget.h"
 #include "Widgets/Images/SImage.h"
@@ -356,7 +357,7 @@ TSharedRef<SWidget> SMVVMViewBindingPanel::GenerateEditViewWidget()
 			NAME_None,
 			LOCTEXT("AddBinding", "Add Binding"),
 			LOCTEXT("AddBindingTooltip", "Add an empty Binding"),
-			FSlateIcon(/*FWidgetReflectorStyle::GetStyleSetName(), WidgetReflectorIcon::TakeSnapshot*/),
+			FSlateIcon(FMVVMEditorStyle::Get().GetStyleSetName(), "BindingView.AddBinding"),
 			EUserInterfaceActionType::Button
 		);
 		ToolbarBuilderGlobal.EndSection();
@@ -373,7 +374,7 @@ TSharedRef<SWidget> SMVVMViewBindingPanel::GenerateEditViewWidget()
 			NAME_None,
 			LOCTEXT("ManageViewModels", "Manage ViewModels"),
 			LOCTEXT("ManageViewModelsTooltip", "Manage ViewModels for this widget"),
-			FSlateIcon(),
+			FSlateIcon(FMVVMEditorStyle::Get().GetStyleSetName(), "BindingView.ManageViewModels"),
 			EUserInterfaceActionType::Button
 		);
 		ToolbarBuilderGlobal.EndSection();
