Correction to ViewModels Creation errors to add more information for user to be able to act on it.

diff --git a/Engine/Plugins/Runtime/ModelViewViewModel/Source/ModelViewViewModelBlueprint/Private/MVVMViewBlueprintCompiler.cpp b/Engine/Plugins/Runtime/ModelViewViewModel/Source/ModelViewViewModelBlueprint/Private/MVVMViewBlueprintCompiler.cpp
--- a/Engine/Plugins/Runtime/ModelViewViewModel/Source/ModelViewViewModelBlueprint/Private/MVVMViewBlueprintCompiler.cpp
+++ b/Engine/Plugins/Runtime/ModelViewViewModel/Source/ModelViewViewModelBlueprint/Private/MVVMViewBlueprintCompiler.cpp
@@ -448,20 +448,22 @@ bool FMVVMViewBlueprintCompiler::PreCompileSourceCreators(UWidgetBlueprintGenera
 		if (SourceCreatorContext.Type == ECompilerSourceCreatorType::ViewModel)
 		{
 			const FMVVMBlueprintViewModelContext& ViewModelContext = SourceCreatorContext.ViewModelContext;
-			checkf(ViewModelContext.GetViewModelClass(), TEXT("The viewmodel class is invalid. It was checked in CreateSourceList"));
+			checkf(ViewModelContext.GetViewModelClass(), TEXT("The ViewModel class is invalid. It was checked in CreateSourceList"));
 
 			if (ViewModelContext.GetViewModelClass()->HasAllClassFlags(CLASS_Deprecated))
 			{
-				WidgetBlueprintCompilerContext.MessageLog.Warning(*FString::Printf(TEXT("The viewmodel '%s' is deprecated. Please update your view.")
-					, *ViewModelContext.GetViewModelClass()->GetDisplayNameText().ToString()));
+				WidgetBlueprintCompilerContext.MessageLog.Warning(*FString::Printf(TEXT("The ViewModel type '%s' is deprecated and should not be used for '%s'. Please update it in the View Binding panel under Manage ViewModels.")
+					, *ViewModelContext.GetViewModelClass()->GetDisplayNameText().ToString()
+					, *ViewModelContext.GetDisplayName().ToString()));
 			}
 
 			if (ViewModelContext.CreationType == EMVVMBlueprintViewModelContextCreationType::CreateInstance)
 			{
 				if (ViewModelContext.GetViewModelClass()->HasAllClassFlags(CLASS_Abstract))
 				{
-					WidgetBlueprintCompilerContext.MessageLog.Error(*FString::Printf(TEXT("The viewmodel '%s' is abstract and can't be created.")
-						, *ViewModelContext.GetViewModelClass()->GetDisplayNameText().ToString()));
+					WidgetBlueprintCompilerContext.MessageLog.Error(*FString::Printf(TEXT("The ViewModel type '%s' is abstract and can't be created for '%s'. You can change it in the View Binding panel under Manage ViewModels.")
+						, *ViewModelContext.GetViewModelClass()->GetDisplayNameText().ToString()
+						, *ViewModelContext.GetDisplayName().ToString()));
 					bAreSourcesCreatorValid = false;
 					continue;
 				}
@@ -470,8 +472,8 @@ bool FMVVMViewBlueprintCompiler::PreCompileSourceCreators(UWidgetBlueprintGenera
 			{
 				if (ViewModelContext.ViewModelPropertyPath.IsEmpty())
 				{
-					WidgetBlueprintCompilerContext.MessageLog.Error(*FString::Printf(TEXT("The viewmodel '%s' doesn't have a valid field path to be resolved at runtime. Add a Getter or use a global identifier.")
-						, *ViewModelContext.GetViewModelClass()->GetDisplayNameText().ToString()));
+					WidgetBlueprintCompilerContext.MessageLog.Error(*FString::Printf(TEXT("The ViewModel '%s' as an invalid Getter. You can select a new one in the View Binding panel under Manage ViewModels.")
+						, *ViewModelContext.GetDisplayName().ToString()));
 					bAreSourcesCreatorValid = true;
 					continue;
 				}
@@ -480,7 +482,8 @@ bool FMVVMViewBlueprintCompiler::PreCompileSourceCreators(UWidgetBlueprintGenera
 				TValueOrError<FCompiledBindingLibraryCompiler::FFieldPathHandle, FString> ReadFieldPathResult = BindingLibraryCompiler.AddObjectFieldPath(Class, ViewModelContext.ViewModelPropertyPath, ViewModelContext.GetViewModelClass(), true);
 				if (ReadFieldPathResult.HasError())
 				{
-					WidgetBlueprintCompilerContext.MessageLog.Error(*FString::Printf(TEXT("The field path for reading the viewmodel '%s' could not be created. %s")
+					WidgetBlueprintCompilerContext.MessageLog.Error(*FString::Printf(TEXT("The ViewModel '%s' (%s) as an invalid Getter. You can select a new one in the View Binding panel under Manage ViewModels. Details: %s")
+						, *ViewModelContext.GetDisplayName().ToString()
 						, *ViewModelContext.GetViewModelClass()->GetDisplayNameText().ToString()
 						, *ReadFieldPathResult.GetError()));
 					bAreSourcesCreatorValid = false;
@@ -493,7 +496,8 @@ bool FMVVMViewBlueprintCompiler::PreCompileSourceCreators(UWidgetBlueprintGenera
 			{
 				if (ViewModelContext.GlobalViewModelIdentifier.IsNone())
 				{
-					WidgetBlueprintCompilerContext.MessageLog.Error(*FString::Printf(TEXT("The viewmodel '%s' doesn't have a valid Global identifier.")
+					WidgetBlueprintCompilerContext.MessageLog.Error(*FString::Printf(TEXT("The ViewModel '%s' (%s) doesn't have a valid Global identifier. You can specify a new one in the View Binding panel under Manage ViewModels.")
+						, *ViewModelContext.GetDisplayName().ToString()
 						, *ViewModelContext.GetViewModelClass()->GetDisplayNameText().ToString()));
 					bAreSourcesCreatorValid = false;
 					continue;
@@ -501,7 +505,8 @@ bool FMVVMViewBlueprintCompiler::PreCompileSourceCreators(UWidgetBlueprintGenera
 			}
 			else
 			{
-				WidgetBlueprintCompilerContext.MessageLog.Error(*FString::Printf(TEXT("The viewmodel '%s' doesn't have a valid creation type.")
+				WidgetBlueprintCompilerContext.MessageLog.Error(*FString::Printf(TEXT("The ViewModel '%s' (%s) doesn't have a valid creation type. You can select one in the View Binding panel under Manage ViewModels.")
+					, *ViewModelContext.GetDisplayName().ToString()
 					, *ViewModelContext.GetViewModelClass()->GetDisplayNameText().ToString()));
 				bAreSourcesCreatorValid = true;
 				continue;
