MVVM - "Add View" button changed to "Create View Model" and also opens the Manage Viewmodels window.

Removed unnecessary derived arguments from SMVVMViewModelContextListWidget and SMVVMManageViewModelsWidget.

Finish button isn't enabled unless a viewmodel is created.

diff --git a/Engine/Plugins/Runtime/ModelViewViewModel/Source/ModelViewViewModelEditor/Private/Widgets/SMVVMManageViewModelsWidget.cpp b/Engine/Plugins/Runtime/ModelViewViewModel/Source/ModelViewViewModelEditor/Private/Widgets/SMVVMManageViewModelsWidget.cpp
--- a/Engine/Plugins/Runtime/ModelViewViewModel/Source/ModelViewViewModelEditor/Private/Widgets/SMVVMManageViewModelsWidget.cpp
+++ b/Engine/Plugins/Runtime/ModelViewViewModel/Source/ModelViewViewModelEditor/Private/Widgets/SMVVMManageViewModelsWidget.cpp
@@ -8,6 +8,7 @@
 #include "Framework/MultiBox/MultiBoxBuilder.h"
 #include "Kismet2/KismetEditorUtilities.h"
 #include "MVVMBlueprintViewModelContext.h"
+#include "MVVMEditorSubsystem.h"
 #include "MVVMViewModelBase.h"
 #include "SMVVMViewModelContextListWidget.h"
 #include "SMVVMViewModelBindingListWidget.h"
@@ -204,9 +205,14 @@ void SMVVMManageViewModelsWidget::Construct(const FArguments& InArgs)
 {
 	using UE::MVVM::Private::FMVVMViewModelTreeNode;
 
+	UMVVMEditorSubsystem* EditorSubsystem = GEditor->GetEditorSubsystem<UMVVMEditorSubsystem>();
+	UMVVMBlueprintView* BlueprintView = EditorSubsystem->RequestView(InArgs._WidgetBlueprint);
+
 	WeakParentWindow = InArgs._ParentWindow;
 	OnViewModelContextsPicked = InArgs._OnViewModelContextsPickedDelegate;
-	ViewModelContextListWidget = InArgs._ManageViewModelsWidget;
+	ViewModelContextListWidget = SNew(SMVVMViewModelContextListWidget)
+		.ExistingViewModelContexts(TArray<FMVVMBlueprintViewModelContext>(BlueprintView->GetViewModels()))
+		.WidgetBlueprint(InArgs._WidgetBlueprint);
 
 	PopulateViewModelsTree();
 
@@ -220,6 +226,7 @@ void SMVVMManageViewModelsWidget::Construct(const FArguments& InArgs)
 				SNew(SPrimaryButton)
 				.Text(LOCTEXT("ViewModelFinishButtonText", "Finish"))
 				.OnClicked(this, &SMVVMManageViewModelsWidget::HandleClicked_Finish)
+				.IsEnabled(this, &SMVVMManageViewModelsWidget::IsFinishEnabled)
 			]
 			+ SUniformGridPanel::Slot(1, 0)
 			[
@@ -822,4 +829,9 @@ TSubclassOf<UMVVMViewModelBase> SMVVMManageViewModelsWidget::GetClassFromNode(TS
 	return nullptr;
 }
 
+bool SMVVMManageViewModelsWidget::IsFinishEnabled() const
+{
+	return ViewModelContextListWidget->GetViewModelContexts().Num() > 0;
+}
+
 #undef LOCTEXT_NAMESPACE
diff --git a/Engine/Plugins/Runtime/ModelViewViewModel/Source/ModelViewViewModelEditor/Private/Widgets/SMVVMManageViewModelsWidget.h b/Engine/Plugins/Runtime/ModelViewViewModel/Source/ModelViewViewModelEditor/Private/Widgets/SMVVMManageViewModelsWidget.h
--- a/Engine/Plugins/Runtime/ModelViewViewModel/Source/ModelViewViewModelEditor/Private/Widgets/SMVVMManageViewModelsWidget.h
+++ b/Engine/Plugins/Runtime/ModelViewViewModel/Source/ModelViewViewModelEditor/Private/Widgets/SMVVMManageViewModelsWidget.h
@@ -19,13 +19,16 @@ class STreeView;
 class STableViewBase;
 class UMVVMViewModelBase;
 class UUserWidget;
+class UWidgetBlueprint;
+
 namespace UE::MVVM
 {
-struct FMVVMFieldVariant;
-namespace Private
-{
-class FMVVMViewModelTreeNode;
-}
+	struct FMVVMFieldVariant;
+
+	namespace Private
+	{
+		class FMVVMViewModelTreeNode;
+	}
 }
 
 DECLARE_DELEGATE_OneParam(FOnViewModelContextsPicked, TArray<FMVVMBlueprintViewModelContext>);
@@ -38,7 +41,7 @@ public:
 		SLATE_ARGUMENT(bool, bUseDefaultButtons)
 		SLATE_ARGUMENT(TWeakPtr<SWindow>, ParentWindow)
 		SLATE_ARGUMENT(FOnViewModelContextsPicked, OnViewModelContextsPickedDelegate)
-		SLATE_ARGUMENT(TSharedPtr<SMVVMViewModelContextListWidget>, ManageViewModelsWidget)
+		SLATE_ARGUMENT(UWidgetBlueprint*, WidgetBlueprint)
 	SLATE_END_ARGS()
 
 	enum class EViewMode : uint8
@@ -81,6 +84,8 @@ private:
 	FReply HandleClicked_Finish();
 	FReply HandleClicked_Cancel();
 
+	bool IsFinishEnabled() const;
+
 	TSubclassOf<UMVVMViewModelBase> GetClassFromNode(TSharedRef<FMVVMViewModelTreeNode> ClassNode);
 
 private:
diff --git a/Engine/Plugins/Runtime/ModelViewViewModel/Source/ModelViewViewModelEditor/Private/Widgets/SMVVMViewBindingPanel.cpp b/Engine/Plugins/Runtime/ModelViewViewModel/Source/ModelViewViewModelEditor/Private/Widgets/SMVVMViewBindingPanel.cpp
--- a/Engine/Plugins/Runtime/ModelViewViewModel/Source/ModelViewViewModelEditor/Private/Widgets/SMVVMViewBindingPanel.cpp
+++ b/Engine/Plugins/Runtime/ModelViewViewModel/Source/ModelViewViewModelEditor/Private/Widgets/SMVVMViewBindingPanel.cpp
@@ -211,16 +211,10 @@ void SMVVMViewBindingPanel::ShowManageViewModelsWindow()
 				.SupportsMaximize(false)
 				.ClientSize(FVector2D(800.0f, 600.0f));
 
-			TSharedRef<SMVVMViewModelContextListWidget> ContextListWidget = SNew(SMVVMViewModelContextListWidget)
-				.Bindings(BlueprintView->GetBindings())
-				.ExistingViewModelContexts(TArray<FMVVMBlueprintViewModelContext>(BlueprintView->GetViewModels()))
-				.WidgetBlueprint(WeakBlueprintEditor.Pin()->GetWidgetBlueprintObj())
-				.OwningWidget(Cast<UBlueprintGeneratedClass>(WeakBlueprintEditor.Pin()->GetBlueprintObj()->SkeletonGeneratedClass));
-
 			TSharedRef<SMVVMManageViewModelsWidget> ManageViewModelsWidget = SNew(SMVVMManageViewModelsWidget)
 				.ParentWindow(ManageViewModelsWindow)
 				.OnViewModelContextsPickedDelegate(ViewModelContextsPickedDelegate)
-				.ManageViewModelsWidget(ContextListWidget);
+				.WidgetBlueprint(WeakBlueprintEditor.Pin()->GetWidgetBlueprintObj());
 
 			ManageViewModelsWindow->SetContent(ManageViewModelsWidget);
 			GEditor->EditorAddModalWindow(ManageViewModelsWindow);
@@ -337,7 +331,7 @@ TSharedRef<SWidget> SMVVMViewBindingPanel::GenerateCreateViewWidget()
 			[
 				SNew(SButton)
 				.OnClicked(this, &SMVVMViewBindingPanel::HandleCreateViewClicked)
-				.Text(LOCTEXT("CreateViewButtonText", "Create View"))
+				.Text(LOCTEXT("AddViewModelButtonText", "Add View Model"))
 			]
 		];
 }
@@ -356,7 +350,7 @@ TSharedRef<SWidget> SMVVMViewBindingPanel::GenerateEditViewWidget()
 			),
 			NAME_None,
 			LOCTEXT("AddBinding", "Add Binding"),
-			LOCTEXT("AddBindingTooltip", "Add an empty Binding"),
+			LOCTEXT("AddBindingTooltip", "Add an empty binding."),
 			FSlateIcon(FMVVMEditorStyle::Get().GetStyleSetName(), "BindingView.AddBinding"),
 			EUserInterfaceActionType::Button
 		);
@@ -373,7 +367,7 @@ TSharedRef<SWidget> SMVVMViewBindingPanel::GenerateEditViewWidget()
 			),
 			NAME_None,
 			LOCTEXT("ManageViewModels", "Manage ViewModels"),
-			LOCTEXT("ManageViewModelsTooltip", "Manage ViewModels for this widget"),
+			LOCTEXT("ManageViewModelsTooltip", "Manage viewmodels for this widget."),
 			FSlateIcon(FMVVMEditorStyle::Get().GetStyleSetName(), "BindingView.ManageViewModels"),
 			EUserInterfaceActionType::Button
 		);
@@ -462,6 +456,8 @@ FReply SMVVMViewBindingPanel::HandleCreateViewClicked()
 		}
 
 		HandleBlueprintViewChangedDelegate();
+
+		ShowManageViewModelsWindow();
 	}
 
 	return FReply::Handled();
diff --git a/Engine/Plugins/Runtime/ModelViewViewModel/Source/ModelViewViewModelEditor/Private/Widgets/SMVVMViewModelContextListWidget.cpp b/Engine/Plugins/Runtime/ModelViewViewModel/Source/ModelViewViewModelEditor/Private/Widgets/SMVVMViewModelContextListWidget.cpp
--- a/Engine/Plugins/Runtime/ModelViewViewModel/Source/ModelViewViewModelEditor/Private/Widgets/SMVVMViewModelContextListWidget.cpp
+++ b/Engine/Plugins/Runtime/ModelViewViewModel/Source/ModelViewViewModelEditor/Private/Widgets/SMVVMViewModelContextListWidget.cpp
@@ -119,13 +119,12 @@ namespace UE::MVVM::Private
 	public:
 		static FName RemoveButtonColumnName;
 		static FName ClassColumnName;
-		static FName ContextIdColumnName;
+		static FName NameColumnName;
 		static FName CreationTypeColumnName;
 		static FName CreationGetterColumnName;
 
 	public:
 		SLATE_BEGIN_ARGS(SMVVMManageViewModelsListEntryRow) {}
-			SLATE_ARGUMENT(UBlueprintGeneratedClass* , OwningWidget)
 			SLATE_ARGUMENT(UWidgetBlueprint* , WidgetBlueprint)
 			SLATE_EVENT(FOnViewModelContextRenamedDelegate, OnViewModelContextRenamed)
 			SLATE_EVENT(FOnViewModelContextRemovedDelegate, OnViewModelContextRemoved)
@@ -134,11 +133,11 @@ namespace UE::MVVM::Private
 		void Construct(const FArguments& InArgs, const TSharedRef<STableViewBase>& OwnerTableView, const TSharedPtr<FMVVMBlueprintViewModelContext>& InEntry)
 		{
 			Entry = InEntry;
-			OwningWidget = InArgs._OwningWidget;
 			OnViewModelContextRenamed = InArgs._OnViewModelContextRenamed;
 			OnViewModelContextRemoved = InArgs._OnViewModelContextRemoved;
+			WidgetBlueprint = InArgs._WidgetBlueprint;
 
-			NameValidator = MakeUnique<FKismetNameValidator>(InArgs._WidgetBlueprint, Entry->GetViewModelName(), OwningWidget);
+			NameValidator = MakeUnique<FKismetNameValidator>(InArgs._WidgetBlueprint, Entry->GetViewModelName(), WidgetBlueprint->SkeletonGeneratedClass);
 
 			SMultiColumnTableRow<TSharedPtr<FMVVMBlueprintViewModelContext>>::Construct(
 				FSuperRowType::FArguments()
@@ -174,7 +173,7 @@ namespace UE::MVVM::Private
 						.Text(Entry->GetViewModelClass()->GetDisplayNameText())
 					];
 			}
-			else if (InColumnName == ContextIdColumnName)
+			else if (InColumnName == NameColumnName)
 			{
 				TSharedPtr<FMVVMBlueprintViewModelContext> TempEntry = Entry;
 
@@ -193,14 +192,14 @@ namespace UE::MVVM::Private
 							{
 								if (InText.IsEmptyOrWhitespace())
 								{
-									OutErrorMessage = LOCTEXT("ViewModelContextIdEmptyOrWhitespaceErrorMsg", "The ContextId cannot be empty or have whitespaces.");
+									OutErrorMessage = LOCTEXT("ViewModelNameEmptyOrWhitespaceError", "The View Model Name cannot be empty.");
 									return false;
 								}
 								for (TCHAR Char : InText.ToString())
 								{
 									if (FText::IsWhitespace(Char))
 									{
-										OutErrorMessage = LOCTEXT("ViewModelContextIdWhitespaceErrorMsg", "The ContextId cannot have whitespaces.");
+										OutErrorMessage = LOCTEXT("ViewModelNameContainsWhitespaceError", "The View Model Name cannot contain whitespaces.");
 										return false;
 									}
 								}
@@ -297,7 +296,7 @@ namespace UE::MVVM::Private
 			{
 				TArray<TSharedPtr<FBindingChainElement>> BindingChain;
 				FMenuBuilder MenuBuilder(true, NULL);
-				GeneratePropertyPathMenuContent(MenuBuilder, OwningWidget, BindingChain);
+				GeneratePropertyPathMenuContent(MenuBuilder, WidgetBlueprint->SkeletonGeneratedClass, BindingChain);
 
 				GetterWidgetContainer->AddSlot()
 				[
@@ -335,7 +334,7 @@ namespace UE::MVVM::Private
 							Self->UpdateGetterContainer();
 						})
 				),
-				SNew(STextBlock).Text(LOCTEXT("CreateInstance", "Create new Instance"))
+				SNew(STextBlock).Text(LOCTEXT("CreateInstance", "Create New Instance"))
 			);
 
 			MenuBuilder.AddMenuEntry(
@@ -347,7 +346,7 @@ namespace UE::MVVM::Private
 							Self->UpdateGetterContainer();
 						})
 				),
-				SNew(STextBlock).Text(LOCTEXT("GlobalViewModelCollection", "Get from Global collection"))
+				SNew(STextBlock).Text(LOCTEXT("GlobalViewModelCollection", "Get From Global collection"))
 			);
 
 			MenuBuilder.AddMenuEntry(
@@ -359,7 +358,7 @@ namespace UE::MVVM::Private
 							Self->UpdateGetterContainer();
 						})
 				),
-				SNew(STextBlock).Text(LOCTEXT("PropertyPath", "Get from Property Path"))
+				SNew(STextBlock).Text(LOCTEXT("PropertyPath", "Get From Property Path"))
 			);
 
 			return MenuBuilder.MakeWidget();
@@ -518,7 +517,7 @@ namespace UE::MVVM::Private
 		}
 
 		TSharedPtr<FMVVMBlueprintViewModelContext> Entry;
-		UBlueprintGeneratedClass* OwningWidget;
+		UWidgetBlueprint* WidgetBlueprint;
 		TSharedPtr<SHorizontalBox> GetterWidgetContainer;
 		TSharedPtr<STextBlock> PropertyPathTextBlock;
 		TSharedPtr<STextBlock> CreationTypeTextBlock;
@@ -530,7 +529,7 @@ namespace UE::MVVM::Private
 
 	FName SMVVMManageViewModelsListEntryRow::RemoveButtonColumnName = "RemoveButton";
 	FName SMVVMManageViewModelsListEntryRow::ClassColumnName = "ViewModel";
-	FName SMVVMManageViewModelsListEntryRow::ContextIdColumnName = "ContextId";
+	FName SMVVMManageViewModelsListEntryRow::NameColumnName = "Name";
 	FName SMVVMManageViewModelsListEntryRow::CreationTypeColumnName = "CreationType";
 	FName SMVVMManageViewModelsListEntryRow::CreationGetterColumnName = "Getter";
 
@@ -631,9 +630,9 @@ namespace UE::MVVM::Private
 		FReply HandleCancel()
 		{
 			bShouldDelete = false;
-			if (WeakParentWindow.IsValid())
+			if (TSharedPtr<SWindow> ParentWindow = WeakParentWindow.Pin())
 			{
-				WeakParentWindow.Pin()->RequestDestroyWindow();
+				ParentWindow->RequestDestroyWindow();
 			}
 			return FReply::Handled();
 		}
@@ -641,9 +640,9 @@ namespace UE::MVVM::Private
 		FReply HandleForceDelete()
 		{
 			bShouldDelete = true;
-			if (WeakParentWindow.IsValid())
+			if (TSharedPtr<SWindow> ParentWindow = WeakParentWindow.Pin())
 			{
-				WeakParentWindow.Pin()->RequestDestroyWindow();
+				ParentWindow->RequestDestroyWindow();
 			}
 			return FReply::Handled();
 		}
@@ -679,7 +678,6 @@ void SMVVMViewModelContextListWidget::Construct(const FArguments& InArgs)
 {
 	using UE::MVVM::Private::SMVVMManageViewModelsListEntryRow;
 
-	OwningWidget = InArgs._OwningWidget;
 	WidgetBlueprint = InArgs._WidgetBlueprint;
 	Bindings = InArgs._Bindings;
 	OnViewModelContextsUpdated = InArgs._OnViewModelContextsUpdated;
@@ -730,8 +728,8 @@ void SMVVMViewModelContextListWidget::Construct(const FArguments& InArgs)
 				.DefaultLabel(LOCTEXT("", "")) 
 				+ SHeaderRow::Column(SMVVMManageViewModelsListEntryRow::ClassColumnName)
 				.DefaultLabel(LOCTEXT("Class", "Class"))
-				+ SHeaderRow::Column(SMVVMManageViewModelsListEntryRow::ContextIdColumnName)
-				.DefaultLabel(LOCTEXT("ContextId", "Context Id"))
+				+ SHeaderRow::Column(SMVVMManageViewModelsListEntryRow::NameColumnName)
+				.DefaultLabel(LOCTEXT("Name", "Name"))
 				+ SHeaderRow::Column(SMVVMManageViewModelsListEntryRow::CreationTypeColumnName)
 				.DefaultLabel(LOCTEXT("CreationType", "Creation Type"))
 				+ SHeaderRow::Column(SMVVMManageViewModelsListEntryRow::CreationGetterColumnName)
@@ -758,13 +756,20 @@ void SMVVMViewModelContextListWidget::AddViewModelContext(TSubclassOf<UMVVMViewM
 	}
 
 	TSharedPtr<FMVVMBlueprintViewModelContext> NewContext = MakeShared<FMVVMBlueprintViewModelContext>(ViewModelClass, FGuid::NewGuid());
-	FName TempNewName = ViewModelClass->GetFName();
+	FString ClassName = ViewModelClass->ClassGeneratedBy != nullptr ? ViewModelClass->ClassGeneratedBy->GetName() : ViewModelClass->GetAuthoredName();
+	FString TempNewName = ClassName;
 	FKismetNameValidator NameValidator(WidgetBlueprint);
-	while ((NameValidator.IsValid(TempNewName) != EValidatorResult::Ok) || (ExistingContextNames.Contains(TempNewName.ToString())))
+
+	int32 Index = 1;
+	while ((NameValidator.IsValid(TempNewName) != EValidatorResult::Ok) || (ExistingContextNames.Contains(TempNewName)))
 	{
-		TempNewName = MakeUniqueObjectName(OwningWidget, ViewModelClass);
+		TempNewName = ClassName + "_";
+		TempNewName.AppendInt(Index);
+
+		++Index;
 	}
-	NewContext->OverrideDisplayName = FText::FromName(TempNewName);
+
+	NewContext->OverrideDisplayName = FText::FromString(TempNewName);
 	ContextListSource.Add(NewContext);
 	ContextListWidget->RequestListRefresh();
 }
@@ -774,7 +779,6 @@ TSharedRef<ITableRow> SMVVMViewModelContextListWidget::HandleGenerateRowForListV
 	using UE::MVVM::Private::SMVVMManageViewModelsListEntryRow;
 	return SNew(SMVVMManageViewModelsListEntryRow, OwnerTable, Item)
 		.WidgetBlueprint(WidgetBlueprint)
-		.OwningWidget(OwningWidget)
 		.OnViewModelContextRenamed(this, &SMVVMViewModelContextListWidget::IsContextNameAvailable)
 		.OnViewModelContextRemoved(this, &SMVVMViewModelContextListWidget::RemoveViewModelContext);
 }
@@ -878,7 +882,7 @@ bool SMVVMViewModelContextListWidget::ValidateRemoveViewModelContext(TSharedPtr<
 	return true;
 }
 
-TArray<FMVVMBlueprintViewModelContext> SMVVMViewModelContextListWidget::GetViewModelContexts()
+TArray<FMVVMBlueprintViewModelContext> SMVVMViewModelContextListWidget::GetViewModelContexts() const
 {
 	TArray<FMVVMBlueprintViewModelContext> NewContextList;
 	Algo::Transform(ContextListSource, NewContextList, [](TSharedPtr<FMVVMBlueprintViewModelContext> InElement)
diff --git a/Engine/Plugins/Runtime/ModelViewViewModel/Source/ModelViewViewModelEditor/Private/Widgets/SMVVMViewModelContextListWidget.h b/Engine/Plugins/Runtime/ModelViewViewModel/Source/ModelViewViewModelEditor/Private/Widgets/SMVVMViewModelContextListWidget.h
--- a/Engine/Plugins/Runtime/ModelViewViewModel/Source/ModelViewViewModelEditor/Private/Widgets/SMVVMViewModelContextListWidget.h
+++ b/Engine/Plugins/Runtime/ModelViewViewModel/Source/ModelViewViewModelEditor/Private/Widgets/SMVVMViewModelContextListWidget.h
@@ -22,7 +22,6 @@ public:
 	SLATE_BEGIN_ARGS(SMVVMViewModelContextListWidget) {}
 		SLATE_NAMED_SLOT(FArguments, ButtonsPanel)
 		SLATE_ARGUMENT(TSharedPtr<SWindow>, ParentWindow)
-		SLATE_ARGUMENT(UBlueprintGeneratedClass*, OwningWidget)
 		SLATE_ARGUMENT(UWidgetBlueprint*, WidgetBlueprint)
 		SLATE_ARGUMENT(TArrayView<const FMVVMBlueprintViewBinding>, Bindings)
 		SLATE_ARGUMENT(TArray<FMVVMBlueprintViewModelContext>, ExistingViewModelContexts)
@@ -31,7 +30,7 @@ public:
 
 	void Construct(const FArguments& InArgs);
 	void AddViewModelContext(TSubclassOf<UMVVMViewModelBase> ViewModelClass);
-	TArray<FMVVMBlueprintViewModelContext> GetViewModelContexts();
+	TArray<FMVVMBlueprintViewModelContext> GetViewModelContexts() const;
 
 private:
 	TSharedRef<ITableRow> HandleGenerateRowForListView(TSharedPtr<FMVVMBlueprintViewModelContext> Item, const TSharedRef<STableViewBase>& OwnerTable);
@@ -45,7 +44,6 @@ private:
 	TArray<TSharedPtr<FMVVMBlueprintViewModelContext>> ContextListSource;
 	TSharedPtr<SListView<TSharedPtr<FMVVMBlueprintViewModelContext>>> ContextListWidget;
 	TArray<FMVVMBlueprintViewModelContext> ExistingViewModelContexts;
-	UBlueprintGeneratedClass* OwningWidget = nullptr;
 	UWidgetBlueprint* WidgetBlueprint = nullptr;
 	FOnViewModelContextsUpdated OnViewModelContextsUpdated;
 	TArrayView<const FMVVMBlueprintViewBinding> Bindings;
