MVVM - Invalid error message that complains about duplicate setter functions no longer triggers.

diff --git a/Engine/Plugins/Runtime/ModelViewViewModel/Source/ModelViewViewModelBlueprint/Private/ViewModel/MVVMViewModelBlueprintCompiler.cpp b/Engine/Plugins/Runtime/ModelViewViewModel/Source/ModelViewViewModelBlueprint/Private/ViewModel/MVVMViewModelBlueprintCompiler.cpp
--- a/Engine/Plugins/Runtime/ModelViewViewModel/Source/ModelViewViewModelBlueprint/Private/ViewModel/MVVMViewModelBlueprintCompiler.cpp
+++ b/Engine/Plugins/Runtime/ModelViewViewModel/Source/ModelViewViewModelBlueprint/Private/ViewModel/MVVMViewModelBlueprintCompiler.cpp
@@ -206,9 +206,10 @@ void FViewModelBlueprintCompilerContext::CreateClassVariablesFromBlueprint()
 						, Property->GetMetaData(NAME_Category)
 						, false);
 
-					if (GeneratedFunction.SetterFunction == nullptr || GeneratedFunction.SetterFunction->GetFName() != FName(*SetterName))
+					if (GeneratedFunction.SetterFunction == nullptr)
 					{
-						MessageLog.Error(*FString::Printf(TEXT("The setter name %s already exist and could not be autogenerated."), *SetterName));
+						MessageLog.Error(*FString::Printf(TEXT("The setter name '%s' already exists and could not be autogenerated."), *SetterName));
+						continue;
 					}
 
 					FunctionGraphHelper::AddFunctionArgument(GeneratedFunction.SetterFunction, Property, TEXT("NewValue"));
@@ -224,9 +225,9 @@ void FViewModelBlueprintCompilerContext::CreateClassVariablesFromBlueprint()
 						, Property->GetMetaData(NAME_Category)
 						, false);
 
-					if (GeneratedFunction.NetRepFunction == nullptr || GeneratedFunction.NetRepFunction->GetFName() != FName(*NetRepName))
+					if (GeneratedFunction.NetRepFunction == nullptr)
 					{
-						MessageLog.Error(*FString::Printf(TEXT("The rep name %s already exist and could not be autogenerated."), *NetRepName));
+						MessageLog.Error(*FString::Printf(TEXT("The net replication function name '%s' already exists and could not be autogenerated."), *NetRepName));
 						continue;
 					}
 					else
