MVVM - Fixes to compiler error messages.

diff --git a/Engine/Plugins/Runtime/ModelViewViewModel/Source/ModelViewViewModel/Private/Bindings/MVVMFieldPathHelper.cpp b/Engine/Plugins/Runtime/ModelViewViewModel/Source/ModelViewViewModel/Private/Bindings/MVVMFieldPathHelper.cpp
--- a/Engine/Plugins/Runtime/ModelViewViewModel/Source/ModelViewViewModel/Private/Bindings/MVVMFieldPathHelper.cpp
+++ b/Engine/Plugins/Runtime/ModelViewViewModel/Source/ModelViewViewModel/Private/Bindings/MVVMFieldPathHelper.cpp
@@ -44,7 +44,7 @@ TValueOrError<TArray<FMVVMFieldVariant>, FString> GenerateFieldPathList(TSubclas
 		{
 			return MakeValue(StructProperty->Struct);
 		}
-		return MakeError(FString::Printf(TEXT("Field can only be object properties or struct properties. %s is a %s"), *Property->GetName(), *Property->GetClass()->GetName()));
+		return MakeError(FString::Printf(TEXT("Only object or struct properties can be used as source paths. '%s' is a '%s'."), *Property->GetName(), *Property->GetClass()->GetName()));
 	};
 
 	auto TransformWithAccessor = [](UStruct* CurrentContainer, FMVVMFieldVariant CurrentField, bool bForReading) -> TValueOrError<FMVVMFieldVariant, FString>
@@ -64,7 +64,7 @@ TValueOrError<TArray<FMVVMFieldVariant>, FString> GenerateFieldPathList(TSubclas
 					}
 					else
 					{
-						return MakeError(FString::Printf(TEXT("The BlueprintGetter %s could not be found on object %s."), *BlueprintGetter, *CurrentContainer->GetName()));
+						return MakeError(FString::Printf(TEXT("The BlueprintGetter '%s' could not be found on object '%s'."), *BlueprintGetter, *CurrentContainer->GetName()));
 					}
 				}
 			}
@@ -83,7 +83,7 @@ TValueOrError<TArray<FMVVMFieldVariant>, FString> GenerateFieldPathList(TSubclas
 					}
 					else
 					{
-						return MakeError(FString::Printf(TEXT("The BlueprintSetter %s could not be found on object %s."), *BlueprintSetter, *CurrentContainer->GetName()));
+						return MakeError(FString::Printf(TEXT("The BlueprintSetter '%s' could not be found on object %s."), *BlueprintSetter, *CurrentContainer->GetName()));
 					}
 				}
 			}
@@ -103,7 +103,7 @@ TValueOrError<TArray<FMVVMFieldVariant>, FString> GenerateFieldPathList(TSubclas
 		FMVVMFieldVariant Field = BindingHelper::FindFieldByName(CurrentContainer, FMVVMBindingName(FName(FoundIndex, InFieldPath.GetData())));
 		if (Field.IsEmpty())
 		{
-			return MakeError(FString::Printf(TEXT("The field %s could not be found on container %s."), *FName(FoundIndex, InFieldPath.GetData()).ToString(), CurrentContainer ? *CurrentContainer->GetName() : TEXT("<none>")));
+			return MakeError(FString::Printf(TEXT("The field '%s' does not exist in the struct '%s'."), *FName(FoundIndex, InFieldPath.GetData()).ToString(), CurrentContainer ? *CurrentContainer->GetName() : TEXT("<none>")));
 		}
 		else if (Field.IsProperty())
 		{
@@ -147,7 +147,7 @@ TValueOrError<TArray<FMVVMFieldVariant>, FString> GenerateFieldPathList(TSubclas
 		FMVVMFieldVariant Field = BindingHelper::FindFieldByName(CurrentContainer, FMVVMBindingName(InFieldPath.GetData()));
 		if (Field.IsEmpty())
 		{
-			return MakeError(FString::Printf(TEXT("The field %s could not be found on container %s."), InFieldPath.GetData(), CurrentContainer ? *CurrentContainer->GetName() : TEXT("<none>")));
+			return MakeError(FString::Printf(TEXT("The field '%s' does not exist in the struct '%s'."), InFieldPath.GetData(), CurrentContainer ? *CurrentContainer->GetName() : TEXT("<none>")));
 		}
 
 		if (Field.IsProperty())
diff --git a/Engine/Plugins/Runtime/ModelViewViewModel/Source/ModelViewViewModelBlueprint/Private/Bindings/MVVMCompiledBindingLibraryCompiler.cpp b/Engine/Plugins/Runtime/ModelViewViewModel/Source/ModelViewViewModelBlueprint/Private/Bindings/MVVMCompiledBindingLibraryCompiler.cpp
--- a/Engine/Plugins/Runtime/ModelViewViewModel/Source/ModelViewViewModelBlueprint/Private/Bindings/MVVMCompiledBindingLibraryCompiler.cpp
+++ b/Engine/Plugins/Runtime/ModelViewViewModel/Source/ModelViewViewModelBlueprint/Private/Bindings/MVVMCompiledBindingLibraryCompiler.cpp
@@ -245,12 +245,12 @@ TValueOrError<FCompiledBindingLibraryCompiler::FFieldPathHandle, FString> FCompi
 			{
 				if (!BindingHelper::IsValidForDestinationBinding(FieldVariant.GetProperty()))
 				{
-					return MakeError(FString::Printf(TEXT("Property %s is not writable at runtime."), *FieldVariant.GetProperty()->GetName()));
+					return MakeError(FString::Printf(TEXT("Property '%s' is not writable at runtime."), *FieldVariant.GetProperty()->GetName()));
 				}
 			}
 			else if (!BindingHelper::IsValidForSourceBinding(FieldVariant.GetProperty()))
 			{
-				return MakeError(FString::Printf(TEXT("Property %s is not readable at runtime."), *FieldVariant.GetProperty()->GetName()));
+				return MakeError(FString::Printf(TEXT("Property '%s' is not readable at runtime."), *FieldVariant.GetProperty()->GetName()));
 			}
 
 			FString ValidatedStr = ValidateContainer(FieldVariant.GetProperty(), true, !bIsLast);
@@ -267,12 +267,12 @@ TValueOrError<FCompiledBindingLibraryCompiler::FFieldPathHandle, FString> FCompi
 			{
 				if (!BindingHelper::IsValidForDestinationBinding(FieldVariant.GetFunction()))
 				{
-					return MakeError(FString::Printf(TEXT("Function %s is not writable at runtime."), *FieldVariant.GetFunction()->GetName()));
+					return MakeError(FString::Printf(TEXT("Function '%s' is not writable at runtime."), *FieldVariant.GetFunction()->GetName()));
 				}
 			}
 			else if (!BindingHelper::IsValidForSourceBinding(FieldVariant.GetFunction()))
 			{
-				return MakeError(FString::Printf(TEXT("Function %s is not readable at runtime."), *FieldVariant.GetFunction()->GetName()));
+				return MakeError(FString::Printf(TEXT("Function '%s' is not readable at runtime."), *FieldVariant.GetFunction()->GetName()));
 			}
 
 			if (bIsLast && !bInRead)
@@ -329,7 +329,7 @@ TValueOrError<FCompiledBindingLibraryCompiler::FFieldPathHandle, FString> FCompi
 	}
 	if (GeneratedField.GetValue().Num() == 0)
 	{
-		return MakeError(FString::Printf(TEXT("The field does not returned a '%s'."), *ExpectedType->GetName()));
+		return MakeError(FString::Printf(TEXT("The field does not return a '%s'."), *ExpectedType->GetName()));
 	}
 
 	const FObjectPropertyBase* ObjectPropertyBase = nullptr;
@@ -344,11 +344,11 @@ TValueOrError<FCompiledBindingLibraryCompiler::FFieldPathHandle, FString> FCompi
 
 	if (ObjectPropertyBase == nullptr)
 	{
-		return MakeError(FString::Printf(TEXT("The field does not returned a '%s'."), *ExpectedType->GetName()));
+		return MakeError(FString::Printf(TEXT("The field does not return a '%s'."), *ExpectedType->GetName()));
 	}
 	if (ObjectPropertyBase->PropertyClass == nullptr || !ExpectedType->IsChildOf(ObjectPropertyBase->PropertyClass))
 	{
-		return MakeError(FString::Printf(TEXT("The field does not returned a '%s'."), *ExpectedType->GetName()));
+		return MakeError(FString::Printf(TEXT("The field does not return a '%s'."), *ExpectedType->GetName()));
 	}
 
 	return AddFieldPath(MakeArrayView(GeneratedField.GetValue()), bInRead);
diff --git a/Engine/Plugins/Runtime/ModelViewViewModel/Source/ModelViewViewModelBlueprint/Private/MVVMViewBlueprintCompiler.cpp b/Engine/Plugins/Runtime/ModelViewViewModel/Source/ModelViewViewModelBlueprint/Private/MVVMViewBlueprintCompiler.cpp
--- a/Engine/Plugins/Runtime/ModelViewViewModel/Source/ModelViewViewModelBlueprint/Private/MVVMViewBlueprintCompiler.cpp
+++ b/Engine/Plugins/Runtime/ModelViewViewModel/Source/ModelViewViewModelBlueprint/Private/MVVMViewBlueprintCompiler.cpp
@@ -472,7 +472,7 @@ bool FMVVMViewBlueprintCompiler::PreCompileSourceCreators(UWidgetBlueprintGenera
 			{
 				if (ViewModelContext.ViewModelPropertyPath.IsEmpty())
 				{
-					WidgetBlueprintCompilerContext.MessageLog.Error(*FString::Printf(TEXT("The ViewModel '%s' as an invalid Getter. You can select a new one in the View Binding panel under Manage ViewModels.")
+					WidgetBlueprintCompilerContext.MessageLog.Error(*FString::Printf(TEXT("The ViewModel '%s' has an invalid Getter. You can select a new one in the View Binding panel under Manage ViewModels.")
 						, *ViewModelContext.GetDisplayName().ToString()));
 					bAreSourcesCreatorValid = true;
 					continue;
