MVVM - Added button to clear source and field selections.

diff --git a/Engine/Plugins/Runtime/ModelViewViewModel/Source/ModelViewViewModelEditor/Private/Widgets/SMVVMFieldSelector.cpp b/Engine/Plugins/Runtime/ModelViewViewModel/Source/ModelViewViewModelEditor/Private/Widgets/SMVVMFieldSelector.cpp
--- a/Engine/Plugins/Runtime/ModelViewViewModel/Source/ModelViewViewModelEditor/Private/Widgets/SMVVMFieldSelector.cpp
+++ b/Engine/Plugins/Runtime/ModelViewViewModel/Source/ModelViewViewModelEditor/Private/Widgets/SMVVMFieldSelector.cpp
@@ -87,45 +87,62 @@ void SMVVMFieldSelector::Construct(const FArguments& InArgs)
 
 	ChildSlot
 	[
-		SNew(SOverlay)
-		+ SOverlay::Slot()
-		.HAlign(HAlign_Center)
+		SNew(SHorizontalBox)
+		+ SHorizontalBox::Slot()
+		.HAlign(HAlign_Fill)
 		.VAlign(VAlign_Center)
 		[
-			SNew(STextBlock)
-			.Text(LOCTEXT("NoSourceSelected", "No source selected"))
-			.TextStyle(FAppStyle::Get(), "HintText")
-			.Visibility_Lambda([this]()
-			{
-				return SelectedSource.IsSet() ? EVisibility::Collapsed : EVisibility::Visible;
-			})
-		]
-		+ SOverlay::Slot()
-		[
-			SAssignNew(FieldComboBox, SComboBox<FMVVMConstFieldVariant>)
-			.Visibility_Lambda([this]()
-			{
-				return SelectedSource.IsSet() ? EVisibility::Visible : EVisibility::Hidden;
-			})
-			.OptionsSource(&AvailableFields)
-			.InitiallySelectedItem(SelectedField)
-			.OnGenerateWidget(this, &SMVVMFieldSelector::OnGenerateFieldWidget)
-			.OnSelectionChanged(this, &SMVVMFieldSelector::OnComboBoxSelectionChanged)
+			SNew(SOverlay)
+			+ SOverlay::Slot()
+			.HAlign(HAlign_Center)
+			.VAlign(VAlign_Center)
+			[
+				SNew(STextBlock)
+				.Text(LOCTEXT("NoSourceSelected", "No source selected"))
+				.TextStyle(FAppStyle::Get(), "HintText")
+				.Visibility_Lambda([this]()
+				{
+					return SelectedSource.IsSet() ? EVisibility::Collapsed : EVisibility::Visible;
+				})
+			]
+			+ SOverlay::Slot()
 			[
-				SAssignNew(SelectedEntry, SMVVMFieldEntry)
-				.TextStyle(TextStyle)
-				.Field(SelectedField)
-				.OnValidate(this, &SMVVMFieldSelector::ValidateField)
+				SAssignNew(FieldComboBox, SComboBox<FMVVMConstFieldVariant>)
+				.Visibility_Lambda([this]()
+				{
+					return SelectedSource.IsSet() ? EVisibility::Visible : EVisibility::Hidden;
+				})
+				.OptionsSource(&AvailableFields)
+				.InitiallySelectedItem(SelectedField)
+				.OnGenerateWidget(this, &SMVVMFieldSelector::OnGenerateFieldWidget)
+				.OnSelectionChanged(this, &SMVVMFieldSelector::OnComboBoxSelectionChanged)
+				[
+					SAssignNew(SelectedEntry, SMVVMFieldEntry)
+					.TextStyle(TextStyle)
+					.Field(SelectedField)
+					.OnValidate(this, &SMVVMFieldSelector::ValidateField)
+				]
 			]
 		]
+		+ SHorizontalBox::Slot()
+		.HAlign(HAlign_Right)
+		.VAlign(VAlign_Center)
+		.AutoWidth()
+		[
+			SNew(SSimpleButton)
+			.Icon(FAppStyle::Get().GetBrush("Icons.X"))
+			.ToolTipText(LOCTEXT("ClearField", "Clear field selection."))
+			.Visibility(this, &SMVVMFieldSelector::GetClearVisibility)
+			.OnClicked(this, &SMVVMFieldSelector::OnClearBinding)
+		]
 	];
 }
 
 void SMVVMFieldSelector::OnComboBoxSelectionChanged(FMVVMConstFieldVariant Selected, ESelectInfo::Type SelectionType)
 {
-	OnSelectionChangedDelegate.ExecuteIfBound(Selected);
-	
 	SelectedEntry->SetField(Selected);
+
+	OnSelectionChangedDelegate.ExecuteIfBound(Selected);
 }
 
 void SMVVMFieldSelector::Refresh()
@@ -171,8 +188,9 @@ void SMVVMFieldSelector::Refresh()
 		}
 	}
 
-	// put all incompatible properties at the end so they don't just disappear from the list without explanation
 	AvailableFields.Append(CompatibleFields.Array());
+
+	// put all incompatible properties at the end so they don't just disappear from the list without explanation
 	AvailableFields.Append(IncompatibleFields.Array());
 	
 	if (FieldComboBox.IsValid())
@@ -301,4 +319,28 @@ void SMVVMFieldEntry::SetField(const UE::MVVM::FMVVMConstFieldVariant& InField)
 	Refresh();
 }
 
+EVisibility SMVVMFieldSelector::GetClearVisibility() const
+{
+	for (const IFieldPathHelper* Helper : PathHelpers)
+	{
+		const FMVVMConstFieldVariant ThisSelection = Helper->GetSelectedField();
+		if (!ThisSelection.IsEmpty())
+		{
+			return EVisibility::Visible;
+		}
+	}
+
+	return EVisibility::Collapsed;
+}
+
+FReply SMVVMFieldSelector::OnClearBinding()
+{
+	if (FieldComboBox.IsValid())
+	{
+		FieldComboBox->ClearSelection();
+	}
+
+	return FReply::Handled();
+}
+
 #undef LOCTEXT_NAMESPACE
diff --git a/Engine/Plugins/Runtime/ModelViewViewModel/Source/ModelViewViewModelEditor/Private/Widgets/SMVVMFieldSelector.h b/Engine/Plugins/Runtime/ModelViewViewModel/Source/ModelViewViewModelEditor/Private/Widgets/SMVVMFieldSelector.h
--- a/Engine/Plugins/Runtime/ModelViewViewModel/Source/ModelViewViewModelEditor/Private/Widgets/SMVVMFieldSelector.h
+++ b/Engine/Plugins/Runtime/ModelViewViewModel/Source/ModelViewViewModelEditor/Private/Widgets/SMVVMFieldSelector.h
@@ -77,6 +77,9 @@ private:
 	TValueOrError<bool, FString> IsValidBindingForField(const UE::MVVM::FMVVMConstFieldVariant& Field, const UE::MVVM::FMVVMConstFieldVariant& CounterpartField) const;
 	TValueOrError<bool, FString> ValidateField(UE::MVVM::FMVVMConstFieldVariant Field) const;
 
+	EVisibility GetClearVisibility() const;
+	FReply OnClearBinding();
+
 private:
 	TArray<UE::MVVM::IFieldPathHelper*> PathHelpers;
 	TArray<UE::MVVM::IFieldPathHelper*> CounterpartHelpers;
diff --git a/Engine/Plugins/Runtime/ModelViewViewModel/Source/ModelViewViewModelEditor/Private/Widgets/SMVVMSourceSelector.cpp b/Engine/Plugins/Runtime/ModelViewViewModel/Source/ModelViewViewModelEditor/Private/Widgets/SMVVMSourceSelector.cpp
--- a/Engine/Plugins/Runtime/ModelViewViewModel/Source/ModelViewViewModelEditor/Private/Widgets/SMVVMSourceSelector.cpp
+++ b/Engine/Plugins/Runtime/ModelViewViewModel/Source/ModelViewViewModelEditor/Private/Widgets/SMVVMSourceSelector.cpp
@@ -3,6 +3,7 @@
 #include "SMVVMSourceSelector.h"
 
 #include "Algo/Transform.h"
+#include "SSimpleButton.h"
 #include "Widgets/SMVVMFieldIcon.h"
 #include "Widgets/SMVVMSourceEntry.h"
 
@@ -43,20 +44,37 @@ void SMVVMSourceSelector::Construct(const FArguments& Args)
 
 	ChildSlot
 	[
-		SAssignNew(SourceComboBox, SComboBox<FBindingSource>)
-		.OptionsSource(&AvailableSources)
-		.OnGenerateWidget_Lambda([this](FBindingSource Source)
-		{
-			return SNew(SMVVMSourceEntry)
+		SNew(SHorizontalBox)
+		+ SHorizontalBox::Slot()
+		.HAlign(HAlign_Fill)
+		.VAlign(VAlign_Center)
+		[
+			SAssignNew(SourceComboBox, SComboBox<FBindingSource>)
+			.OptionsSource(&AvailableSources)
+			.OnGenerateWidget_Lambda([this](FBindingSource Source)
+			{
+				return SNew(SMVVMSourceEntry)
+					.TextStyle(TextStyle)
+					.Source(Source);
+			})
+			.OnSelectionChanged(this, &SMVVMSourceSelector::OnComboBoxSelectionChanged)
+			.InitiallySelectedItem(SelectedSource)
+			[
+				SAssignNew(SelectedSourceWidget, SMVVMSourceEntry)
 				.TextStyle(TextStyle)
-				.Source(Source);
-		})
-		.OnSelectionChanged(this, &SMVVMSourceSelector::OnComboBoxSelectionChanged)
-		.InitiallySelectedItem(SelectedSource)
+				.Source(SelectedSource)
+			]
+		]
+		+ SHorizontalBox::Slot()
+		.HAlign(HAlign_Right)
+		.VAlign(VAlign_Center)
+		.AutoWidth()
 		[
-			SAssignNew(SelectedSourceWidget, SMVVMSourceEntry)
-			.TextStyle(TextStyle)
-			.Source(SelectedSource)
+			SNew(SSimpleButton)
+			.Icon(FAppStyle::Get().GetBrush("Icons.X"))
+			.ToolTipText(LOCTEXT("ClearField", "Clear source selection."))
+			.Visibility(this, &SMVVMSourceSelector::GetClearVisibility)
+			.OnClicked(this, &SMVVMSourceSelector::OnClearSource)
 		]
 	];
 }
@@ -110,4 +128,28 @@ void SMVVMSourceSelector::Refresh()
 	}
 }
 
+EVisibility SMVVMSourceSelector::GetClearVisibility() const
+{
+	for (const IFieldPathHelper* Helper : PathHelpers.Get())
+	{
+		const FBindingSource ThisSelection = Helper->GetSelectedSource();
+		if (ThisSelection.IsValid())
+		{
+			return EVisibility::Visible;
+		}
+	}
+
+	return EVisibility::Collapsed;
+}
+
+FReply SMVVMSourceSelector::OnClearSource()
+{
+	if (SourceComboBox.IsValid())
+	{
+		SourceComboBox->ClearSelection();
+	}
+
+	return FReply::Handled();
+}
+
 #undef LOCTEXT_NAMESPACE
diff --git a/Engine/Plugins/Runtime/ModelViewViewModel/Source/ModelViewViewModelEditor/Private/Widgets/SMVVMSourceSelector.h b/Engine/Plugins/Runtime/ModelViewViewModel/Source/ModelViewViewModelEditor/Private/Widgets/SMVVMSourceSelector.h
--- a/Engine/Plugins/Runtime/ModelViewViewModel/Source/ModelViewViewModelEditor/Private/Widgets/SMVVMSourceSelector.h
+++ b/Engine/Plugins/Runtime/ModelViewViewModel/Source/ModelViewViewModelEditor/Private/Widgets/SMVVMSourceSelector.h
@@ -31,6 +31,9 @@ public:
 private:
 	void OnComboBoxSelectionChanged(UE::MVVM::FBindingSource Selected, ESelectInfo::Type SelectionType);
 
+	EVisibility GetClearVisibility() const;
+	FReply OnClearSource();
+
 private:
 	TAttribute<TArray<UE::MVVM::IFieldPathHelper*>> PathHelpers;
 	FSelectionChanged OnSelectionChanged;
