Fix for uninitialized script struct member causing validation failure on QAGame cook

- default FMVVMPropertyPathBase::BindingKind to Function to pass validation test AttemptToFindUninitializedScriptStructMembers

- testing
  - BCR of Win64 QAGame locally

diff --git a/Engine/Plugins/Runtime/ModelViewViewModel/Source/ModelViewViewModelBlueprint/Public/MVVMPropertyPath.h b/Engine/Plugins/Runtime/ModelViewViewModel/Source/ModelViewViewModelBlueprint/Public/MVVMPropertyPath.h
--- a/Engine/Plugins/Runtime/ModelViewViewModel/Source/ModelViewViewModelBlueprint/Public/MVVMPropertyPath.h
+++ b/Engine/Plugins/Runtime/ModelViewViewModel/Source/ModelViewViewModelBlueprint/Public/MVVMPropertyPath.h
@@ -33,7 +33,7 @@ private:
 
 	/** If we are referencing a UFunction or FProperty */
 	UPROPERTY()
-	EBindingKind BindingKind;
+	EBindingKind BindingKind = EBindingKind::Function;
 
 public:
 	/** Get the binding name, resolves reference deprecation / redirectors / etc before returning */
