MVVM: Convert function path to member reference. Use the proper function skel class when compiling the conversion function.
#prelfight 6273eb01de15651eeb3547de

diff --git a/Engine/Plugins/Runtime/ModelViewViewModel/Source/ModelViewViewModelBlueprint/Private/Bindings/MVVMCompiledBindingLibraryCompiler.cpp b/Engine/Plugins/Runtime/ModelViewViewModel/Source/ModelViewViewModelBlueprint/Private/Bindings/MVVMCompiledBindingLibraryCompiler.cpp
--- a/Engine/Plugins/Runtime/ModelViewViewModel/Source/ModelViewViewModelBlueprint/Private/Bindings/MVVMCompiledBindingLibraryCompiler.cpp
+++ b/Engine/Plugins/Runtime/ModelViewViewModel/Source/ModelViewViewModelBlueprint/Private/Bindings/MVVMCompiledBindingLibraryCompiler.cpp
@@ -355,7 +355,7 @@ TValueOrError<FCompiledBindingLibraryCompiler::FFieldPathHandle, FString> FCompi
 }
 
 
-TValueOrError<FCompiledBindingLibraryCompiler::FFieldPathHandle, FString> FCompiledBindingLibraryCompiler::AddConversionFunctionFieldPath(TSubclassOf<UObject> SourceClass, FStringView FieldPath)
+TValueOrError<FCompiledBindingLibraryCompiler::FFieldPathHandle, FString> FCompiledBindingLibraryCompiler::AddConversionFunctionFieldPath(TSubclassOf<UObject> SourceClass, const UFunction* Function)
 {
 	Impl->bCompiled = false;
 
@@ -363,37 +363,26 @@ TValueOrError<FCompiledBindingLibraryCompiler::FFieldPathHandle, FString> FCompi
 	{
 		return MakeError(TEXT("The source class is invalid."));
 	}
-	if (FieldPath.Len() == 0)
+	if (Function == nullptr)
 	{
 		return MakeError(TEXT("The function path is empty."));
 	}
 
-	FMVVMFieldVariant FoundFunction = FMVVMFieldVariant(FindObject<UFunction>(nullptr, FieldPath.GetData(), true));
-	if (FoundFunction.IsEmpty() || !FoundFunction.IsFunction() || FoundFunction.GetFunction() == nullptr)
+	if (!BindingHelper::IsValidForRuntimeConversion(Function))
 	{
-		FoundFunction = UE::MVVM::BindingHelper::FindFieldByName(SourceClass.Get(), FMVVMBindingName(FieldPath.GetData()));
+		return MakeError(FString::Printf(TEXT("Function %s cannot be used as a runtime conversion function."), *Function->GetName()));
 	}
 
-	if (FoundFunction.IsEmpty() || !FoundFunction.IsFunction() || FoundFunction.GetFunction() == nullptr)
+	if (!Function->HasAllFunctionFlags(FUNC_Static))
 	{
-		return MakeError(FString::Printf(TEXT("The function %s could not be found."), FieldPath.GetData()));
-	}
-
-	if (!BindingHelper::IsValidForRuntimeConversion(FoundFunction.GetFunction()))
-	{
-		return MakeError(FString::Printf(TEXT("Function %s cannot be used as a runtime conversion function."), *FoundFunction.GetFunction()->GetName()));
-	}
-
-	if (!FoundFunction.GetFunction()->HasAllFunctionFlags(FUNC_Static))
-	{
-		if (!SourceClass->IsChildOf(FoundFunction.GetFunction()->GetOuterUClass()))
+		if (!SourceClass->IsChildOf(Function->GetOuterUClass()))
 		{
-			return MakeError(FString::Printf(TEXT("Function %s is going to be executed with an invalid self."), *FoundFunction.GetFunction()->GetName()));
+			return MakeError(FString::Printf(TEXT("Function %s is going to be executed with an invalid self."), *Function->GetName()));
 		}
 	}
 
 	TArray<int32> RawFieldIndexes;
-	RawFieldIndexes.Add(Impl->AddUniqueField(FoundFunction));
+	RawFieldIndexes.Add(Impl->AddUniqueField(FMVVMConstFieldVariant(Function)));
 	const int32 FoundFieldPath = Impl->FieldPaths.IndexOfByPredicate([&RawFieldIndexes](const Private::FRawFieldPath& Other)
 		{
 			return Other.RawFieldIndexes == RawFieldIndexes;
diff --git a/Engine/Plugins/Runtime/ModelViewViewModel/Source/ModelViewViewModelBlueprint/Private/MVVMViewBlueprintCompiler.cpp b/Engine/Plugins/Runtime/ModelViewViewModel/Source/ModelViewViewModelBlueprint/Private/MVVMViewBlueprintCompiler.cpp
--- a/Engine/Plugins/Runtime/ModelViewViewModel/Source/ModelViewViewModelBlueprint/Private/MVVMViewBlueprintCompiler.cpp
+++ b/Engine/Plugins/Runtime/ModelViewViewModel/Source/ModelViewViewModelBlueprint/Private/MVVMViewBlueprintCompiler.cpp
@@ -614,7 +614,7 @@ bool FMVVMViewBlueprintCompiler::PreCompileBindings(UWidgetBlueprintGeneratedCla
 		}
 
 		FMVVMViewBlueprintCompiler* Self = this;
-		auto AddBinding = [Self](UWidgetBlueprintGeneratedClass* Class, const TArrayView<FString> Getters, const FString& Setter, const FString& ConversionFunction) -> TValueOrError<FCompilerBinding, FString>
+		auto AddBinding = [Self](UWidgetBlueprintGeneratedClass* Class, const TArrayView<FString> Getters, const FString& Setter, const FMemberReference& ConversionFunctionReference) -> TValueOrError<FCompilerBinding, FString>
 		{
 			FCompilerBinding Result;
 
@@ -642,13 +642,21 @@ bool FMVVMViewBlueprintCompiler::PreCompileBindings(UWidgetBlueprintGeneratedCla
 				Result.DestinationWrite = FieldPathResult.GetValue();
 			}
 
-			if (!ConversionFunction.IsEmpty())
+
+			if (!ConversionFunctionReference.GetMemberName().IsNone())
 			{
+				const UFunction* ConversionFunction = ConversionFunctionReference.ResolveMember<UFunction>(Class);
+				if (ConversionFunction == nullptr)
+				{
+					return MakeError(FString::Printf(TEXT("The Conversion Function '%s' could not be resolved.")
+						, *ConversionFunctionReference.GetMemberName().ToString()));
+				}
+
 				TValueOrError<FCompiledBindingLibraryCompiler::FFieldPathHandle, FString> FieldPathResult = Self->BindingLibraryCompiler.AddConversionFunctionFieldPath(Class, ConversionFunction);
 				if (FieldPathResult.HasError())
 				{
 					return MakeError(FString::Printf(TEXT("Couldn't create the conversion function field path '%s'. %s")
-						, *ConversionFunction
+						, *ConversionFunction->GetPathName()
 						, *FieldPathResult.GetError()));
 				}
 				Result.ConversionFunction = FieldPathResult.GetValue();
@@ -730,7 +738,7 @@ bool FMVVMViewBlueprintCompiler::PreCompileBindings(UWidgetBlueprintGeneratedCla
 				}
 			}
 
-			TValueOrError<FCompilerBinding, FString> AddBindingResult = AddBinding(Class, Getters, Setter, Binding.Conversion.SourceToDestinationFunctionPath);
+			TValueOrError<FCompilerBinding, FString> AddBindingResult = AddBinding(Class, Getters, Setter, Binding.Conversion.SourceToDestinationFunction);
 			if (AddBindingResult.HasError())
 			{
 				AddErrorForBinding(Binding, BlueprintView, FString::Printf(TEXT("The binding could not be created. %s"), *AddBindingResult.GetError()));
@@ -815,14 +823,14 @@ bool FMVVMViewBlueprintCompiler::PreCompileBindings(UWidgetBlueprintGeneratedCla
 				if (!Binding.ViewModelPath.GetSetterPropertyPath().IsEmpty())
 				{
 					Setter.AppendChar(TEXT('.'));
-					Setter.Append(Binding.WidgetPath.GetSetterPropertyPath());
+					Setter.Append(Binding.ViewModelPath.GetSetterPropertyPath());
 				}
 			}
 
-			TValueOrError<FCompilerBinding, FString> AddBindingResult = AddBinding(Class, Getters, Setter, Binding.Conversion.DestinationToSourceFunctionPath);
+			TValueOrError<FCompilerBinding, FString> AddBindingResult = AddBinding(Class, Getters, Setter, Binding.Conversion.DestinationToSourceFunction);
 			if (AddBindingResult.HasError())
 			{
-				AddErrorForBinding(Binding, BlueprintView, FString::Printf(TEXT("The binding '%d' could not be created. %s"), *AddBindingResult.GetError()));
+				AddErrorForBinding(Binding, BlueprintView, FString::Printf(TEXT("The binding could not be created. %s"), *AddBindingResult.GetError()));
 				bIsBindingsValid = false;
 				continue;
 			}
diff --git a/Engine/Plugins/Runtime/ModelViewViewModel/Source/ModelViewViewModelBlueprint/Public/Bindings/MVVMCompiledBindingLibraryCompiler.h b/Engine/Plugins/Runtime/ModelViewViewModel/Source/ModelViewViewModelBlueprint/Public/Bindings/MVVMCompiledBindingLibraryCompiler.h
--- a/Engine/Plugins/Runtime/ModelViewViewModel/Source/ModelViewViewModelBlueprint/Public/Bindings/MVVMCompiledBindingLibraryCompiler.h
+++ b/Engine/Plugins/Runtime/ModelViewViewModel/Source/ModelViewViewModelBlueprint/Public/Bindings/MVVMCompiledBindingLibraryCompiler.h
@@ -168,7 +168,7 @@ public:
 	TValueOrError<FFieldPathHandle, FString> AddObjectFieldPath(TSubclassOf<UObject> SourceClass, FStringView FieldPath, UClass* ExpectedType, bool bRead);
 
 	/** */
-	TValueOrError<FFieldPathHandle, FString> AddConversionFunctionFieldPath(TSubclassOf<UObject> SourceClass, FStringView FieldPath);
+	TValueOrError<FFieldPathHandle, FString> AddConversionFunctionFieldPath(TSubclassOf<UObject> SourceClass, const UFunction* Function);
 
 	/** */
 	TValueOrError<FBindingHandle, FString> AddBinding(FFieldPathHandle Source, FFieldPathHandle Destination);
diff --git a/Engine/Plugins/Runtime/ModelViewViewModel/Source/ModelViewViewModelBlueprint/Public/MVVMBlueprintViewBinding.h b/Engine/Plugins/Runtime/ModelViewViewModel/Source/ModelViewViewModelBlueprint/Public/MVVMBlueprintViewBinding.h
--- a/Engine/Plugins/Runtime/ModelViewViewModel/Source/ModelViewViewModelBlueprint/Public/MVVMBlueprintViewBinding.h
+++ b/Engine/Plugins/Runtime/ModelViewViewModel/Source/ModelViewViewModelBlueprint/Public/MVVMBlueprintViewBinding.h
@@ -19,11 +19,11 @@ struct FMVVMBlueprintViewConversionPath
 
 	/** The Conversion function when converting the value from the destination to the source. */
 	UPROPERTY(EditAnywhere, Category = "MVVM", AdvancedDisplay)
-	FString DestinationToSourceFunctionPath;
+	FMemberReference DestinationToSourceFunction;
 
 	/** The Conversion function when converting the value from the source to the destination. */
 	UPROPERTY(EditAnywhere, Category = "MVVM", AdvancedDisplay)
-	FString SourceToDestinationFunctionPath;
+	FMemberReference SourceToDestinationFunction;
 };
 
 /**
diff --git a/Engine/Plugins/Runtime/ModelViewViewModel/Source/ModelViewViewModelBlueprint/Public/MVVMPropertyPath.h b/Engine/Plugins/Runtime/ModelViewViewModel/Source/ModelViewViewModelBlueprint/Public/MVVMPropertyPath.h
--- a/Engine/Plugins/Runtime/ModelViewViewModel/Source/ModelViewViewModelBlueprint/Public/MVVMPropertyPath.h
+++ b/Engine/Plugins/Runtime/ModelViewViewModel/Source/ModelViewViewModelBlueprint/Public/MVVMPropertyPath.h
@@ -23,14 +23,6 @@ private:
 	UPROPERTY(EditAnywhere, Category = "MVVM")
 	FMemberReference BindingReference;
 
-	/** The path of the property (either the getter or the get action). */
-	UPROPERTY(EditAnywhere, Category = "MVVM")
-	FString GetterPropertyPath;
-
-	/** The path of the property (either the setter or the set action). */
-	UPROPERTY(EditAnywhere, Category = "MVVM")
-	FString SetterPropertyPath;
-
 	/** If we are referencing a UFunction or FProperty */
 	UPROPERTY()
 	EBindingKind BindingKind = EBindingKind::Function;
@@ -82,19 +74,26 @@ public:
 
 	void SetBindingReference(const UE::MVVM::FMVVMConstFieldVariant& InField)
 	{
-		BindingReference = InField.CreateMemberReference();
-
-		if (InField.IsProperty())
-		{
-			BindingKind = EBindingKind::Property;
-		}
-		else if (InField.IsFunction())
+		if (InField.IsEmpty())
 		{
-			BindingKind = EBindingKind::Function;
+			Reset();
 		}
 		else
 		{
-			ensureAlwaysMsgf(false, TEXT("Binding to field of unknown type!"));
+			BindingReference = InField.CreateMemberReference();
+
+			if (InField.IsProperty())
+			{
+				BindingKind = EBindingKind::Property;
+			}
+			else if (InField.IsFunction())
+			{
+				BindingKind = EBindingKind::Function;
+			}
+			else
+			{
+				ensureAlwaysMsgf(false, TEXT("Binding to field of unknown type!"));
+			}
 		}
 	}
 
@@ -105,19 +104,17 @@ public:
 
 	FString GetGetterPropertyPath() const
 	{
-		return !GetterPropertyPath.IsEmpty() ? GetterPropertyPath : GetBindingName().ToString();
+		return GetBindingName().ToString();
 	}
 
 	FString GetSetterPropertyPath() const
 	{
-		return !SetterPropertyPath.IsEmpty() ? SetterPropertyPath : GetBindingName().ToString();
+		return GetBindingName().ToString();
 	}
 
 	bool operator==(const FMVVMPropertyPathBase& Other) const
 	{
-		return BindingReference.GetMemberName() == Other.BindingReference.GetMemberName() &&
-			GetterPropertyPath == Other.GetterPropertyPath &&
-			SetterPropertyPath == Other.SetterPropertyPath;
+		return BindingReference.GetMemberName() == Other.BindingReference.GetMemberName();
 	}
 	bool operator!=(const FMVVMPropertyPathBase& Other) const
 	{
diff --git a/Engine/Plugins/Runtime/ModelViewViewModel/Source/ModelViewViewModelEditor/Private/Customizations/MVVMConversionPathCustomization.cpp b/Engine/Plugins/Runtime/ModelViewViewModel/Source/ModelViewViewModelEditor/Private/Customizations/MVVMConversionPathCustomization.cpp
--- a/Engine/Plugins/Runtime/ModelViewViewModel/Source/ModelViewViewModelEditor/Private/Customizations/MVVMConversionPathCustomization.cpp
+++ b/Engine/Plugins/Runtime/ModelViewViewModel/Source/ModelViewViewModelEditor/Private/Customizations/MVVMConversionPathCustomization.cpp
@@ -19,8 +19,8 @@ UE::MVVM::FConversionPathCustomization::FConversionPathCustomization(UWidgetBlue
 
 void UE::MVVM::FConversionPathCustomization::CustomizeHeader(TSharedRef<IPropertyHandle> InPropertyHandle, FDetailWidgetRow& HeaderRow, IPropertyTypeCustomizationUtils& CustomizationUtils)
 {
-	SourceToDestinationProperty = InPropertyHandle->GetChildHandle(GET_MEMBER_NAME_CHECKED(FMVVMBlueprintViewConversionPath, SourceToDestinationFunctionPath));
-	DestinationToSourceProperty = InPropertyHandle->GetChildHandle(GET_MEMBER_NAME_CHECKED(FMVVMBlueprintViewConversionPath, DestinationToSourceFunctionPath));
+	SourceToDestinationProperty = InPropertyHandle->GetChildHandle(GET_MEMBER_NAME_CHECKED(FMVVMBlueprintViewConversionPath, SourceToDestinationFunction));
+	DestinationToSourceProperty = InPropertyHandle->GetChildHandle(GET_MEMBER_NAME_CHECKED(FMVVMBlueprintViewConversionPath, DestinationToSourceFunction));
 
 	HeaderRow.NameContent()
 	[
@@ -34,12 +34,12 @@ void UE::MVVM::FConversionPathCustomization::CustomizeHeader(TSharedRef<IPropert
 
 FText UE::MVVM::FConversionPathCustomization::GetSourceToDestinationPath() const
 {
-	FString Value;
-	FPropertyAccess::Result Result = SourceToDestinationProperty->GetValue(Value);
+	void* Value;
+	FPropertyAccess::Result Result = SourceToDestinationProperty->GetValueData(Value);
 
 	if (Result == FPropertyAccess::Success)
 	{
-		return FText::FromString(Value);
+		return FText::FromName(reinterpret_cast<FMemberReference*>(Value)->GetMemberName());
 	}
 
 	if (Result == FPropertyAccess::MultipleValues)
@@ -52,12 +52,12 @@ FText UE::MVVM::FConversionPathCustomization::GetSourceToDestinationPath() const
 
 FText UE::MVVM::FConversionPathCustomization::GetDestinationToSourcePath() const
 {
-	FString Value;
-	FPropertyAccess::Result Result = DestinationToSourceProperty->GetValue(Value);
+	void* Value;
+	FPropertyAccess::Result Result = DestinationToSourceProperty->GetValueData(Value);
 
 	if (Result == FPropertyAccess::Success)
 	{
-		return FText::FromString(Value);
+		return FText::FromName(reinterpret_cast<FMemberReference*>(Value)->GetMemberName());
 	}
 
 	if (Result == FPropertyAccess::MultipleValues)
@@ -71,13 +71,38 @@ FText UE::MVVM::FConversionPathCustomization::GetDestinationToSourcePath() const
 void UE::MVVM::FConversionPathCustomization::OnTextCommitted(const FText& NewValue, ETextCommit::Type CommitType, bool bSourceToDestination)
 {
 	FString NewString = NewValue.ToString();
-	OnFunctionPathChanged(NewString, bSourceToDestination);
+
+	UFunction* FoundFunction = nullptr;
+	if (WidgetBlueprint)
+	{
+		FoundFunction = WidgetBlueprint->GetBlueprintClass()->FindFunctionByName(*NewString);
+	}
+	if (FoundFunction == nullptr)
+	{
+		FoundFunction = FindObject<UFunction>(nullptr, *NewString, true);
+	}
+	OnFunctionPathChanged(FoundFunction, bSourceToDestination);
 }
 
-void UE::MVVM::FConversionPathCustomization::OnFunctionPathChanged(const FString& NewPath, bool bSourceToDestination)
+void UE::MVVM::FConversionPathCustomization::OnFunctionPathChanged(const UFunction* NewFunction, bool bSourceToDestination)
 {
 	TSharedPtr<IPropertyHandle> Handle = bSourceToDestination ? SourceToDestinationProperty : DestinationToSourceProperty;
-	Handle->SetValue(NewPath);
+
+	FStructProperty* StructProperty = CastFieldChecked<FStructProperty>(Handle->GetProperty());
+	check(StructProperty);
+
+	void* PreviousValue;
+	Handle->GetValueData(PreviousValue);
+
+	FMemberReference NewReference;
+	if (NewFunction)
+	{
+		NewReference.SetFromField<UFunction>(NewFunction, WidgetBlueprint ? WidgetBlueprint->SkeletonGeneratedClass : nullptr);
+	}
+
+	FString TextValue;
+	StructProperty->Struct->ExportText(TextValue, &NewReference, PreviousValue, nullptr, EPropertyPortFlags::PPF_None, nullptr);
+	ensure(Handle->SetValueFromFormattedString(TextValue, EPropertyValueSetFlags::DefaultFlags) == FPropertyAccess::Result::Success);
 }
 
 void UE::MVVM::FConversionPathCustomization::CustomizeChildren(TSharedRef<IPropertyHandle> InPropertyHandle, IDetailChildrenBuilder& ChildBuilder, IPropertyTypeCustomizationUtils& CustomizationUtils)
diff --git a/Engine/Plugins/Runtime/ModelViewViewModel/Source/ModelViewViewModelEditor/Private/Customizations/MVVMConversionPathCustomization.h b/Engine/Plugins/Runtime/ModelViewViewModel/Source/ModelViewViewModelEditor/Private/Customizations/MVVMConversionPathCustomization.h
--- a/Engine/Plugins/Runtime/ModelViewViewModel/Source/ModelViewViewModelEditor/Private/Customizations/MVVMConversionPathCustomization.h
+++ b/Engine/Plugins/Runtime/ModelViewViewModel/Source/ModelViewViewModelEditor/Private/Customizations/MVVMConversionPathCustomization.h
@@ -26,7 +26,7 @@ namespace UE::MVVM
 		FText GetDestinationToSourcePath() const;
 
 		void OnTextCommitted(const FText& NewValue, ETextCommit::Type CommitType, bool bSourceToDestination);
-		void OnFunctionPathChanged(const FString& NewPath, bool bSourceToDestination);
+		void OnFunctionPathChanged(const UFunction* NewPath, bool bSourceToDestination);
 
 	private:
 		UWidgetBlueprint* WidgetBlueprint = nullptr;
diff --git a/Engine/Plugins/Runtime/ModelViewViewModel/Source/ModelViewViewModelEditor/Private/Widgets/SMVVMConversionPath.cpp b/Engine/Plugins/Runtime/ModelViewViewModel/Source/ModelViewViewModelEditor/Private/Widgets/SMVVMConversionPath.cpp
--- a/Engine/Plugins/Runtime/ModelViewViewModel/Source/ModelViewViewModelEditor/Private/Widgets/SMVVMConversionPath.cpp
+++ b/Engine/Plugins/Runtime/ModelViewViewModel/Source/ModelViewViewModelEditor/Private/Widgets/SMVVMConversionPath.cpp
@@ -93,21 +93,21 @@ FString SMVVMConversionPath::GetFunctionPath() const
 	TArray<FMVVMBlueprintViewBinding*> ViewBindings = Bindings.Get(TArray<FMVVMBlueprintViewBinding*>());
 
 	bool bFirst = true;
-	FString FunctionPath;
+	FMemberReference MemberReference;
 	for (const FMVVMBlueprintViewBinding* Binding : ViewBindings)
 	{
-		FString ThisPath = bSourceToDestination ? Binding->Conversion.SourceToDestinationFunctionPath : Binding->Conversion.DestinationToSourceFunctionPath;
+		const FMemberReference& CurrentFunction = bSourceToDestination ? Binding->Conversion.SourceToDestinationFunction : Binding->Conversion.DestinationToSourceFunction;
 		if (bFirst)
 		{
-			FunctionPath = ThisPath;
+			MemberReference = CurrentFunction;
 		}
-		else if (FunctionPath != ThisPath)
+		else if (!MemberReference.IsSameReference(CurrentFunction))
 		{
 			return TEXT("Multiple Values");
 		}
 	}
 
-	return FunctionPath;
+	return MemberReference.GetMemberName().IsNone() ? TEXT("") : MemberReference.GetMemberName().ToString();
 }
 
 FText SMVVMConversionPath::GetFunctionToolTip() const
@@ -156,18 +156,29 @@ void SMVVMConversionPath::SetConversionFunction(const UFunction* Function)
 	{
 		if (bSourceToDestination)
 		{
-			Binding->Conversion.SourceToDestinationFunctionPath = Function != nullptr ? Function->GetPathName() : FString();
+			if (Function)
+			{
+				Binding->Conversion.SourceToDestinationFunction.SetFromField<UFunction>(Function, WidgetBlueprint->SkeletonGeneratedClass);
+			}
+			else
+			{
+				Binding->Conversion.SourceToDestinationFunction = FMemberReference();
+			}
 		}
 		else
 		{
-			Binding->Conversion.DestinationToSourceFunctionPath = Function != nullptr ? Function->GetPathName() : FString();
+			if (Function)
+			{
+				Binding->Conversion.DestinationToSourceFunction.SetFromField<UFunction>(Function, WidgetBlueprint->SkeletonGeneratedClass);
+			}
+			else
+			{
+				Binding->Conversion.DestinationToSourceFunction = FMemberReference();
+			}
 		}
 	}
 
-	if (OnFunctionChanged.IsBound())
-	{
-		OnFunctionChanged.Execute(Function != nullptr ? Function->GetPathName() : FString());
-	}
+	OnFunctionChanged.ExecuteIfBound(Function);
 }
 
 TSharedRef<SWidget> SMVVMConversionPath::GetFunctionMenuContent()
diff --git a/Engine/Plugins/Runtime/ModelViewViewModel/Source/ModelViewViewModelEditor/Private/Widgets/SMVVMConversionPath.h b/Engine/Plugins/Runtime/ModelViewViewModel/Source/ModelViewViewModelEditor/Private/Widgets/SMVVMConversionPath.h
--- a/Engine/Plugins/Runtime/ModelViewViewModel/Source/ModelViewViewModelEditor/Private/Widgets/SMVVMConversionPath.h
+++ b/Engine/Plugins/Runtime/ModelViewViewModel/Source/ModelViewViewModelEditor/Private/Widgets/SMVVMConversionPath.h
@@ -11,7 +11,7 @@ class UWidgetBlueprint;
 class SMVVMConversionPath : public SCompoundWidget
 {
 public:
-	DECLARE_DELEGATE_OneParam(FOnFunctionChanged, const FString&);
+	DECLARE_DELEGATE_OneParam(FOnFunctionChanged, const UFunction*);
 
 	SLATE_BEGIN_ARGS(SMVVMConversionPath)
 		{
