MVVM - The Getter dropdown in Manage ViewModels when the mode is  set to Get From Property Path list no longer includes the viewmodel entry itself.

Removed unnecessarily duplicated widget code taken from SMVVMFieldEntry in favor of just using that.

#rnx

diff --git a/Engine/Plugins/Runtime/ModelViewViewModel/Source/ModelViewViewModelEditor/Private/Widgets/SMVVMViewModelContextListWidget.cpp b/Engine/Plugins/Runtime/ModelViewViewModel/Source/ModelViewViewModelEditor/Private/Widgets/SMVVMViewModelContextListWidget.cpp
--- a/Engine/Plugins/Runtime/ModelViewViewModel/Source/ModelViewViewModelEditor/Private/Widgets/SMVVMViewModelContextListWidget.cpp
+++ b/Engine/Plugins/Runtime/ModelViewViewModel/Source/ModelViewViewModelEditor/Private/Widgets/SMVVMViewModelContextListWidget.cpp
@@ -15,6 +15,8 @@
 #include "SSimpleButton.h"
 #include "Styling/MVVMEditorStyle.h"
 #include "Styling/SlateIconFinder.h"
+#include "Types/MVVMFieldVariant.h"
+#include "Widgets/SMVVMFieldSelector.h"
 #include "Widgets/Input/SComboButton.h"
 #include "Widgets/Input/SEditableTextBox.h"
 #include "Widgets/Images/SImage.h"
@@ -366,62 +368,13 @@ namespace UE::MVVM::Private
 
 		void GeneratePropertyPathMenuContent(FMenuBuilder& MenuBuilder, UStruct* InOwnerStruct, TArray<TSharedPtr<FBindingChainElement>> InBindingChain)
 		{
+			using UE::MVVM::FMVVMConstFieldVariant;
 			using namespace UE::MVVM::Private;
 
-			auto MakePropertyWidget = [](FProperty* InProperty)
+			auto MakeFieldWidget = [](FMVVMConstFieldVariant Field)
 			{
-				return SNew(SHorizontalBox)
-					.ToolTipText(InProperty->GetToolTipText())
-					+ SHorizontalBox::Slot()
-					.AutoWidth()
-					[
-						SNew(SSpacer)
-						.Size(FVector2D(18.0f, 0.0f))
-					]
-				+ SHorizontalBox::Slot()
-					.AutoWidth()
-					.VAlign(VAlign_Center)
-					.Padding(1.0f, 0.0f)
-					[
-						SNew(SImage)
-						.Image(FSlateIconFinder::FindIconBrushForClass(UMVVMViewModelBase::StaticClass()))
-					]
-				+ SHorizontalBox::Slot()
-					.AutoWidth()
-					.VAlign(VAlign_Center)
-					.Padding(4.0f, 0.0f)
-					[
-						SNew(STextBlock)
-						.Text(InProperty->GetDisplayNameText())
-					];
-			};
-
-			auto MakeFunctionWidget = [](const UFunction* Info)
-			{
-				return SNew(SHorizontalBox)
-					.ToolTipText(Info->GetMetaDataText("ToolTip"))
-					+ SHorizontalBox::Slot()
-					.AutoWidth()
-					[
-						SNew(SSpacer)
-						.Size(FVector2D(18.0f, 0.0f))
-					]
-				+ SHorizontalBox::Slot()
-					.AutoWidth()
-					.VAlign(VAlign_Center)
-					.Padding(1.0f, 0.0f)
-					[
-						SNew(SImage)
-						.Image(FAppStyle::Get().GetBrush("GraphEditor.Function_16x"))
-					]
-				+ SHorizontalBox::Slot()
-					.AutoWidth()
-					.VAlign(VAlign_Center)
-					.Padding(4.0f, 0.0f)
-					[
-						SNew(STextBlock)
-						.Text(Info->GetDisplayNameText())
-					];
+				return SNew(SMVVMFieldEntry)
+					.Field(Field);
 			};
 
 			UClass* InBindableClass = Cast<UClass>(InOwnerStruct);
@@ -441,7 +394,7 @@ namespace UE::MVVM::Private
 				{
 					MenuBuilder.AddMenuEntry(
 						FUIAction(FExecuteAction::CreateSP(this, &SMVVMManageViewModelsListEntryRow::HandleAddBinding, NewBindingChain)),
-						MakeFunctionWidget(Function)
+						MakeFieldWidget(FMVVMConstFieldVariant(Function))
 					);
 				}
 				else
@@ -453,7 +406,7 @@ namespace UE::MVVM::Private
 						if (HasBindableViewModelFieldRecursive(ObjectPropertyBase->PropertyClass, ViewModelClass, VisitedSet, 1))
 						{
 							MenuBuilder.AddSubMenu(
-								MakeFunctionWidget(Function),
+								MakeFieldWidget(FMVVMConstFieldVariant(Function)),
 								FNewMenuDelegate::CreateSP(this, &SMVVMManageViewModelsListEntryRow::GeneratePropertyPathMenuContent, static_cast<UStruct*>(ObjectPropertyBase->PropertyClass), NewBindingChain)
 							);
 						}
@@ -463,7 +416,7 @@ namespace UE::MVVM::Private
 						if (HasBindableViewModelFieldRecursive(StructProperty->Struct, ViewModelClass, VisitedSet, 1))
 						{
 							MenuBuilder.AddSubMenu(
-								MakeFunctionWidget(Function),
+								MakeFieldWidget(FMVVMConstFieldVariant(Function)),
 								FNewMenuDelegate::CreateSP(this, &SMVVMManageViewModelsListEntryRow::GeneratePropertyPathMenuContent, static_cast<UStruct*>(StructProperty->Struct), NewBindingChain)
 							);
 						}
@@ -480,11 +433,12 @@ namespace UE::MVVM::Private
 				TArray<TSharedPtr<FBindingChainElement>> NewBindingChain(InBindingChain);
 				NewBindingChain.Emplace(MakeShared<FBindingChainElement>(Property));
 
-				if (IsPropertyTypeChildOf(Property, ViewModelClass))
+				if (IsPropertyTypeChildOf(Property, ViewModelClass) &&
+					Property->GetFName() != Entry->GetViewModelName()) // don't include the entry itself)
 				{
 					MenuBuilder.AddMenuEntry(
 						FUIAction(FExecuteAction::CreateSP(this, &SMVVMManageViewModelsListEntryRow::HandleAddBinding, NewBindingChain)),
-						MakePropertyWidget(Property)
+						MakeFieldWidget(FMVVMConstFieldVariant(Property))
 					);
 				}
 				else
@@ -496,7 +450,7 @@ namespace UE::MVVM::Private
 						if (HasBindableViewModelFieldRecursive(ObjectPropertyBase->PropertyClass, ViewModelClass, VisitedSet, 1))
 						{
 							MenuBuilder.AddSubMenu(
-								MakePropertyWidget(Property),
+								MakeFieldWidget(FMVVMConstFieldVariant(Property)),
 								FNewMenuDelegate::CreateSP(this, &SMVVMManageViewModelsListEntryRow::GeneratePropertyPathMenuContent, static_cast<UStruct*>(ObjectPropertyBase->PropertyClass), NewBindingChain)
 							);
 						}
@@ -506,7 +460,7 @@ namespace UE::MVVM::Private
 						if (HasBindableViewModelFieldRecursive(StructProperty->Struct, ViewModelClass, VisitedSet, 1))
 						{
 							MenuBuilder.AddSubMenu(
-								MakePropertyWidget(Property),
+								MakeFieldWidget(FMVVMConstFieldVariant(Property)),
 								FNewMenuDelegate::CreateSP(this, &SMVVMManageViewModelsListEntryRow::GeneratePropertyPathMenuContent, static_cast<UStruct*>(StructProperty->Struct), NewBindingChain)
 							);
 						}
