MVVM: Add add EMVVMBlueprintViewModelContextCreationType::Manual. An option for how to create/select a view model.
Rename a graph temporary function when the compiler didn't requested a GC at the end of the compilation.

diff --git a/Engine/Plugins/Runtime/ModelViewViewModel/Source/ModelViewViewModel/Private/View/MVVMViewClass.cpp b/Engine/Plugins/Runtime/ModelViewViewModel/Source/ModelViewViewModel/Private/View/MVVMViewClass.cpp
--- a/Engine/Plugins/Runtime/ModelViewViewModel/Source/ModelViewViewModel/Private/View/MVVMViewClass.cpp
+++ b/Engine/Plugins/Runtime/ModelViewViewModel/Source/ModelViewViewModel/Private/View/MVVMViewClass.cpp
@@ -44,6 +44,7 @@ FMVVMViewClass_SourceCreator FMVVMViewClass_SourceCreator::MakeFieldPath(FName I
 			Result.PropertyName = InName;
 			Result.ExpectedSourceType = InNotifyFieldValueChangedClass;
 			Result.FieldPath = InFieldPath;
+			Result.bOptional = false;
 		}
 	}
 	return Result;
@@ -58,6 +59,7 @@ FMVVMViewClass_SourceCreator FMVVMViewClass_SourceCreator::MakeGlobalContext(FNa
 		Result.PropertyName = InName;
 		Result.ExpectedSourceType = InContext.ContextClass;
 		Result.GlobalViewModelInstance = MoveTemp(InContext);
+		Result.bOptional = false;
 	}
 	return Result;
 }
diff --git a/Engine/Plugins/Runtime/ModelViewViewModel/Source/ModelViewViewModelBlueprint/Private/MVVMViewBlueprintCompiler.cpp b/Engine/Plugins/Runtime/ModelViewViewModel/Source/ModelViewViewModelBlueprint/Private/MVVMViewBlueprintCompiler.cpp
--- a/Engine/Plugins/Runtime/ModelViewViewModel/Source/ModelViewViewModelBlueprint/Private/MVVMViewBlueprintCompiler.cpp
+++ b/Engine/Plugins/Runtime/ModelViewViewModel/Source/ModelViewViewModelBlueprint/Private/MVVMViewBlueprintCompiler.cpp
@@ -6,6 +6,7 @@
 #include "Bindings/MVVMBindingHelper.h"
 #include "Bindings/MVVMCompiledBindingLibraryCompiler.h"
 #include "Bindings/MVVMFieldPathHelper.h"
+#include "EdGraphSchema_K2.h"
 #include "MVVMBlueprintView.h"
 #include "MVVMFunctionGraphHelper.h"
 #include "MVVMSubsystem.h"
@@ -54,7 +55,7 @@ void FMVVMViewBlueprintCompiler::CleanOldData(UWidgetBlueprintGeneratedClass* Cl
 				{
 					Children.Add(Child);
 				}
-			}, bIncludeNestedObjects);
+			}, bIncludeNestedObjects);		
 
 		for (UObject* Child : Children)
 		{
@@ -64,6 +65,20 @@ void FMVVMViewBlueprintCompiler::CleanOldData(UWidgetBlueprintGeneratedClass* Cl
 }
 
 
+void FMVVMViewBlueprintCompiler::CleanTemporaries(UWidgetBlueprintGeneratedClass* ClassToClean)
+{
+	for (FCompilerSourceCreatorContext& SourceCreatorContext : SourceCreatorContexts)
+	{
+		if (SourceCreatorContext.SetterGraph)
+		{
+			ensureMsgf(SourceCreatorContext.SetterGraph->HasAnyFlags(RF_Transient), TEXT("The graph should be temporary and should be generated automaticly."));
+			// GC may not have clean the graph (GC doesn't run when bRegenerateSkeletonOnly is on)
+			SourceCreatorContext.SetterGraph->Rename(nullptr, SourceCreatorContext.SetterGraph->GetOuter(), REN_DoNotDirty | REN_ForceNoResetLoaders);
+		}
+	}
+}
+
+
 void FMVVMViewBlueprintCompiler::CreateFunctions(UMVVMBlueprintView* BlueprintView)
 {
 	if (!bAreSourcesCreatorValid || !bIsBindingsValid)
@@ -117,19 +132,19 @@ void FMVVMViewBlueprintCompiler::CreateVariables(const FWidgetBlueprintCompilerC
 #if WITH_EDITOR
 			if (!SourceContext.BlueprintSetter.IsEmpty())
 			{
-				const FName NAME_MetaDataBlueprintSetter = "BlueprintSetter";
-				ViewModelProperty->SetMetaData(NAME_MetaDataBlueprintSetter, *SourceContext.BlueprintSetter);
+				ViewModelProperty->SetMetaData(FBlueprintMetadata::MD_PropertySetFunction, *SourceContext.BlueprintSetter);
 			}
 			if (!SourceContext.DisplayName.IsEmpty())
 			{
-				const FName NAME_MetaDataDisplayName = "DisplayName";
-				ViewModelProperty->SetMetaData(NAME_MetaDataDisplayName, *SourceContext.DisplayName);
+				ViewModelProperty->SetMetaData(FBlueprintMetadata::MD_FunctionCategory, *SourceContext.DisplayName);
 			}
-
 			if (!SourceContext.CategoryName.IsEmpty())
 			{
-				const FName NAME_MetaDataCategoryName = "Category";
-				ViewModelProperty->SetMetaData(NAME_MetaDataCategoryName, *SourceContext.CategoryName);
+				ViewModelProperty->SetMetaData(FBlueprintMetadata::MD_FunctionCategory, *SourceContext.CategoryName);
+			}
+			if (SourceContext.bExposeOnSpawn)
+			{
+				ViewModelProperty->SetMetaData(FBlueprintMetadata::MD_ExposeOnSpawn, TEXT("true"));
 			}
 #endif
 		}
@@ -263,13 +278,15 @@ void FMVVMViewBlueprintCompiler::CreateSourceLists(const FWidgetBlueprintCompile
 			continue;
 		}
 
+		const bool bCreateSetterFunction = ViewModelContext.bCreateSetterFunction || ViewModelContext.CreationType == EMVVMBlueprintViewModelContextCreationType::Manual;
+
 		int32 FoundSourceContextIndex = INDEX_NONE;
 		if (Context.GetCompileType() == EKismetCompileType::SkeletonOnly)
 		{
 			FCompilerSourceCreatorContext SourceContext;
 			SourceContext.ViewModelContext = ViewModelContext;
 			SourceContext.Type = ECompilerSourceCreatorType::ViewModel;
-			if (ViewModelContext.bCreateSetterFunction)
+			if (bCreateSetterFunction)
 			{
 				SourceContext.SetterFunctionName = TEXT("Set") + ViewModelContext.GetViewModelName().ToString();
 			}
@@ -290,7 +307,7 @@ void FMVVMViewBlueprintCompiler::CreateSourceLists(const FWidgetBlueprintCompile
 		SourceVariable.PropertyName = ViewModelContext.GetViewModelName();
 		SourceVariable.DisplayName = ViewModelContext.GetDisplayName().ToString();
 		SourceVariable.CategoryName = TEXT("Viewmodel");
-		SourceVariable.bExposeOnSpawn = ViewModelContext.bCreateSetterFunction;
+		SourceVariable.bExposeOnSpawn = bCreateSetterFunction;
 		SourceVariable.BlueprintSetter = SourceCreatorContexts[FoundSourceContextIndex].SetterFunctionName;
 		SourceContexts.Emplace(MoveTemp(SourceVariable));
 	}
@@ -457,7 +474,10 @@ bool FMVVMViewBlueprintCompiler::PreCompileSourceCreators(UWidgetBlueprintGenera
 					, *ViewModelContext.GetDisplayName().ToString()));
 			}
 
-			if (ViewModelContext.CreationType == EMVVMBlueprintViewModelContextCreationType::CreateInstance)
+			if (ViewModelContext.CreationType == EMVVMBlueprintViewModelContextCreationType::Manual)
+			{
+			}
+			else if (ViewModelContext.CreationType == EMVVMBlueprintViewModelContextCreationType::CreateInstance)
 			{
 				if (ViewModelContext.GetViewModelClass()->HasAllClassFlags(CLASS_Abstract))
 				{
@@ -527,59 +547,62 @@ bool FMVVMViewBlueprintCompiler::CompileSourceCreators(const FCompiledBindingLib
 
 	for (const FCompilerSourceCreatorContext& SourceCreatorContext : SourceCreatorContexts)
 	{
-		FMVVMViewClass_SourceCreator CompiledSourceCreator;
-
 		if (SourceCreatorContext.Type == ECompilerSourceCreatorType::ViewModel)
 		{
 			const FMVVMBlueprintViewModelContext& ViewModelContext = SourceCreatorContext.ViewModelContext;
-			if (ViewModelContext.CreationType == EMVVMBlueprintViewModelContextCreationType::CreateInstance)
+			if (ViewModelContext.CreationType != EMVVMBlueprintViewModelContextCreationType::Manual)
 			{
-				CompiledSourceCreator = FMVVMViewClass_SourceCreator::MakeInstance(ViewModelContext.GetViewModelName(), ViewModelContext.GetViewModelClass());
-			}
-			else if (ViewModelContext.CreationType == EMVVMBlueprintViewModelContextCreationType::PropertyPath)
-			{
-				const FMVVMVCompiledFieldPath* CompiledFieldPath = CompileResult.FieldPaths.Find(SourceCreatorContext.ReadPropertyPath);
-				if (CompiledFieldPath == nullptr)
+				FMVVMViewClass_SourceCreator CompiledSourceCreator;
+
+				if (ViewModelContext.CreationType == EMVVMBlueprintViewModelContextCreationType::CreateInstance)
 				{
-					WidgetBlueprintCompilerContext.MessageLog.Error(*FString::Printf(TEXT("The viewmodel '%s' initialization binding was not generated.")
-						, *ViewModelContext.GetViewModelClass()->GetDisplayNameText().ToString()));
-					bAreSourcesCreatorValid = false;
-					continue;
+					CompiledSourceCreator = FMVVMViewClass_SourceCreator::MakeInstance(ViewModelContext.GetViewModelName(), ViewModelContext.GetViewModelClass());
 				}
-				CompiledSourceCreator = FMVVMViewClass_SourceCreator::MakeFieldPath(ViewModelContext.GetViewModelName(), ViewModelContext.GetViewModelClass(), *CompiledFieldPath);
-			}
-			else if (ViewModelContext.CreationType == EMVVMBlueprintViewModelContextCreationType::GlobalViewModelCollection)
-			{
-				if (ViewModelContext.GlobalViewModelIdentifier.IsNone())
+				else if (ViewModelContext.CreationType == EMVVMBlueprintViewModelContextCreationType::PropertyPath)
 				{
-					WidgetBlueprintCompilerContext.MessageLog.Error(*FString::Printf(TEXT("The viewmodel '%s' doesn't have a valid Global identifier.")
-						, *ViewModelContext.GetViewModelClass()->GetDisplayNameText().ToString()));
-					bAreSourcesCreatorValid = false;
-					continue;
+					const FMVVMVCompiledFieldPath* CompiledFieldPath = CompileResult.FieldPaths.Find(SourceCreatorContext.ReadPropertyPath);
+					if (CompiledFieldPath == nullptr)
+					{
+						WidgetBlueprintCompilerContext.MessageLog.Error(*FString::Printf(TEXT("The viewmodel '%s' initialization binding was not generated.")
+							, *ViewModelContext.GetViewModelClass()->GetDisplayNameText().ToString()));
+						bAreSourcesCreatorValid = false;
+						continue;
+					}
+					CompiledSourceCreator = FMVVMViewClass_SourceCreator::MakeFieldPath(ViewModelContext.GetViewModelName(), ViewModelContext.GetViewModelClass(), *CompiledFieldPath);
 				}
+				else if (ViewModelContext.CreationType == EMVVMBlueprintViewModelContextCreationType::GlobalViewModelCollection)
+				{
+					if (ViewModelContext.GlobalViewModelIdentifier.IsNone())
+					{
+						WidgetBlueprintCompilerContext.MessageLog.Error(*FString::Printf(TEXT("The viewmodel '%s' doesn't have a valid Global identifier.")
+							, *ViewModelContext.GetViewModelClass()->GetDisplayNameText().ToString()));
+						bAreSourcesCreatorValid = false;
+						continue;
+					}
 
-				FMVVMViewModelContext GlobalViewModelInstance;
-				GlobalViewModelInstance.ContextClass = ViewModelContext.GetViewModelClass();
-				GlobalViewModelInstance.ContextName = ViewModelContext.GlobalViewModelIdentifier;
-				if (!GlobalViewModelInstance.IsValid())
+					FMVVMViewModelContext GlobalViewModelInstance;
+					GlobalViewModelInstance.ContextClass = ViewModelContext.GetViewModelClass();
+					GlobalViewModelInstance.ContextName = ViewModelContext.GlobalViewModelIdentifier;
+					if (!GlobalViewModelInstance.IsValid())
+					{
+						WidgetBlueprintCompilerContext.MessageLog.Error(*FString::Printf(TEXT("The context for viewmodel '%s' could not be created.")
+							, *ViewModelContext.GetViewModelClass()->GetDisplayNameText().ToString()));
+						bAreSourcesCreatorValid = false;
+						continue;
+					}
+
+					CompiledSourceCreator = FMVVMViewClass_SourceCreator::MakeGlobalContext(ViewModelContext.GetViewModelName(), MoveTemp(GlobalViewModelInstance));
+				}
+				else
 				{
-					WidgetBlueprintCompilerContext.MessageLog.Error(*FString::Printf(TEXT("The context for viewmodel '%s' could not be created.")
+					WidgetBlueprintCompilerContext.MessageLog.Error(*FString::Printf(TEXT("The viewmodel '%s' doesn't have a valid creation type.")
 						, *ViewModelContext.GetViewModelClass()->GetDisplayNameText().ToString()));
 					bAreSourcesCreatorValid = false;
 					continue;
 				}
 
-				CompiledSourceCreator = FMVVMViewClass_SourceCreator::MakeGlobalContext(ViewModelContext.GetViewModelName(), MoveTemp(GlobalViewModelInstance));
+				ViewExtension->SourceCreators.Add(MoveTemp(CompiledSourceCreator));
 			}
-			else
-			{
-				WidgetBlueprintCompilerContext.MessageLog.Error(*FString::Printf(TEXT("The viewmodel '%s' doesn't have a valid creation type.")
-					, *ViewModelContext.GetViewModelClass()->GetDisplayNameText().ToString()));
-				bAreSourcesCreatorValid = false;
-				continue;
-			}
-
-			ViewExtension->SourceCreators.Add(MoveTemp(CompiledSourceCreator));
 		}
 	}
 
diff --git a/Engine/Plugins/Runtime/ModelViewViewModel/Source/ModelViewViewModelBlueprint/Private/MVVMViewBlueprintCompiler.h b/Engine/Plugins/Runtime/ModelViewViewModel/Source/ModelViewViewModelBlueprint/Private/MVVMViewBlueprintCompiler.h
--- a/Engine/Plugins/Runtime/ModelViewViewModel/Source/ModelViewViewModelBlueprint/Private/MVVMViewBlueprintCompiler.h
+++ b/Engine/Plugins/Runtime/ModelViewViewModel/Source/ModelViewViewModelBlueprint/Private/MVVMViewBlueprintCompiler.h
@@ -42,6 +42,7 @@ public:
 
 	void AddExtension(UWidgetBlueprintGeneratedClass* Class, UMVVMViewClass* ViewExtension);
 	void CleanOldData(UWidgetBlueprintGeneratedClass* ClassToClean, UObject* OldCDO);
+	void CleanTemporaries(UWidgetBlueprintGeneratedClass* ClassToClean);
 
 	/** Generate function that are hidden from the user (not on the Skeleton class). */
 	void CreateFunctions(UMVVMBlueprintView* BlueprintView);
diff --git a/Engine/Plugins/Runtime/ModelViewViewModel/Source/ModelViewViewModelBlueprint/Private/MVVMWidgetBlueprintExtension_View.cpp b/Engine/Plugins/Runtime/ModelViewViewModel/Source/ModelViewViewModelBlueprint/Private/MVVMWidgetBlueprintExtension_View.cpp
--- a/Engine/Plugins/Runtime/ModelViewViewModel/Source/ModelViewViewModelBlueprint/Private/MVVMWidgetBlueprintExtension_View.cpp
+++ b/Engine/Plugins/Runtime/ModelViewViewModel/Source/ModelViewViewModelBlueprint/Private/MVVMWidgetBlueprintExtension_View.cpp
@@ -115,19 +115,22 @@ void UMVVMWidgetBlueprintExtension_View::HandleFinishCompilingClass(UWidgetBluep
 
 	if (CurrentCompilerContext->GetCompilerContext().bAssignDelegateSignatureFunction)
 	{
+		CurrentCompilerContext->CleanTemporaries(Class);
 		return;
 	}
 
-	if (!CurrentCompilerContext->PreCompile(Class, BlueprintView))
+	UMVVMViewClass* ViewExtension = nullptr;
+	bool bCompiled = false;
+	if (CurrentCompilerContext->PreCompile(Class, BlueprintView))
 	{
-		return;
+		ViewExtension = NewObject<UMVVMViewClass>(Class);
+		bCompiled = CurrentCompilerContext->Compile(Class, BlueprintView, ViewExtension);
 	}
 
-	UMVVMViewClass* ViewExtension = NewObject<UMVVMViewClass>(Class);
-	if (!CurrentCompilerContext->Compile(Class, BlueprintView, ViewExtension))
+	CurrentCompilerContext->CleanTemporaries(Class);
+	if (bCompiled)
 	{
-		return;
+		check(ViewExtension);
+		CurrentCompilerContext->AddExtension(Class, ViewExtension);
 	}
-
-	CurrentCompilerContext->AddExtension(Class, ViewExtension);
 }
diff --git a/Engine/Plugins/Runtime/ModelViewViewModel/Source/ModelViewViewModelBlueprint/Public/MVVMBlueprintViewModelContext.h b/Engine/Plugins/Runtime/ModelViewViewModel/Source/ModelViewViewModelBlueprint/Public/MVVMBlueprintViewModelContext.h
--- a/Engine/Plugins/Runtime/ModelViewViewModel/Source/ModelViewViewModelBlueprint/Public/MVVMBlueprintViewModelContext.h
+++ b/Engine/Plugins/Runtime/ModelViewViewModel/Source/ModelViewViewModelBlueprint/Public/MVVMBlueprintViewModelContext.h
@@ -15,9 +15,10 @@
 UENUM()
 enum class EMVVMBlueprintViewModelContextCreationType : uint8
 {
-	CreateInstance,
-	GlobalViewModelCollection,
-	PropertyPath,
+	Manual,	// The viewmodel will be assigned later.
+	CreateInstance, // A new instance of the viewmodel will be created when the widget is created.
+	GlobalViewModelCollection, // The viewmodel exists and is added to the MVVMSubsystem. It will be fetched there.
+	PropertyPath, // The viewmodel will be fetched by evaluating a function or a property path.
 };
 
 /**
diff --git a/Engine/Plugins/Runtime/ModelViewViewModel/Source/ModelViewViewModelEditor/Private/Widgets/SMVVMViewModelContextListWidget.cpp b/Engine/Plugins/Runtime/ModelViewViewModel/Source/ModelViewViewModelEditor/Private/Widgets/SMVVMViewModelContextListWidget.cpp
--- a/Engine/Plugins/Runtime/ModelViewViewModel/Source/ModelViewViewModelEditor/Private/Widgets/SMVVMViewModelContextListWidget.cpp
+++ b/Engine/Plugins/Runtime/ModelViewViewModel/Source/ModelViewViewModelEditor/Private/Widgets/SMVVMViewModelContextListWidget.cpp
@@ -327,6 +327,18 @@ namespace UE::MVVM::Private
 			FMenuBuilder MenuBuilder(true, NULL);
 			SMVVMManageViewModelsListEntryRow* Self = this;
 
+			MenuBuilder.AddMenuEntry(
+				FUIAction(
+					FExecuteAction::CreateLambda([Self]() 
+						{ 
+							Self->Entry->CreationType = EMVVMBlueprintViewModelContextCreationType::Manual;
+							Self->CreationTypeTextBlock->SetText(UEnum::GetDisplayValueAsText(Self->Entry->CreationType));
+							Self->UpdateGetterContainer();
+						})
+				),
+				SNew(STextBlock).Text(LOCTEXT("Manual", "Manual"))
+							);
+
 			MenuBuilder.AddMenuEntry(
 				FUIAction(
 					FExecuteAction::CreateLambda([Self]() 
