MVVM - Removed unnecessary SetBindingReference() overload.

#rnx

diff --git a/Engine/Plugins/Runtime/ModelViewViewModel/Source/ModelViewViewModelBlueprint/Public/MVVMPropertyPath.h b/Engine/Plugins/Runtime/ModelViewViewModel/Source/ModelViewViewModelBlueprint/Public/MVVMPropertyPath.h
--- a/Engine/Plugins/Runtime/ModelViewViewModel/Source/ModelViewViewModelBlueprint/Public/MVVMPropertyPath.h
+++ b/Engine/Plugins/Runtime/ModelViewViewModel/Source/ModelViewViewModelBlueprint/Public/MVVMPropertyPath.h
@@ -54,8 +54,14 @@ public:
 		return BindingReference;
 	}
 
-	void SetBindingReference(const UE::MVVM::FMVVMFieldVariant& InField)
+	void SetBindingReference(const UE::MVVM::FMVVMConstFieldVariant& InField)
 	{
+		if (InField.IsEmpty())
+		{
+			Reset();
+			return;
+		}
+
 		BindingReference = InField.CreateMemberReference();
 
 		if (InField.IsProperty())
@@ -72,31 +78,6 @@ public:
 		}
 	}
 
-	void SetBindingReference(const UE::MVVM::FMVVMConstFieldVariant& InField)
-	{
-		if (InField.IsEmpty())
-		{
-			Reset();
-		}
-		else
-		{
-			BindingReference = InField.CreateMemberReference();
-
-			if (InField.IsProperty())
-			{
-				BindingKind = EBindingKind::Property;
-			}
-			else if (InField.IsFunction())
-			{
-				BindingKind = EBindingKind::Function;
-			}
-			else
-			{
-				ensureAlwaysMsgf(false, TEXT("Binding to field of unknown type!"));
-			}
-		}
-	}
-
 	void Reset()
 	{
 		BindingReference = FMemberReference();
