SMVVMFieldSelector can now be used without a BindingMode set.

#rnx

diff --git a/Engine/Plugins/Runtime/ModelViewViewModel/Source/ModelViewViewModelEditor/Private/Widgets/SMVVMFieldSelector.cpp b/Engine/Plugins/Runtime/ModelViewViewModel/Source/ModelViewViewModelEditor/Private/Widgets/SMVVMFieldSelector.cpp
--- a/Engine/Plugins/Runtime/ModelViewViewModel/Source/ModelViewViewModelEditor/Private/Widgets/SMVVMFieldSelector.cpp
+++ b/Engine/Plugins/Runtime/ModelViewViewModel/Source/ModelViewViewModelEditor/Private/Widgets/SMVVMFieldSelector.cpp
@@ -80,8 +80,6 @@ void SMVVMFieldSelector::Construct(const FArguments& InArgs)
 	bIsSource = InArgs._IsSource;
 	TextStyle = InArgs._TextStyle;
 
-	check(BindingMode.IsSet());
-
 	FMVVMConstFieldVariant SelectedField = Private::GetSelectedFieldFromHelpers(PathHelpers);
 	Refresh();
 
@@ -207,13 +205,16 @@ TValueOrError<bool, FString> SMVVMFieldSelector::IsValidBindingForField(const FM
 		return MakeValue(true);
 	}
 
-	EMVVMBindingMode Mode = BindingMode.Get();
 	UMVVMSubsystem* Subsystem = GEngine->GetEngineSubsystem<UMVVMSubsystem>();
 
-	bool bToDestination =
-		Mode == EMVVMBindingMode::OneTimeToDestination ||
-		Mode == EMVVMBindingMode::OneWayToDestination;
+	if (!BindingMode.IsSet())
+	{
+		return MakeValue(true);
+	}
+
+	EMVVMBindingMode Mode = BindingMode.Get();
 
+	bool bToDestination = UE::MVVM::IsForwardBinding(Mode);
 	if (!bIsSource)
 	{
 		bToDestination = !bToDestination;
