MVVM: Use a native function to set a property on the viewmodel.

diff --git a/Engine/Plugins/Runtime/ModelViewViewModel/Source/ModelViewViewModel/Private/MVVMViewModelBase.cpp b/Engine/Plugins/Runtime/ModelViewViewModel/Source/ModelViewViewModel/Private/MVVMViewModelBase.cpp
--- a/Engine/Plugins/Runtime/ModelViewViewModel/Source/ModelViewViewModel/Private/MVVMViewModelBase.cpp
+++ b/Engine/Plugins/Runtime/ModelViewViewModel/Source/ModelViewViewModel/Private/MVVMViewModelBase.cpp
@@ -5,6 +5,7 @@
 #include "ViewModel/MVVMViewModelBlueprintGeneratedClass.h"
 #include "Bindings/MVVMBindingHelper.h"
 
+#define LOCTEXT_NAMESPACE "MVVMViewModelBase"
 
 FDelegateHandle UMVVMViewModelBase::AddFieldValueChangedDelegate(UE::FieldNotification::FFieldId InFieldId, FFieldValueChangedDelegate InNewDelegate)
 {
@@ -77,7 +78,16 @@ void UMVVMViewModelBase::FFieldNotificationClassDescriptor::ForEachField(const U
 }
 
 
-void UMVVMViewModelBase::BindingFieldValueChanged(UE::FieldNotification::FFieldId InFieldId)
+void UMVVMViewModelBase::NotifyFieldValudChanged(UE::FieldNotification::FFieldId InFieldId)
+{
+	// Execute replication call
+
+	// Execute Broadcast call
+	BroadcastFieldValueChanged(InFieldId);
+}
+
+
+void UMVVMViewModelBase::BroadcastFieldValueChanged(UE::FieldNotification::FFieldId InFieldId)
 {
 	if (InFieldId.IsValid() && EnabledFieldNotifications.IsValidIndex(InFieldId.GetIndex()) && EnabledFieldNotifications[InFieldId.GetIndex()])
 	{
@@ -89,5 +99,61 @@ void UMVVMViewModelBase::BindingFieldValueChanged(UE::FieldNotification::FFieldI
 void UMVVMViewModelBase::K2_BroadcastFieldValueChanged(FFieldNotificationId ViewModelPropertyName)
 {
 	UE::FieldNotification::FFieldId FieldId = GetFieldNotificationDescriptor().GetField(GetClass(), ViewModelPropertyName.GetFieldName());
-	BindingFieldValueChanged(FieldId);
+	BroadcastFieldValueChanged(FieldId);
+}
+
+
+DEFINE_FUNCTION(UMVVMViewModelBase::execK2_SetPropertyValue)
+{
+	Stack.StepCompiledIn<FProperty>(nullptr);
+	FProperty* TargetProperty = Stack.MostRecentProperty;
+	void* TargetValuePtr = Stack.MostRecentPropertyAddress;
+
+	Stack.StepCompiledIn<FProperty>(nullptr);
+	FProperty* SourceProperty = Stack.MostRecentProperty;
+	void* SourceValuePtr = Stack.MostRecentPropertyAddress;
+
+	P_FINISH;
+
+	UMVVMViewModelBase* ViewModelContext = Cast<UMVVMViewModelBase>(Context);
+
+	bool bResult = false;
+	if (ViewModelContext == nullptr || TargetProperty == nullptr || SourceProperty == nullptr)
+	{
+		const FBlueprintExceptionInfo ExceptionInfo(
+			EBlueprintExceptionType::AccessViolation,
+			LOCTEXT("MissingInputProperty", "Failed to resolve the input parameter for SetPropertyValue.")
+		);
+		FBlueprintCoreDelegates::ThrowScriptException(P_THIS, Stack, ExceptionInfo);
+	}
+	else
+	{
+		P_NATIVE_BEGIN
+
+		UE::FieldNotification::FFieldId FieldId = ViewModelContext->GetFieldNotificationDescriptor().GetField(ViewModelContext->GetClass(), TargetProperty->GetFName());
+		if (!FieldId.IsValid())
+		{
+			const FBlueprintExceptionInfo ExceptionInfo(
+				EBlueprintExceptionType::AccessViolation,
+				LOCTEXT("Bad FieldId", "Failed to find the FieldId that correspond to the set value.")
+			);
+			FBlueprintCoreDelegates::ThrowScriptException(P_THIS, Stack, ExceptionInfo);
+		}
+		else
+		{
+			bResult = TargetProperty->Identical(TargetValuePtr, SourceValuePtr);
+			if (!bResult)
+			{
+				// Set the value then notify that the value changed.
+				TargetProperty->SetValue_InContainer(ViewModelContext, SourceValuePtr);
+				ViewModelContext->NotifyFieldValudChanged(FieldId);
+			}
+		}
+
+		P_NATIVE_END
+	}
+
+	*StaticCast<bool*>(RESULT_PARAM) = (!bResult);
 }
+
+#undef LOCTEXT_NAMESPACE
diff --git a/Engine/Plugins/Runtime/ModelViewViewModel/Source/ModelViewViewModel/Public/MVVMViewModelBase.h b/Engine/Plugins/Runtime/ModelViewViewModel/Source/ModelViewViewModel/Public/MVVMViewModelBase.h
--- a/Engine/Plugins/Runtime/ModelViewViewModel/Source/ModelViewViewModel/Public/MVVMViewModelBase.h
+++ b/Engine/Plugins/Runtime/ModelViewViewModel/Source/ModelViewViewModel/Public/MVVMViewModelBase.h
@@ -14,33 +14,20 @@
 
 #include "MVVMViewModelBase.generated.h"
 
-//UENUM()
-//enum class EMVVMNGetterType : uint8
-//{
-//	/** The Source Property has a Delegate associated with it that will notify us when the value change. */
-//	Delegate,
-//	/** The value is fetch when requested. */
-//	FetchOnRequest,
-//	/** The value is fetch every frame and compared. */ 
-//	FetchOnTickAndCompared,
-//	/** The value is fetch every frame and is considered changed. */
-//	FetchOnTick,
-//};
-
-///** The user will notify the ViewModel when the value change via the PropertyValueChanged event. */
-//Event,
-//Interface,		// The Source Context implement the MVVMModelInterface.
-					//Use the interface to register to the SourceContext's OnPropertyValueChanged.
-//PushModel,		// The property has a NetId generated by UBT. Used to notify us when the value changed.
-
+/** After a field value changed. Execute all the notification (replication and broadcast). */
 #define UE_MVVM_NOTIFY_FIELD_VALUE_CHANGED(MemberName) \
-	BindingFieldValueChanged(ThisClass::FFieldNotificationClassDescriptor::MemberName)
+	NotifyFieldValudChanged(ThisClass::FFieldNotificationClassDescriptor::MemberName)
+
+/** After a field value changed. Broadcast the event (doesn't execute the replication code). */
+#define UE_MVVM_BROADCAST_FIELD_VALUE_CHANGED(MemberName) \
+	BroadcastFieldValueChanged(ThisClass::FFieldNotificationClassDescriptor::MemberName)
 
+/** Set the new value and notify if the property value changed. */
 #define UE_MVVM_SET_PROPERTY_VALUE(MemberName, NewValue) \
 	SetPropertyValue(MemberName, NewValue, ThisClass::FFieldNotificationClassDescriptor::MemberName)
 
 
-/** Sub class of this will be generated from the BP to cached the source data. */
+/** Base class for MVVM viewmodel. */
 UCLASS(Blueprintable, Abstract, DisplayName="MVVM ViewModel")
 class MODELVIEWVIEWMODEL_API UMVVMViewModelBase : public UObject, public INotifyFieldValueChanged
 {
@@ -61,12 +48,23 @@ public:
 	virtual const UE::FieldNotification::IClassDescriptor& GetFieldNotificationDescriptor() const override;
 	//~ End INotifyFieldValueChanged Interface
 
+
+protected:
+	/** Execute all the notification (replication and broadcast). */
+	void NotifyFieldValudChanged(UE::FieldNotification::FFieldId InFieldId);
+
 protected:
-	UFUNCTION(BlueprintCallable, Category="FieldNotify", meta=(DisplayName="Broadcast Field Value Changed", ScriptName="BroadcastFieldValueChanged"))
+	UFUNCTION(BlueprintCallable, Category="FieldNotify", meta=(DisplayName="Broadcast Field Value Changed", ScriptName="BroadcastFieldValueChanged", BlueprintInternalUseOnly="true"))
 	void K2_BroadcastFieldValueChanged(FFieldNotificationId FieldId);
-	void BindingFieldValueChanged(UE::FieldNotification::FFieldId InFieldId);
+
+	/** Broadcast the event (doesn't execute the replication code). */
+	void BroadcastFieldValueChanged(UE::FieldNotification::FFieldId InFieldId);
 
 protected:
+	UFUNCTION(BlueprintCallable, CustomThunk, Category="Viewmodel", meta=(CustomStructureParam="OldValue,NewValue", ScriptName="SetPropertyValue", BlueprintInternalUseOnly="true"))
+	bool K2_SetPropertyValue(UPARAM(ref) const int32& OldValue, UPARAM(ref) const int32& NewValue);
+
+	/** Set the new value and notify if the property value changed. */
 	template<typename T>
 	bool SetPropertyValue(T& Value, const T& NewValue, UE::FieldNotification::FFieldId FieldId)
 	{
@@ -76,10 +74,13 @@ protected:
 		}
 
 		Value = NewValue;
-		BindingFieldValueChanged(FieldId);
+		NotifyFieldValudChanged(FieldId);
 		return true;
 	}
 
+private:
+	DECLARE_FUNCTION(execK2_SetPropertyValue);
+
 private:
 	UE::FieldNotification::FFieldMulticastDelegate Delegates;
 	TBitArray<> EnabledFieldNotifications;
diff --git a/Engine/Plugins/Runtime/ModelViewViewModel/Source/ModelViewViewModelBlueprint/Private/MVVMFunctionGraphHelper.cpp b/Engine/Plugins/Runtime/ModelViewViewModel/Source/ModelViewViewModelBlueprint/Private/MVVMFunctionGraphHelper.cpp
--- a/Engine/Plugins/Runtime/ModelViewViewModel/Source/ModelViewViewModelBlueprint/Private/MVVMFunctionGraphHelper.cpp
+++ b/Engine/Plugins/Runtime/ModelViewViewModel/Source/ModelViewViewModelBlueprint/Private/MVVMFunctionGraphHelper.cpp
@@ -235,163 +235,49 @@ bool GenerateViewModelFieldNotifySetter(FKismetCompilerContext& InContext, UEdGr
 		return false;
 	}
 
-	UK2Node_VariableSet* SetVariableNode = InContext.SpawnIntermediateNode<UK2Node_VariableSet>(FunctionEntry, InFunctionGraph);
+	UK2Node_VariableGet* OldValueGetNode = InContext.SpawnIntermediateNode<UK2Node_VariableGet>(FunctionEntry, InFunctionGraph);
 	{
-		//SetVariableNode->SetFromProperty(InProperty, true, nullptr);
-		SetVariableNode->VariableReference.SetSelfMember(InProperty->GetFName());
-		SetVariableNode->AllocateDefaultPins();
+		bool bSelfContext = true;
+		OldValueGetNode->VariableReference.SetFromField<FProperty>(InProperty, bSelfContext);
+		OldValueGetNode->AllocateDefaultPins();
 	}
 
-	UK2Node_CallFunction* CallBroadcast = InContext.SpawnIntermediateNode<UK2Node_CallFunction>(FunctionEntry, InFunctionGraph);
-	{
-		CallBroadcast->FunctionReference.SetExternalMember("K2_BroadcastFieldValueChanged", UMVVMViewModelBase::StaticClass());
-		CallBroadcast->AllocateDefaultPins();
-
-		UEdGraphPin* NamePin = CallBroadcast->FindPin(TEXT("FieldId"), EGPD_Input);
-		ensure(NamePin);
-		if (!NamePin)
-		{
-			return false;
-		}
-
-		FFieldNotificationId NewPropertyNotifyId;
-		NewPropertyNotifyId.FieldName = InProperty->GetFName();
-
-		FString ValueString;
-		FFieldNotificationId::StaticStruct()->ExportText(ValueString, &NewPropertyNotifyId, nullptr, nullptr, EPropertyPortFlags::PPF_None, nullptr);
-
-		bool bMarkAsModified = false;
-		Schema->TrySetDefaultValue(*NamePin, ValueString, bMarkAsModified);
-	}
-
-	UK2Node_VariableGet* GetVariableNode = InContext.SpawnIntermediateNode<UK2Node_VariableGet>(FunctionEntry, InFunctionGraph);
+	UK2Node_CallFunction* CallSetPropertyValue = InContext.SpawnIntermediateNode<UK2Node_CallFunction>(FunctionEntry, InFunctionGraph);
 	{
-		//GetVariableNode->SetFromProperty(InProperty, true, nullptr);
-		GetVariableNode->VariableReference.SetSelfMember(InProperty->GetFName());
-		GetVariableNode->AllocateDefaultPins();
-	}
-
-	UK2Node_IfThenElse* BranchNode = InContext.SpawnIntermediateNode<UK2Node_IfThenElse>(FunctionEntry, InFunctionGraph);
-	{
-		BranchNode->AllocateDefaultPins();
-	}
-
-	UFunction* CompareFunction = nullptr;
-	{
-		TArray<UEdGraphPin*> PromotableOperatorPinsToConsider;
-		UEdGraphPin* EntryValuePin = FunctionEntry->FindPin(InInputPinName, EGPD_Output);
-		if (ensure(EntryValuePin))
-		{
-			PromotableOperatorPinsToConsider.Add(EntryValuePin);
-		}
-		UEdGraphPin* GetVariableNode_Property = GetVariableNode->FindPin(InProperty->GetFName(), EGPD_Output);
-		if (ensure(GetVariableNode_Property))
-		{
-			PromotableOperatorPinsToConsider.Add(GetVariableNode_Property);
-		}
-		UEdGraphPin* BranchCondition = BranchNode->FindPin(UEdGraphSchema_K2::PN_Condition, EGPD_Input);
-		if (ensure(BranchCondition))
-		{
-			PromotableOperatorPinsToConsider.Add(BranchCondition);
-		}
-
-		CompareFunction = FTypePromotion::FindBestMatchingFunc(TEXT("NotEqual"), PromotableOperatorPinsToConsider);
+		CallSetPropertyValue->FunctionReference.SetExternalMember("K2_SetPropertyValue", UMVVMViewModelBase::StaticClass());
+		CallSetPropertyValue->AllocateDefaultPins();
 	}
 
 	bool bResult = true;
-
-	bool bDoBranchTest = CompareFunction != nullptr;
-	if (bDoBranchTest)
+	// Entry -> SetProperty
 	{
-		UK2Node_PromotableOperator* CompareOperator = InContext.SpawnIntermediateNode<UK2Node_PromotableOperator>(FunctionEntry, InFunctionGraph);
-		{
-			CompareOperator->SetFromFunction(CompareFunction);
-			CompareOperator->AllocateDefaultPins();
-		}
-
-		// Entry -> Branch
-		{
-			UEdGraphPin* ThenPin = FunctionEntry->FindPin(UEdGraphSchema_K2::PN_Then, EGPD_Output);
-			UEdGraphPin* ExecPin = BranchNode->FindPin(UEdGraphSchema_K2::PN_Execute, EGPD_Input);
-			if (ensure(ThenPin && ExecPin))
-			{
-				ensure(Schema->TryCreateConnection(ThenPin, ExecPin));
-			}
-			bResult = bResult && ThenPin && ExecPin;
-		}
-
-		// Entry -> Compare, Get -> Compare, Compare -> Branch
-		{
-			UEdGraphPin* EntryValuePin = FunctionEntry->FindPin(InInputPinName, EGPD_Output);
-			UEdGraphPin* TopInputPin = CompareOperator->FindPin(TEXT("A"), EGPD_Input);
-			if (ensure(EntryValuePin && TopInputPin))
-			{
-				ensure(Schema->TryCreateConnection(EntryValuePin, TopInputPin));
-			}
-			bResult = bResult && EntryValuePin && TopInputPin;
-
-			UEdGraphPin* GetPin = GetVariableNode->FindPin(InProperty->GetFName(), EGPD_Output);
-			UEdGraphPin* BottomInputPin = CompareOperator->FindPin(TEXT("B"), EGPD_Input);
-			if (ensure(GetPin && BottomInputPin))
-			{
-				ensure(Schema->TryCreateConnection(GetPin, BottomInputPin));
-			}
-			bResult = bResult && GetPin && BottomInputPin;
-
-			UEdGraphPin* OutputPin = CompareOperator->GetOutputPin();
-			UEdGraphPin* BranchCondition = BranchNode->FindPin(UEdGraphSchema_K2::PN_Condition, EGPD_Input);
-			if (ensure(BranchCondition && OutputPin))
-			{
-				ensure(Schema->TryCreateConnection(OutputPin, BranchCondition));
-			}
-			bResult = bResult && BranchCondition && OutputPin;
-		}
-
-		// Branch -> Set
+		UEdGraphPin* ThenPin = FunctionEntry->FindPin(UEdGraphSchema_K2::PN_Then, EGPD_Output);
+		UEdGraphPin* ExecPin = CallSetPropertyValue->FindPin(UEdGraphSchema_K2::PN_Execute, EGPD_Input);
+		if (ensure(ThenPin && ExecPin))
 		{
-			UEdGraphPin* ThenPin = BranchNode->FindPin(UEdGraphSchema_K2::PN_Then, EGPD_Output);
-			UEdGraphPin* ExecPin = SetVariableNode->FindPin(UEdGraphSchema_K2::PN_Execute, EGPD_Input);
-			if (ensure(ThenPin && ExecPin))
-			{
-				ensure(Schema->TryCreateConnection(ThenPin, ExecPin));
-			}
-			bResult = bResult && ThenPin && ExecPin;
+			ensure(Schema->TryCreateConnection(ThenPin, ExecPin));
 		}
+		bResult = bResult && ThenPin && ExecPin;
 	}
-	else
+	// OldValue -> SetProperty(OldValue)
 	{
-		// Entry -> Set
+		UEdGraphPin* ViewModelValuePin = OldValueGetNode->FindPin(InProperty->GetFName(), EGPD_Output);
+		UEdGraphPin* OldValuePin = CallSetPropertyValue->FindPin(TEXT("OldValue"), EGPD_Input);
+		if (ensure(ViewModelValuePin && OldValuePin))
 		{
-			UEdGraphPin* ThenPin = FunctionEntry->FindPin(UEdGraphSchema_K2::PN_Then, EGPD_Output);
-			UEdGraphPin* ExecPin = SetVariableNode->FindPin(UEdGraphSchema_K2::PN_Execute, EGPD_Input);
-			if (ensure(ThenPin && ExecPin))
-			{
-				ensure(Schema->TryCreateConnection(ThenPin, ExecPin));
-			}
-			bResult = bResult && ThenPin && ExecPin;
+			ensure(Schema->TryCreateConnection(ViewModelValuePin, OldValuePin));
 		}
+		bResult = bResult && ViewModelValuePin && OldValuePin;
 	}
-
-	// Entry -> Set {Set value}
+	// NewValue -> SetProperty(NewValue)
 	{
 		UEdGraphPin* EntryValuePin = FunctionEntry->FindPin(InInputPinName, EGPD_Output);
-		UEdGraphPin* SetValuePin = SetVariableNode->FindPin(InProperty->GetFName(), EGPD_Input);
-		if (ensure(EntryValuePin && SetValuePin))
+		UEdGraphPin* NewValuePin = CallSetPropertyValue->FindPin(TEXT("NewValue"), EGPD_Input);
+		if (ensure(EntryValuePin && NewValuePin))
 		{
-			ensure(Schema->TryCreateConnection(EntryValuePin, SetValuePin));
+			ensure(Schema->TryCreateConnection(EntryValuePin, NewValuePin));
 		}
-		bResult = bResult && EntryValuePin && SetValuePin;
-	}
-
-	// Set -> Broadcast
-	{
-		UEdGraphPin* ThenPin = SetVariableNode->FindPin(UEdGraphSchema_K2::PN_Then, EGPD_Output);
-		UEdGraphPin* ExecPin = CallBroadcast->FindPin(UEdGraphSchema_K2::PN_Execute, EGPD_Input);
-		if (ensure(ThenPin && ExecPin))
-		{
-			ensure(Schema->TryCreateConnection(ThenPin, ExecPin));
-		}
-		bResult = bResult && ThenPin && ExecPin;
+		bResult = bResult && EntryValuePin && NewValuePin;
 	}
 
 	return bResult;
diff --git a/Engine/Plugins/Runtime/ModelViewViewModel/Source/ModelViewViewModelBlueprint/Private/ViewModel/MVVMViewModelBlueprintCompiler.cpp b/Engine/Plugins/Runtime/ModelViewViewModel/Source/ModelViewViewModelBlueprint/Private/ViewModel/MVVMViewModelBlueprintCompiler.cpp
--- a/Engine/Plugins/Runtime/ModelViewViewModel/Source/ModelViewViewModelBlueprint/Private/ViewModel/MVVMViewModelBlueprintCompiler.cpp
+++ b/Engine/Plugins/Runtime/ModelViewViewModel/Source/ModelViewViewModelBlueprint/Private/ViewModel/MVVMViewModelBlueprintCompiler.cpp
@@ -79,8 +79,7 @@ FProperty* FViewModelBlueprintCompilerContext::FindPropertyByNameOnNewClass(FNam
 	{
 		if (CurrentField->GetFName() == PropertyName)
 		{
-			return  CastField<FProperty>(CurrentField);
-			break;
+			return CastField<FProperty>(CurrentField);
 		}
 	}
 	return nullptr;
@@ -131,22 +130,9 @@ void FViewModelBlueprintCompilerContext::CreateFunctionList()
 				continue;
 			}
 		}
-
-		// HACK: Since the VariableSetter doesn't work because of the BlueprintReadOnly, clear it before processing the functions
-		Function.SkelProperty->ClearPropertyFlags(CPF_BlueprintReadOnly);
 	}
 
 	Super::CreateFunctionList();
-
-	// HACK: Reset the BlueprintReadOnly flag
-	for (const FGeneratedFunction& Function : GeneratedFunctions)
-	{
-		if (Function.SkelProperty == nullptr)
-		{
-			continue;
-		}
-		Function.SkelProperty->SetPropertyFlags(CPF_BlueprintReadOnly);
-	}
 }
 
 
@@ -167,38 +153,33 @@ void FViewModelBlueprintCompilerContext::CreateClassVariablesFromBlueprint()
 		{
 			Property->SetPropertyFlags(CPF_BlueprintReadOnly | CPF_DisableEditOnInstance);
 
-			FGeneratedFunction* FoundGeneratedFunction = GeneratedFunctions.FindByPredicate([PropertyName](const FGeneratedFunction& Other)
-				{
-					return Other.PropertyName == PropertyName;
-				});
+			const FString SetterName = FString::Printf(TEXT("Set%s"), *Property->GetName());
+			const FString NetRepName = FString::Printf(TEXT("__RepNotify_%s"), *Property->GetName());
 
-
-			if (FoundGeneratedFunction)
+#if WITH_EDITOR
+			if (bGenerateSetterFunction)
 			{
-				if (CompileOptions.CompileType == EKismetCompileType::SkeletonOnly)
-				{
-					FoundGeneratedFunction->SkelProperty = Property;
-					ensureMsgf(false, TEXT("The generated function should already exist"));
-				}
-				else
-				{
-					FoundGeneratedFunction->Property = Property;
-					if (FoundGeneratedFunction->NetRepFunction)
-					{
-						Property->SetPropertyFlags(CPF_Net);
-						Property->RepNotifyFunc = FoundGeneratedFunction->NetRepFunction->GetFName();
-					}
-				}
+				Property->SetMetaData(FBlueprintMetadata::MD_PropertySetFunction, *SetterName);
 			}
-			else if (CompileOptions.CompileType == EKismetCompileType::SkeletonOnly)
+#endif
+
+			if (CompileOptions.CompileType == EKismetCompileType::SkeletonOnly)
 			{
 				FGeneratedFunction GeneratedFunction;
 				GeneratedFunction.PropertyName = PropertyName;
 				GeneratedFunction.SkelProperty = Property;
 
+				if (GeneratedFunctions.ContainsByPredicate([PropertyName](const FGeneratedFunction& Other)
+					{
+						return Other.PropertyName == PropertyName;
+					}))
+				{
+					MessageLog.Error(*FString::Printf(TEXT("Internal error. The functions are already generated for properpty."), *PropertyName.ToString()));
+					continue;
+				}
+
 				if (bGenerateSetterFunction)
 				{
-					const FString SetterName = FString::Printf(TEXT("Set%s"), *Property->GetName());
 					GeneratedFunction.SetterFunction = FunctionGraphHelper::CreateIntermediateFunctionGraph(
 						*this
 						, *SetterName
@@ -217,7 +198,6 @@ void FViewModelBlueprintCompilerContext::CreateClassVariablesFromBlueprint()
 
 				if (bGenerateReplicatedFunction)
 				{
-					const FString NetRepName = FString::Printf(TEXT("__RepNotify_%s"), *Property->GetName());
 					GeneratedFunction.NetRepFunction = FunctionGraphHelper::CreateIntermediateFunctionGraph(
 						*this
 						, *NetRepName
@@ -271,19 +251,6 @@ void FViewModelBlueprintCompilerContext::OnNewClassSet(UBlueprintGeneratedClass*
 }
 
 
-void FViewModelBlueprintCompilerContext::PrecompileFunction(FKismetFunctionContext& Context, EInternalCompilerFlags InternalFlags)
-{
-	// HACK: Since the VariableSet doesn't work because of the BlueprintReadOnly, clear it before processing the functions
-	for (const FGeneratedFunction& Function : GeneratedFunctions)
-	{
-		check(Function.Property);
-		Function.Property->ClearPropertyFlags(CPF_BlueprintReadOnly);
-	}
-
-	Super::PrecompileFunction(Context, InternalFlags);
-}
-
-
 void FViewModelBlueprintCompilerContext::PostcompileFunction(FKismetFunctionContext& Context)
 {
 	Super::PostcompileFunction(Context);
@@ -296,13 +263,6 @@ void FViewModelBlueprintCompilerContext::PostcompileFunction(FKismetFunctionCont
 		}
 		NewViewModelBlueprintClass->FieldNotifyNames.Emplace(Context.Function->GetFName());
 	}
-
-	// HACK: Reset the BlueprintReadOnly flag
-	for (const FGeneratedFunction& Function : GeneratedFunctions)
-	{
-		check(Function.Property);
-		Function.Property->SetPropertyFlags(CPF_BlueprintReadOnly);
-	}
 }
 
 
@@ -320,6 +280,22 @@ void FViewModelBlueprintCompilerContext::FinishCompilingClass(UClass* Class)
 				BPGClass->InitializeFieldNotification(ViewModelBase);
 			}
 		}
+
+		auto DeleteGraph = [](UEdGraph* Graph)
+		{
+			if (Graph)
+			{
+				ensureMsgf(Graph->HasAnyFlags(RF_Transient), TEXT("The graph should be temporary and should be generated automatically."));
+				// GC may not have clean the graph (GC doesn't run when bRegenerateSkeletonOnly is on)
+				Graph->Rename(nullptr, Graph->GetOuter(), REN_DoNotDirty | REN_ForceNoResetLoaders);
+			}
+		};
+
+		for (FGeneratedFunction& GeneratedFunction : GeneratedFunctions)
+		{
+			DeleteGraph(GeneratedFunction.SetterFunction);
+			DeleteGraph(GeneratedFunction.NetRepFunction);
+		}
 	}
 }
 
diff --git a/Engine/Plugins/Runtime/ModelViewViewModel/Source/ModelViewViewModelBlueprint/Private/ViewModel/MVVMViewModelBlueprintCompiler.h b/Engine/Plugins/Runtime/ModelViewViewModel/Source/ModelViewViewModelBlueprint/Private/ViewModel/MVVMViewModelBlueprintCompiler.h
--- a/Engine/Plugins/Runtime/ModelViewViewModel/Source/ModelViewViewModelBlueprint/Private/ViewModel/MVVMViewModelBlueprintCompiler.h
+++ b/Engine/Plugins/Runtime/ModelViewViewModel/Source/ModelViewViewModelBlueprint/Private/ViewModel/MVVMViewModelBlueprintCompiler.h
@@ -54,7 +54,6 @@ private:
 	virtual void CreateClassVariablesFromBlueprint() override;
 	virtual void SpawnNewClass(const FString& NewClassName) override;
 	virtual void OnNewClassSet(UBlueprintGeneratedClass* ClassToUse) override;
-	virtual void PrecompileFunction(FKismetFunctionContext& Context, EInternalCompilerFlags InternalFlags) override;
 	virtual void PostcompileFunction(FKismetFunctionContext& Context) override;
 	virtual void FinishCompilingClass(UClass* Class) override;
 	//~ End FKismetCompilerContext Interface
