MVVM - Conversion function selection now filters the list to just UBlueprintFunctionLibraries and the UUserWidget itself.

Conversion functions can now have multiple arguments as long as one of them is compatible.

#rnx

diff --git a/Engine/Plugins/Runtime/ModelViewViewModel/Source/ModelViewViewModel/Private/Bindings/MVVMBindingHelper.cpp b/Engine/Plugins/Runtime/ModelViewViewModel/Source/ModelViewViewModel/Private/Bindings/MVVMBindingHelper.cpp
--- a/Engine/Plugins/Runtime/ModelViewViewModel/Source/ModelViewViewModel/Private/Bindings/MVVMBindingHelper.cpp
+++ b/Engine/Plugins/Runtime/ModelViewViewModel/Source/ModelViewViewModel/Private/Bindings/MVVMBindingHelper.cpp
@@ -76,10 +76,11 @@ namespace UE::MVVM::BindingHelper
 		// A setter can return a value "bool SetValue(int32)" but is not valid for MVVM
 		//&& ((InFunction->NumParms == 1 && GetReturnProperty() == nullptr) || (InFunction->NumParms == 2 && GetReturnProperty() != nullptr));
 
+		TValueOrError<TArray<const FProperty*>, FString> Result = TryGetArgumentsForConversionFunction(InFunction);
+
 		return Private::IsValidCommon(InFunction)
 			&& !InFunction->HasAnyFunctionFlags(FUNC_Const | FUNC_BlueprintPure)
-			&& InFunction->NumParms == 1
-			&& GetFirstArgumentProperty(InFunction) != nullptr;
+			&& Result.HasValue() && Result.GetValue().Num() > 0;
 	}
 
 
@@ -372,8 +373,7 @@ namespace UE::MVVM::BindingHelper
 			return TryGetPropertyTypeForDestinationBinding(InField.GetFunction());
 		}
 	}
-
-	TValueOrError<FConversionFunctionArguments, FString> TryGetPropertyTypeForConversionFunction(const UFunction* InFunction)
+	TValueOrError<const FProperty*, FString> TryGetReturnTypeForConversionFunction(const UFunction* InFunction)
 	{
 		FString CommonResult = Private::TryGetPropertyTypeCommon(InFunction);
 		if (!CommonResult.IsEmpty())
@@ -386,27 +386,62 @@ namespace UE::MVVM::BindingHelper
 			return MakeError(FString::Printf(TEXT("The function '%s' is not static or is not const and pure."), *InFunction->GetName()));
 		}
 
-		if (InFunction->NumParms != 2)
+		if (InFunction->NumParms < 2)
 		{
 			return MakeError(FString::Printf(TEXT("The function '%s' does not have the correct number of arguments."), *InFunction->GetName()));
 		}
 
-		FConversionFunctionArguments Result;
-		Result.ReturnProperty = GetReturnProperty(InFunction);
-		if (Result.ReturnProperty == nullptr)
+		const FProperty* ReturnProperty = GetReturnProperty(InFunction);
+		if (ReturnProperty == nullptr)
 		{
 			return MakeError(FString::Printf(TEXT("The return value for function '%s' is invalid."), *InFunction->GetName()));
 		}
 
-		Result.ArgumentProperty = GetFirstArgumentProperty(InFunction);
-		if (Result.ArgumentProperty == nullptr)
+		return MakeValue(ReturnProperty);
+	}
+
+	TValueOrError<TArray<const FProperty*>, FString> TryGetArgumentsForConversionFunction(const UFunction* InFunction)
+	{
+		FString CommonResult = Private::TryGetPropertyTypeCommon(InFunction);
+		if (!CommonResult.IsEmpty())
 		{
-			return MakeError(FString::Printf(TEXT("The function '%s' doesn't not have a valid argument."), *InFunction->GetName()));
+			return MakeError(MoveTemp(CommonResult));
 		}
 
-		return MakeValue(Result);
-	}
+		if (!InFunction->HasAllFunctionFlags(FUNC_Static) && !InFunction->HasAnyFunctionFlags(FUNC_Const | FUNC_BlueprintPure))
+		{
+			return MakeError(FString::Printf(TEXT("The function '%s' is not static or is not const and pure."), *InFunction->GetName()));
+		}
+
+		if (InFunction->NumParms < 2)
+		{
+			return MakeError(FString::Printf(TEXT("The function '%s' does not have the correct number of arguments."), *InFunction->GetName()));
+		}
 
+		const FProperty* ReturnProperty = GetReturnProperty(InFunction);
+		if (ReturnProperty == nullptr)
+		{
+			return MakeError(FString::Printf(TEXT("The return value for function '%s' is invalid."), *InFunction->GetName()));
+		}
+
+		TArray<const FProperty*> Arguments;
+		for (TFieldIterator<FProperty> It(InFunction); It && (It->PropertyFlags & CPF_Parm); ++It)
+		{
+			const FProperty* Property = *It;
+			if (It->HasAllPropertyFlags(CPF_Parm) && !It->HasAnyPropertyFlags(CPF_Deprecated | CPF_EditorOnly | CPF_ReturnParm))
+			{
+				// argument is either an out parameter or not const
+				if (It->HasAllPropertyFlags(CPF_OutParm) && !It->HasAllPropertyFlags(CPF_ConstParm))
+				{
+					continue;
+				}
+
+				Arguments.Add(*It);
+			}
+		}
+
+		return MakeValue(Arguments);
+	}
 
 	namespace Private
 	{
diff --git a/Engine/Plugins/Runtime/ModelViewViewModel/Source/ModelViewViewModel/Private/MVVMSubsystem.cpp b/Engine/Plugins/Runtime/ModelViewViewModel/Source/ModelViewViewModel/Private/MVVMSubsystem.cpp
--- a/Engine/Plugins/Runtime/ModelViewViewModel/Source/ModelViewViewModel/Private/MVVMSubsystem.cpp
+++ b/Engine/Plugins/Runtime/ModelViewViewModel/Source/ModelViewViewModel/Private/MVVMSubsystem.cpp
@@ -128,18 +128,15 @@ namespace UE::MVVM::Private
 		TArray<FMVVMAvailableBinding> Result;
 		Result.Reserve(FieldDescriptors.Num());
 
-		for (TFieldIterator<const UFunction> FunctionItt(Class, EFieldIteratorFlags::IncludeSuper); FunctionItt; ++FunctionItt)
+		for (TFieldIterator<const UFunction> FunctionIt(Class, EFieldIteratorFlags::IncludeSuper); FunctionIt; ++FunctionIt)
 		{
-			const UFunction* Function = *FunctionItt;
+			const UFunction* Function = *FunctionIt;
 			check(Function);
-			if (!Function->HasAnyFunctionFlags(FUNC_Private | FUNC_Protected))
+			const bool bIsReadable = FieldDescriptors.Contains(Function->GetFName()) && BindingHelper::IsValidForSourceBinding(Function);
+			const bool bIsWritable = BindingHelper::IsValidForDestinationBinding(Function);
+			if (bIsReadable || bIsWritable)
 			{
-				const bool bIsReadable = FieldDescriptors.Contains(Function->GetFName()) && BindingHelper::IsValidForSourceBinding(Function);
-				const bool bIsWritable = BindingHelper::IsValidForDestinationBinding(Function);
-				if (bIsReadable || bIsWritable)
-				{
-					Result.Add(FMVVMAvailableBinding(FMVVMBindingName(Function->GetFName()), bIsReadable, bIsWritable));
-				}
+				Result.Add(FMVVMAvailableBinding(FMVVMBindingName(Function->GetFName()), bIsReadable, bIsWritable));
 			}
 		}
 
@@ -247,26 +244,26 @@ TValueOrError<bool, FString> UMVVMSubsystem::IsBindingValid(FDirectionalBindingA
 TValueOrError<bool, FString> UMVVMSubsystem::IsBindingValid(FConstDirectionalBindingArgs Args) const
 {
 	// Test Source
-	TValueOrError<const FProperty*, FString> GetterProperty = UE::MVVM::BindingHelper::TryGetPropertyTypeForSourceBinding(Args.SourceBinding);
-	if (GetterProperty.HasError())
+	TValueOrError<const FProperty*, FString> SourceResult = UE::MVVM::BindingHelper::TryGetPropertyTypeForSourceBinding(Args.SourceBinding);
+	if (SourceResult.HasError())
 	{
-		return MakeError(GetterProperty.StealError());
+		return MakeError(SourceResult.StealError());
 	}
 
-	const FProperty* SourceProperty = GetterProperty.StealValue();
+	const FProperty* SourceProperty = SourceResult.StealValue();
 	if (SourceProperty == nullptr)
 	{
 		return MakeError(TEXT("There is no value to bind at the source."));
 	}
 
 	// Test Destination
-	TValueOrError<const FProperty*, FString> SetterProperty = UE::MVVM::BindingHelper::TryGetPropertyTypeForDestinationBinding(Args.DestinationBinding);
-	if (SetterProperty.HasError())
+	TValueOrError<const FProperty*, FString> DestinationResult = UE::MVVM::BindingHelper::TryGetPropertyTypeForDestinationBinding(Args.DestinationBinding);
+	if (DestinationResult.HasError())
 	{
-		return MakeError(SetterProperty.StealError());
+		return MakeError(DestinationResult.StealError());
 	}
 
-	const FProperty* DestinationProperty = SetterProperty.StealValue();
+	const FProperty* DestinationProperty = DestinationResult.StealValue();
 	if (DestinationProperty == nullptr)
 	{
 		return MakeError(TEXT("There is no value to bind at the destination."));
@@ -275,23 +272,42 @@ TValueOrError<bool, FString> UMVVMSubsystem::IsBindingValid(FConstDirectionalBin
 	// Test the conversion function
 	if (Args.ConversionFunction)
 	{
-		TValueOrError<UE::MVVM::BindingHelper::FConversionFunctionArguments, FString> ConversionResult = UE::MVVM::BindingHelper::TryGetPropertyTypeForConversionFunction(Args.ConversionFunction);
-		if (ConversionResult.HasError())
+		TValueOrError<const FProperty*, FString> ReturnResult = UE::MVVM::BindingHelper::TryGetReturnTypeForConversionFunction(Args.ConversionFunction);
+		if (ReturnResult.HasError())
 		{
-			return MakeError(ConversionResult.StealError());
+			return MakeError(ReturnResult.StealError());
+		}
+
+		TValueOrError<TArray<const FProperty*>, FString> ArgumentsResult = UE::MVVM::BindingHelper::TryGetArgumentsForConversionFunction(Args.ConversionFunction);
+		if (ArgumentsResult.HasError())
+		{
+			return MakeError(ArgumentsResult.StealError());
 		}
 
 		// The compiled version should look like Setter(Conversion(Getter())).
-		const FProperty* ConversionReturnProperty = ConversionResult.GetValue().ReturnProperty;
-		if (!UE::MVVM::BindingHelper::ArePropertiesCompatible(ConversionReturnProperty, DestinationProperty))
+		const FProperty* ReturnProperty = ReturnResult.GetValue();
+		if (!UE::MVVM::BindingHelper::ArePropertiesCompatible(ReturnProperty, DestinationProperty))
+		{
+			return MakeError(FString::Printf(TEXT("The destination property '%s' (%s) does not match the return type of the conversion function (%s)."), *DestinationProperty->GetName(), *DestinationProperty->GetCPPType(), *ReturnProperty->GetCPPType()));
+		}
+
+		bool bAnyCompatible = false;
+
+		const TArray<const FProperty*>& ArgumentProperties = ArgumentsResult.GetValue();
+		for (const FProperty* ArgumentProperty : ArgumentProperties)
 		{
-			return MakeError(FString::Printf(TEXT("The destination property '%s' (%s) does not match the return type of the conversion function (%s)."), *DestinationProperty->GetName(), *DestinationProperty->GetCPPType(), *ConversionReturnProperty->GetCPPType()));
+			if (UE::MVVM::BindingHelper::ArePropertiesCompatible(SourceProperty, ArgumentProperty))
+			{
+				bAnyCompatible = true;
+				break;
+			}
 		}
 
-		const FProperty* ConversionArgProperty = ConversionResult.GetValue().ArgumentProperty;
-		if (!UE::MVVM::BindingHelper::ArePropertiesCompatible(SourceProperty, ConversionArgProperty))
+		if (!bAnyCompatible)
 		{
-			return MakeError(FString::Printf(TEXT("The source property '%s' (%s) does not match the argument type of the conversion function (%s)."), *SourceProperty->GetName(), *SourceProperty->GetCPPType(), *DestinationProperty->GetCPPType()));
+			FString ArgumentsString = FString::JoinBy(ArgumentProperties, TEXT(", "), [](const FProperty* Property) { return Property->GetCPPType(); });
+
+			return MakeError(FString::Printf(TEXT("The source property '%s' (%s) does not match any of the argument types of the conversion function (%s)."), *SourceProperty->GetName(), *SourceProperty->GetCPPType(), *ArgumentsString));
 		}
 	}
 	else if (!UE::MVVM::BindingHelper::ArePropertiesCompatible(SourceProperty, DestinationProperty))
@@ -301,4 +317,3 @@ TValueOrError<bool, FString> UMVVMSubsystem::IsBindingValid(FConstDirectionalBin
 
 	return MakeValue(true);
 }
-
diff --git a/Engine/Plugins/Runtime/ModelViewViewModel/Source/ModelViewViewModel/Public/Bindings/MVVMBindingHelper.h b/Engine/Plugins/Runtime/ModelViewViewModel/Source/ModelViewViewModel/Public/Bindings/MVVMBindingHelper.h
--- a/Engine/Plugins/Runtime/ModelViewViewModel/Source/ModelViewViewModel/Public/Bindings/MVVMBindingHelper.h
+++ b/Engine/Plugins/Runtime/ModelViewViewModel/Source/ModelViewViewModel/Public/Bindings/MVVMBindingHelper.h
@@ -79,13 +79,11 @@ namespace UE::MVVM::BindingHelper
 	UE_NODISCARD MODELVIEWVIEWMODEL_API TValueOrError<const FProperty*, FString> TryGetPropertyTypeForSourceBinding(const UFunction* InFunction);
 	UE_NODISCARD MODELVIEWVIEWMODEL_API TValueOrError<const FProperty*, FString> TryGetPropertyTypeForDestinationBinding(const UFunction* InFunction);
 
-	struct FConversionFunctionArguments
-	{
-		const FProperty* ReturnProperty = nullptr;
-		const FProperty* ArgumentProperty = nullptr;
-	};
-	/** Returns the Return and 1st argument for a conversion function. */
-	UE_NODISCARD MODELVIEWVIEWMODEL_API TValueOrError<FConversionFunctionArguments, FString> TryGetPropertyTypeForConversionFunction(const UFunction* InFunction);
+	/** Returns the return property for a conversion function, or an error if the function is not valid as a conversion function. */
+	UE_NODISCARD MODELVIEWVIEWMODEL_API TValueOrError<const FProperty*, FString> TryGetReturnTypeForConversionFunction(const UFunction* InFunction);
+
+	/** Returns the argument properties for a conversion function, or an error if the function is not valid as a conversion function. */
+	UE_NODISCARD MODELVIEWVIEWMODEL_API TValueOrError<TArray<const FProperty*>, FString> TryGetArgumentsForConversionFunction(const UFunction* InFunction);
 
 	/** Are type the same or could be converted at runtime. */
 	UE_NODISCARD MODELVIEWVIEWMODEL_API bool ArePropertiesCompatible(const FProperty* Source, const FProperty* Destination);
diff --git a/Engine/Plugins/Runtime/ModelViewViewModel/Source/ModelViewViewModelEditor/Private/MVVMEditorSubsystem.cpp b/Engine/Plugins/Runtime/ModelViewViewModel/Source/ModelViewViewModelEditor/Private/MVVMEditorSubsystem.cpp
--- a/Engine/Plugins/Runtime/ModelViewViewModel/Source/ModelViewViewModelEditor/Private/MVVMEditorSubsystem.cpp
+++ b/Engine/Plugins/Runtime/ModelViewViewModel/Source/ModelViewViewModelEditor/Private/MVVMEditorSubsystem.cpp
@@ -6,6 +6,7 @@
 #include "BlueprintActionDatabase.h"
 #include "BlueprintNodeSpawner.h"
 #include "Engine/Engine.h"
+#include "Kismet/BlueprintFunctionLibrary.h"
 #include "MVVMSubsystem.h"
 #include "MVVMWidgetBlueprintExtension_View.h"
 #include "WidgetBlueprint.h"
@@ -68,10 +69,64 @@ void UMVVMEditorSubsystem::RemoveBinding(UWidgetBlueprint* WidgetBlueprint, cons
 	}
 }
 
-TArray<const UFunction*> UMVVMEditorSubsystem::GetAvailableConversionFunctions(const UE::MVVM::FMVVMConstFieldVariant& Source, const UE::MVVM::FMVVMConstFieldVariant& Destination) const
+bool UMVVMEditorSubsystem::IsValidConversionFunction(const UFunction* Function, const UE::MVVM::FMVVMConstFieldVariant& Source, const UE::MVVM::FMVVMConstFieldVariant& Destination) const
+{
+	if (Source.IsEmpty() || Destination.IsEmpty())
+	{
+		return false;
+	}
+
+	TValueOrError<const FProperty*, FString> ReturnResult = UE::MVVM::BindingHelper::TryGetReturnTypeForConversionFunction(Function);
+	if (ReturnResult.HasError())
+	{
+		return false;
+	}
+
+	TValueOrError<TArray<const FProperty*>, FString> ArgumentsResult = UE::MVVM::BindingHelper::TryGetArgumentsForConversionFunction(Function);
+	if (ArgumentsResult.HasError())
+	{
+		return false;
+	}
+
+	const FProperty* ReturnProperty = ReturnResult.GetValue();
+	const FProperty* SourceProperty = Source.IsProperty() ? Source.GetProperty() : UE::MVVM::BindingHelper::GetReturnProperty(Source.GetFunction());
+	const FProperty* DestinationProperty = Destination.IsProperty() ? Destination.GetProperty() : UE::MVVM::BindingHelper::GetFirstArgumentProperty(Destination.GetFunction());
+
+	// check that at least one source -> argument binding is compatible
+	bool bAnyCompatible = false;
+
+	const TArray<const FProperty*>& ConversionArgProperties = ArgumentsResult.GetValue();
+	for (const FProperty* ArgumentProperty : ConversionArgProperties)
+	{
+		if (ArgumentProperty->IsA<FObjectProperty>())
+		{
+			// filter out any functions with UObject properties - they aren't valid conversion functions
+			return false;
+		}
+
+		if (UE::MVVM::BindingHelper::ArePropertiesCompatible(SourceProperty, ArgumentProperty))
+		{
+			bAnyCompatible = true;
+		}
+	}
+
+	if (!bAnyCompatible)
+	{
+		return false;
+	}
+
+	// check that the return -> dest is valid
+	if (!UE::MVVM::BindingHelper::ArePropertiesCompatible(ReturnProperty, DestinationProperty))
+	{
+		return false;
+	}
+
+	return true;
+}
+
+TArray<const UFunction*> UMVVMEditorSubsystem::GetAvailableConversionFunctions(const UE::MVVM::FMVVMConstFieldVariant& Source, const UE::MVVM::FMVVMConstFieldVariant& Destination, const UWidgetBlueprint* WidgetBlueprint) const
 {
 	TArray<const UFunction*> ConversionFunctions;
-	TArray<const UFunction*> DEBUG_StaticTwoArgFunctions;
 
 	if (Source.IsEmpty() || Destination.IsEmpty())
 	{
@@ -93,33 +148,16 @@ TArray<const UFunction*> UMVVMEditorSubsystem::GetAvailableConversionFunctions(c
 				const UFunction* Function = BlueprintAction.GetAssociatedFunction();
 				if (Function != nullptr)
 				{
-					TValueOrError<UE::MVVM::BindingHelper::FConversionFunctionArguments, FString> FunctionProperties = UE::MVVM::BindingHelper::TryGetPropertyTypeForConversionFunction(Function);
-					if (FunctionProperties.HasError())
+					// functions in the widget blueprint can do anything they want, other functions have to be static functions in a BlueprintFunctionLibrary
+					const UClass* FunctionClass = Function->GetOuterUClass();
+					if ((FunctionClass->ClassGeneratedBy == WidgetBlueprint) ||
+						(FunctionClass->IsChildOf<UBlueprintFunctionLibrary>() && Function->HasAllFunctionFlags(FUNC_Static | FUNC_BlueprintPure)))
 					{
-						continue;
+						if (IsValidConversionFunction(Function, Destination, Source))
+						{
+							ConversionFunctions.Add(Function);
+						}
 					}
-
-					const FProperty* ArgumentProperty = FunctionProperties.GetValue().ArgumentProperty;
-					const FProperty* ReturnProperty = FunctionProperties.GetValue().ReturnProperty;
-
-					DEBUG_StaticTwoArgFunctions.Add(Function);
-
-					const FProperty* SourceProperty = Source.IsProperty() ? Source.GetProperty() : UE::MVVM::BindingHelper::GetReturnProperty(Source.GetFunction());
-					const FProperty* DestinationProperty = Destination.IsProperty() ? Destination.GetProperty() : UE::MVVM::BindingHelper::GetFirstArgumentProperty(Destination.GetFunction());
-
-					// check that the source -> argument is valid
-					if (!UE::MVVM::BindingHelper::ArePropertiesCompatible(SourceProperty, ArgumentProperty))
-					{
-						continue;
-					}
-					
-					// check that the return -> dest is valid
-					if (!UE::MVVM::BindingHelper::ArePropertiesCompatible(ReturnProperty, DestinationProperty))
-					{
-						continue;
-					}
-
-					ConversionFunctions.Add(Function);
 				}
 			}
 		}
diff --git a/Engine/Plugins/Runtime/ModelViewViewModel/Source/ModelViewViewModelEditor/Private/Widgets/SMVVMConversionPath.cpp b/Engine/Plugins/Runtime/ModelViewViewModel/Source/ModelViewViewModelEditor/Private/Widgets/SMVVMConversionPath.cpp
--- a/Engine/Plugins/Runtime/ModelViewViewModel/Source/ModelViewViewModelEditor/Private/Widgets/SMVVMConversionPath.cpp
+++ b/Engine/Plugins/Runtime/ModelViewViewModel/Source/ModelViewViewModelEditor/Private/Widgets/SMVVMConversionPath.cpp
@@ -181,6 +181,26 @@ void SMVVMConversionPath::SetConversionFunction(const UFunction* Function)
 	OnFunctionChanged.ExecuteIfBound(Function);
 }
 
+void SMVVMConversionPath::PopulateMenuForEntry(FMenuBuilder& MenuBuilder, const FFunctionEntry* FunctionEntry)
+{
+	for (const FFunctionEntry& Category : FunctionEntry->Categories)
+	{
+		MenuBuilder.AddSubMenu(FText::FromString(Category.CategoryName), 
+			FText::FromString(Category.CategoryName), 
+			FNewMenuDelegate::CreateSP(this, &SMVVMConversionPath::PopulateMenuForEntry, &Category));
+	}
+
+	for (const UFunction* Function : FunctionEntry->Functions)
+	{
+		MenuBuilder.AddMenuEntry(
+			Function->GetDisplayNameText(),
+			Function->GetToolTipText(),
+			FSlateIcon(FAppStyle::GetAppStyleSetName(), "Kismet.AllClasses.FunctionIcon"),
+			FUIAction(FExecuteAction::CreateSP(this, &SMVVMConversionPath::SetConversionFunction, Function))
+		);
+	}
+}
+
 TSharedRef<SWidget> SMVVMConversionPath::GetFunctionMenuContent()
 {
 	TArray<FMVVMBlueprintViewBinding*> ViewBindings = Bindings.Get(TArray<FMVVMBlueprintViewBinding*>());
@@ -204,7 +224,7 @@ TSharedRef<SWidget> SMVVMConversionPath::GetFunctionMenuContent()
 		Args.DestinationBinding = bSourceToDestination ? WidgetField : ViewModelField;
 
 		UMVVMEditorSubsystem* EditorSubsystem = GEditor->GetEditorSubsystem<UMVVMEditorSubsystem>();
-		TArray<const UFunction*> FunctionsForThis = EditorSubsystem->GetAvailableConversionFunctions(Args.SourceBinding, Args.DestinationBinding);
+		TArray<const UFunction*> FunctionsForThis = EditorSubsystem->GetAvailableConversionFunctions(Args.SourceBinding, Args.DestinationBinding, WidgetBlueprint);
 
 		if (ConversionFunctions.Num() > 0)
 		{
@@ -233,15 +253,40 @@ TSharedRef<SWidget> SMVVMConversionPath::GetFunctionMenuContent()
 			true // searchable
 		);
 	}
-
-	for (const UFunction* Function : ConversionFunctions)
+	else
 	{
-		FUIAction Action(FExecuteAction::CreateSP(this, &SMVVMConversionPath::SetConversionFunction, Function));
-		MenuBuilder.AddMenuEntry(
-			Function->GetDisplayNameText(),
-			Function->GetToolTipText(),
-			FSlateIcon(FAppStyle::GetAppStyleSetName(), "Kismet.AllClasses.FunctionIcon"),
-			Action);
+		RootEntry = FFunctionEntry();
+
+		TArray<FString> SubCategories;
+		for (const UFunction* Function : ConversionFunctions)
+		{
+			const FString Category = Function->GetMetaData("Category");
+			SubCategories.Reset();
+			Category.ParseIntoArray(SubCategories, TEXT("|"));
+
+			// create categories
+			FFunctionEntry* CurrentEntry = &RootEntry;
+			for (const FString& CategoryName : SubCategories)
+			{
+				FFunctionEntry* ExistingCategory = CurrentEntry->Categories.FindByPredicate([CategoryName](const FFunctionEntry& Entry)
+					{
+						return Entry.CategoryName == CategoryName;
+					});
+				if (ExistingCategory == nullptr)
+				{
+					int32 NewIndex = CurrentEntry->Categories.Add(FFunctionEntry());
+					ExistingCategory = &CurrentEntry->Categories[NewIndex];
+
+					ExistingCategory->CategoryName = CategoryName;
+				}
+
+				CurrentEntry = ExistingCategory;
+			}
+
+			CurrentEntry->Functions.Add(Function);
+		}
+
+		PopulateMenuForEntry(MenuBuilder, &RootEntry);
 	}
 
 	const FString Path = GetFunctionPath();
diff --git a/Engine/Plugins/Runtime/ModelViewViewModel/Source/ModelViewViewModelEditor/Private/Widgets/SMVVMConversionPath.h b/Engine/Plugins/Runtime/ModelViewViewModel/Source/ModelViewViewModelEditor/Private/Widgets/SMVVMConversionPath.h
--- a/Engine/Plugins/Runtime/ModelViewViewModel/Source/ModelViewViewModelEditor/Private/Widgets/SMVVMConversionPath.h
+++ b/Engine/Plugins/Runtime/ModelViewViewModel/Source/ModelViewViewModelEditor/Private/Widgets/SMVVMConversionPath.h
@@ -4,12 +4,24 @@
 
 #include "Widgets/SCompoundWidget.h"
 
+struct FCreateWidgetForActionData;
+struct FEdGraphSchemaAction;
+struct FGraphActionListBuilderBase;
+class FMenuBuilder;
 struct FMVVMBlueprintViewBinding;
 class SMenuAnchor;
 class UWidgetBlueprint;
 
 class SMVVMConversionPath : public SCompoundWidget
 {
+private:
+	struct FFunctionEntry
+	{
+		FString CategoryName;
+		TArray<FFunctionEntry> Categories;
+		TArray<const UFunction*> Functions;
+	};
+
 public:
 	DECLARE_DELEGATE_OneParam(FOnFunctionChanged, const UFunction*);
 
@@ -30,8 +42,11 @@ private:
 	TSharedRef<SWidget> GetFunctionMenuContent();
 	FReply OnButtonClicked() const;
 	FString GetFunctionPath() const;
+	void PopulateMenuForEntry(FMenuBuilder& MenuBuilder, const FFunctionEntry* FunctionEntry);
 
 private:
+
+	FFunctionEntry RootEntry;
 	TAttribute<TArray<FMVVMBlueprintViewBinding*>> Bindings;
 	FOnFunctionChanged OnFunctionChanged;
 	TSharedPtr<SMenuAnchor> Anchor;
diff --git a/Engine/Plugins/Runtime/ModelViewViewModel/Source/ModelViewViewModelEditor/Public/MVVMEditorSubsystem.h b/Engine/Plugins/Runtime/ModelViewViewModel/Source/ModelViewViewModelEditor/Public/MVVMEditorSubsystem.h
--- a/Engine/Plugins/Runtime/ModelViewViewModel/Source/ModelViewViewModelEditor/Public/MVVMEditorSubsystem.h
+++ b/Engine/Plugins/Runtime/ModelViewViewModel/Source/ModelViewViewModelEditor/Public/MVVMEditorSubsystem.h
@@ -36,5 +36,6 @@ public:
 	void RemoveBinding(UWidgetBlueprint* WidgetBlueprint, const FMVVMBlueprintViewBinding& Binding);
 
 	TArray<UE::MVVM::FMVVMConstFieldVariant> GetChildViewModels(TSubclassOf<UMVVMViewModelBase> Class);
-	TArray<const UFunction*> GetAvailableConversionFunctions(const UE::MVVM::FMVVMConstFieldVariant& Source, const UE::MVVM::FMVVMConstFieldVariant& Dest) const;
+	TArray<const UFunction*> GetAvailableConversionFunctions(const UE::MVVM::FMVVMConstFieldVariant& Source, const UE::MVVM::FMVVMConstFieldVariant& Destination, const UWidgetBlueprint* WidgetBlueprint) const;
+	bool IsValidConversionFunction(const UFunction* Function, const UE::MVVM::FMVVMConstFieldVariant& Source, const UE::MVVM::FMVVMConstFieldVariant& Destination) const;
 };
