Reverting support for multiple argument conversion function, added in CL20180539, since parts of the feature are not yet fully supported.

#rnx

diff --git a/Engine/Plugins/Runtime/ModelViewViewModel/Source/ModelViewViewModel/Private/Bindings/MVVMBindingHelper.cpp b/Engine/Plugins/Runtime/ModelViewViewModel/Source/ModelViewViewModel/Private/Bindings/MVVMBindingHelper.cpp
--- a/Engine/Plugins/Runtime/ModelViewViewModel/Source/ModelViewViewModel/Private/Bindings/MVVMBindingHelper.cpp
+++ b/Engine/Plugins/Runtime/ModelViewViewModel/Source/ModelViewViewModel/Private/Bindings/MVVMBindingHelper.cpp
@@ -76,11 +76,10 @@ namespace UE::MVVM::BindingHelper
 		// A setter can return a value "bool SetValue(int32)" but is not valid for MVVM
 		//&& ((InFunction->NumParms == 1 && GetReturnProperty() == nullptr) || (InFunction->NumParms == 2 && GetReturnProperty() != nullptr));
 
-		TValueOrError<TArray<const FProperty*>, FString> Result = TryGetArgumentsForConversionFunction(InFunction);
-
 		return Private::IsValidCommon(InFunction)
 			&& !InFunction->HasAnyFunctionFlags(FUNC_Const | FUNC_BlueprintPure)
-			&& Result.HasValue() && Result.GetValue().Num() > 0;
+			&& InFunction->NumParms == 1
+			&& GetFirstArgumentProperty(InFunction) != nullptr;
 	}
 
 
@@ -373,34 +372,8 @@ namespace UE::MVVM::BindingHelper
 			return TryGetPropertyTypeForDestinationBinding(InField.GetFunction());
 		}
 	}
-	TValueOrError<const FProperty*, FString> TryGetReturnTypeForConversionFunction(const UFunction* InFunction)
-	{
-		FString CommonResult = Private::TryGetPropertyTypeCommon(InFunction);
-		if (!CommonResult.IsEmpty())
-		{
-			return MakeError(MoveTemp(CommonResult));
-		}
 
-		if (!InFunction->HasAllFunctionFlags(FUNC_Static) && !InFunction->HasAllFunctionFlags(FUNC_Const | FUNC_BlueprintPure))
-		{
-			return MakeError(FString::Printf(TEXT("The function '%s' is not static or is not const and pure."), *InFunction->GetName()));
-		}
-
-		if (InFunction->NumParms < 2)
-		{
-			return MakeError(FString::Printf(TEXT("The function '%s' does not have the correct number of arguments."), *InFunction->GetName()));
-		}
-
-		const FProperty* ReturnProperty = GetReturnProperty(InFunction);
-		if (ReturnProperty == nullptr)
-		{
-			return MakeError(FString::Printf(TEXT("The return value for function '%s' is invalid."), *InFunction->GetName()));
-		}
-
-		return MakeValue(ReturnProperty);
-	}
-
-	TValueOrError<TArray<const FProperty*>, FString> TryGetArgumentsForConversionFunction(const UFunction* InFunction)
+	TValueOrError<FConversionFunctionArguments, FString> TryGetPropertyTypeForConversionFunction(const UFunction* InFunction)
 	{
 		FString CommonResult = Private::TryGetPropertyTypeCommon(InFunction);
 		if (!CommonResult.IsEmpty())
@@ -408,41 +381,33 @@ namespace UE::MVVM::BindingHelper
 			return MakeError(MoveTemp(CommonResult));
 		}
 
-		if (!InFunction->HasAllFunctionFlags(FUNC_Static) && !InFunction->HasAnyFunctionFlags(FUNC_Const | FUNC_BlueprintPure))
+		if (!InFunction->HasAllFunctionFlags(FUNC_Static) && !InFunction->HasAllFunctionFlags(FUNC_Const | FUNC_BlueprintPure))
 		{
 			return MakeError(FString::Printf(TEXT("The function '%s' is not static or is not const and pure."), *InFunction->GetName()));
 		}
 
-		if (InFunction->NumParms < 2)
+		if (InFunction->NumParms != 2)
 		{
 			return MakeError(FString::Printf(TEXT("The function '%s' does not have the correct number of arguments."), *InFunction->GetName()));
 		}
 
-		const FProperty* ReturnProperty = GetReturnProperty(InFunction);
-		if (ReturnProperty == nullptr)
+		FConversionFunctionArguments Result;
+		Result.ReturnProperty = GetReturnProperty(InFunction);
+		if (Result.ReturnProperty == nullptr)
 		{
 			return MakeError(FString::Printf(TEXT("The return value for function '%s' is invalid."), *InFunction->GetName()));
 		}
 
-		TArray<const FProperty*> Arguments;
-		for (TFieldIterator<FProperty> It(InFunction); It && (It->PropertyFlags & CPF_Parm); ++It)
+		Result.ArgumentProperty = GetFirstArgumentProperty(InFunction);
+		if (Result.ArgumentProperty == nullptr)
 		{
-			const FProperty* Property = *It;
-			if (It->HasAllPropertyFlags(CPF_Parm) && !It->HasAnyPropertyFlags(CPF_Deprecated | CPF_EditorOnly | CPF_ReturnParm))
-			{
-				// argument is either an out parameter or not const
-				if (It->HasAllPropertyFlags(CPF_OutParm) && !It->HasAllPropertyFlags(CPF_ConstParm))
-				{
-					continue;
-				}
-
-				Arguments.Add(*It);
-			}
+			return MakeError(FString::Printf(TEXT("The function '%s' doesn't not have a valid argument."), *InFunction->GetName()));
 		}
 
-		return MakeValue(Arguments);
+		return MakeValue(Result);
 	}
 
+
 	namespace Private
 	{
 		bool IsObjectPropertyCompatible(const FProperty* Source, const FProperty* Destination)
diff --git a/Engine/Plugins/Runtime/ModelViewViewModel/Source/ModelViewViewModel/Private/MVVMSubsystem.cpp b/Engine/Plugins/Runtime/ModelViewViewModel/Source/ModelViewViewModel/Private/MVVMSubsystem.cpp
--- a/Engine/Plugins/Runtime/ModelViewViewModel/Source/ModelViewViewModel/Private/MVVMSubsystem.cpp
+++ b/Engine/Plugins/Runtime/ModelViewViewModel/Source/ModelViewViewModel/Private/MVVMSubsystem.cpp
@@ -128,15 +128,18 @@ namespace UE::MVVM::Private
 		TArray<FMVVMAvailableBinding> Result;
 		Result.Reserve(FieldDescriptors.Num());
 
-		for (TFieldIterator<const UFunction> FunctionIt(Class, EFieldIteratorFlags::IncludeSuper); FunctionIt; ++FunctionIt)
+		for (TFieldIterator<const UFunction> FunctionItt(Class, EFieldIteratorFlags::IncludeSuper); FunctionItt; ++FunctionItt)
 		{
-			const UFunction* Function = *FunctionIt;
+			const UFunction* Function = *FunctionItt;
 			check(Function);
-			const bool bIsReadable = FieldDescriptors.Contains(Function->GetFName()) && BindingHelper::IsValidForSourceBinding(Function);
-			const bool bIsWritable = BindingHelper::IsValidForDestinationBinding(Function);
-			if (bIsReadable || bIsWritable)
+			if (!Function->HasAnyFunctionFlags(FUNC_Private | FUNC_Protected))
 			{
-				Result.Add(FMVVMAvailableBinding(FMVVMBindingName(Function->GetFName()), bIsReadable, bIsWritable));
+				const bool bIsReadable = FieldDescriptors.Contains(Function->GetFName()) && BindingHelper::IsValidForSourceBinding(Function);
+				const bool bIsWritable = BindingHelper::IsValidForDestinationBinding(Function);
+				if (bIsReadable || bIsWritable)
+				{
+					Result.Add(FMVVMAvailableBinding(FMVVMBindingName(Function->GetFName()), bIsReadable, bIsWritable));
+				}
 			}
 		}
 
@@ -244,26 +247,26 @@ TValueOrError<bool, FString> UMVVMSubsystem::IsBindingValid(FDirectionalBindingA
 TValueOrError<bool, FString> UMVVMSubsystem::IsBindingValid(FConstDirectionalBindingArgs Args) const
 {
 	// Test Source
-	TValueOrError<const FProperty*, FString> SourceResult = UE::MVVM::BindingHelper::TryGetPropertyTypeForSourceBinding(Args.SourceBinding);
-	if (SourceResult.HasError())
+	TValueOrError<const FProperty*, FString> GetterProperty = UE::MVVM::BindingHelper::TryGetPropertyTypeForSourceBinding(Args.SourceBinding);
+	if (GetterProperty.HasError())
 	{
-		return MakeError(SourceResult.StealError());
+		return MakeError(GetterProperty.StealError());
 	}
 
-	const FProperty* SourceProperty = SourceResult.StealValue();
+	const FProperty* SourceProperty = GetterProperty.StealValue();
 	if (SourceProperty == nullptr)
 	{
 		return MakeError(TEXT("There is no value to bind at the source."));
 	}
 
 	// Test Destination
-	TValueOrError<const FProperty*, FString> DestinationResult = UE::MVVM::BindingHelper::TryGetPropertyTypeForDestinationBinding(Args.DestinationBinding);
-	if (DestinationResult.HasError())
+	TValueOrError<const FProperty*, FString> SetterProperty = UE::MVVM::BindingHelper::TryGetPropertyTypeForDestinationBinding(Args.DestinationBinding);
+	if (SetterProperty.HasError())
 	{
-		return MakeError(DestinationResult.StealError());
+		return MakeError(SetterProperty.StealError());
 	}
 
-	const FProperty* DestinationProperty = DestinationResult.StealValue();
+	const FProperty* DestinationProperty = SetterProperty.StealValue();
 	if (DestinationProperty == nullptr)
 	{
 		return MakeError(TEXT("There is no value to bind at the destination."));
@@ -272,42 +275,23 @@ TValueOrError<bool, FString> UMVVMSubsystem::IsBindingValid(FConstDirectionalBin
 	// Test the conversion function
 	if (Args.ConversionFunction)
 	{
-		TValueOrError<const FProperty*, FString> ReturnResult = UE::MVVM::BindingHelper::TryGetReturnTypeForConversionFunction(Args.ConversionFunction);
-		if (ReturnResult.HasError())
+		TValueOrError<UE::MVVM::BindingHelper::FConversionFunctionArguments, FString> ConversionResult = UE::MVVM::BindingHelper::TryGetPropertyTypeForConversionFunction(Args.ConversionFunction);
+		if (ConversionResult.HasError())
 		{
-			return MakeError(ReturnResult.StealError());
-		}
-
-		TValueOrError<TArray<const FProperty*>, FString> ArgumentsResult = UE::MVVM::BindingHelper::TryGetArgumentsForConversionFunction(Args.ConversionFunction);
-		if (ArgumentsResult.HasError())
-		{
-			return MakeError(ArgumentsResult.StealError());
+			return MakeError(ConversionResult.StealError());
 		}
 
 		// The compiled version should look like Setter(Conversion(Getter())).
-		const FProperty* ReturnProperty = ReturnResult.GetValue();
-		if (!UE::MVVM::BindingHelper::ArePropertiesCompatible(ReturnProperty, DestinationProperty))
-		{
-			return MakeError(FString::Printf(TEXT("The destination property '%s' (%s) does not match the return type of the conversion function (%s)."), *DestinationProperty->GetName(), *DestinationProperty->GetCPPType(), *ReturnProperty->GetCPPType()));
-		}
-
-		bool bAnyCompatible = false;
-
-		const TArray<const FProperty*>& ArgumentProperties = ArgumentsResult.GetValue();
-		for (const FProperty* ArgumentProperty : ArgumentProperties)
+		const FProperty* ConversionReturnProperty = ConversionResult.GetValue().ReturnProperty;
+		if (!UE::MVVM::BindingHelper::ArePropertiesCompatible(ConversionReturnProperty, DestinationProperty))
 		{
-			if (UE::MVVM::BindingHelper::ArePropertiesCompatible(SourceProperty, ArgumentProperty))
-			{
-				bAnyCompatible = true;
-				break;
-			}
+			return MakeError(FString::Printf(TEXT("The destination property '%s' (%s) does not match the return type of the conversion function (%s)."), *DestinationProperty->GetName(), *DestinationProperty->GetCPPType(), *ConversionReturnProperty->GetCPPType()));
 		}
 
-		if (!bAnyCompatible)
+		const FProperty* ConversionArgProperty = ConversionResult.GetValue().ArgumentProperty;
+		if (!UE::MVVM::BindingHelper::ArePropertiesCompatible(SourceProperty, ConversionArgProperty))
 		{
-			FString ArgumentsString = FString::JoinBy(ArgumentProperties, TEXT(", "), [](const FProperty* Property) { return Property->GetCPPType(); });
-
-			return MakeError(FString::Printf(TEXT("The source property '%s' (%s) does not match any of the argument types of the conversion function (%s)."), *SourceProperty->GetName(), *SourceProperty->GetCPPType(), *ArgumentsString));
+			return MakeError(FString::Printf(TEXT("The source property '%s' (%s) does not match the argument type of the conversion function (%s)."), *SourceProperty->GetName(), *SourceProperty->GetCPPType(), *DestinationProperty->GetCPPType()));
 		}
 	}
 	else if (!UE::MVVM::BindingHelper::ArePropertiesCompatible(SourceProperty, DestinationProperty))
@@ -317,3 +301,4 @@ TValueOrError<bool, FString> UMVVMSubsystem::IsBindingValid(FConstDirectionalBin
 
 	return MakeValue(true);
 }
+
diff --git a/Engine/Plugins/Runtime/ModelViewViewModel/Source/ModelViewViewModel/Public/Bindings/MVVMBindingHelper.h b/Engine/Plugins/Runtime/ModelViewViewModel/Source/ModelViewViewModel/Public/Bindings/MVVMBindingHelper.h
--- a/Engine/Plugins/Runtime/ModelViewViewModel/Source/ModelViewViewModel/Public/Bindings/MVVMBindingHelper.h
+++ b/Engine/Plugins/Runtime/ModelViewViewModel/Source/ModelViewViewModel/Public/Bindings/MVVMBindingHelper.h
@@ -79,11 +79,13 @@ namespace UE::MVVM::BindingHelper
 	UE_NODISCARD MODELVIEWVIEWMODEL_API TValueOrError<const FProperty*, FString> TryGetPropertyTypeForSourceBinding(const UFunction* InFunction);
 	UE_NODISCARD MODELVIEWVIEWMODEL_API TValueOrError<const FProperty*, FString> TryGetPropertyTypeForDestinationBinding(const UFunction* InFunction);
 
-	/** Returns the return property for a conversion function, or an error if the function is not valid as a conversion function. */
-	UE_NODISCARD MODELVIEWVIEWMODEL_API TValueOrError<const FProperty*, FString> TryGetReturnTypeForConversionFunction(const UFunction* InFunction);
-
-	/** Returns the argument properties for a conversion function, or an error if the function is not valid as a conversion function. */
-	UE_NODISCARD MODELVIEWVIEWMODEL_API TValueOrError<TArray<const FProperty*>, FString> TryGetArgumentsForConversionFunction(const UFunction* InFunction);
+	struct FConversionFunctionArguments
+	{
+		const FProperty* ReturnProperty = nullptr;
+		const FProperty* ArgumentProperty = nullptr;
+	};
+	/** Returns the Return and 1st argument for a conversion function. */
+	UE_NODISCARD MODELVIEWVIEWMODEL_API TValueOrError<FConversionFunctionArguments, FString> TryGetPropertyTypeForConversionFunction(const UFunction* InFunction);
 
 	/** Are type the same or could be converted at runtime. */
 	UE_NODISCARD MODELVIEWVIEWMODEL_API bool ArePropertiesCompatible(const FProperty* Source, const FProperty* Destination);
diff --git a/Engine/Plugins/Runtime/ModelViewViewModel/Source/ModelViewViewModelEditor/Private/MVVMEditorSubsystem.cpp b/Engine/Plugins/Runtime/ModelViewViewModel/Source/ModelViewViewModelEditor/Private/MVVMEditorSubsystem.cpp
--- a/Engine/Plugins/Runtime/ModelViewViewModel/Source/ModelViewViewModelEditor/Private/MVVMEditorSubsystem.cpp
+++ b/Engine/Plugins/Runtime/ModelViewViewModel/Source/ModelViewViewModelEditor/Private/MVVMEditorSubsystem.cpp
@@ -6,7 +6,6 @@
 #include "BlueprintActionDatabase.h"
 #include "BlueprintNodeSpawner.h"
 #include "Engine/Engine.h"
-#include "Kismet/BlueprintFunctionLibrary.h"
 #include "MVVMSubsystem.h"
 #include "MVVMWidgetBlueprintExtension_View.h"
 #include "WidgetBlueprint.h"
@@ -69,62 +68,7 @@ void UMVVMEditorSubsystem::RemoveBinding(UWidgetBlueprint* WidgetBlueprint, cons
 	}
 }
 
-bool UMVVMEditorSubsystem::IsValidConversionFunction(const UFunction* Function, const UE::MVVM::FMVVMConstFieldVariant& Source, const UE::MVVM::FMVVMConstFieldVariant& Destination) const
-{
-	if (Source.IsEmpty() || Destination.IsEmpty())
-	{
-		return false;
-	}
-
-	TValueOrError<const FProperty*, FString> ReturnResult = UE::MVVM::BindingHelper::TryGetReturnTypeForConversionFunction(Function);
-	if (ReturnResult.HasError())
-	{
-		return false;
-	}
-
-	TValueOrError<TArray<const FProperty*>, FString> ArgumentsResult = UE::MVVM::BindingHelper::TryGetArgumentsForConversionFunction(Function);
-	if (ArgumentsResult.HasError())
-	{
-		return false;
-	}
-
-	const FProperty* ReturnProperty = ReturnResult.GetValue();
-	const FProperty* SourceProperty = Source.IsProperty() ? Source.GetProperty() : UE::MVVM::BindingHelper::GetReturnProperty(Source.GetFunction());
-	const FProperty* DestinationProperty = Destination.IsProperty() ? Destination.GetProperty() : UE::MVVM::BindingHelper::GetFirstArgumentProperty(Destination.GetFunction());
-
-	// check that at least one source -> argument binding is compatible
-	bool bAnyCompatible = false;
-
-	const TArray<const FProperty*>& ConversionArgProperties = ArgumentsResult.GetValue();
-	for (const FProperty* ArgumentProperty : ConversionArgProperties)
-	{
-		if (ArgumentProperty->IsA<FObjectProperty>())
-		{
-			// filter out any functions with UObject properties - they aren't valid conversion functions
-			return false;
-		}
-
-		if (UE::MVVM::BindingHelper::ArePropertiesCompatible(SourceProperty, ArgumentProperty))
-		{
-			bAnyCompatible = true;
-		}
-	}
-
-	if (!bAnyCompatible)
-	{
-		return false;
-	}
-
-	// check that the return -> dest is valid
-	if (!UE::MVVM::BindingHelper::ArePropertiesCompatible(ReturnProperty, DestinationProperty))
-	{
-		return false;
-	}
-
-	return true;
-}
-
-TArray<const UFunction*> UMVVMEditorSubsystem::GetAvailableConversionFunctions(const UE::MVVM::FMVVMConstFieldVariant& Source, const UE::MVVM::FMVVMConstFieldVariant& Destination, const UWidgetBlueprint* WidgetBlueprint) const
+TArray<const UFunction*> UMVVMEditorSubsystem::GetAvailableConversionFunctions(const UE::MVVM::FMVVMConstFieldVariant& Source, const UE::MVVM::FMVVMConstFieldVariant& Destination) const
 {
 	TArray<const UFunction*> ConversionFunctions;
 
@@ -148,16 +92,31 @@ TArray<const UFunction*> UMVVMEditorSubsystem::GetAvailableConversionFunctions(c
 				const UFunction* Function = BlueprintAction.GetAssociatedFunction();
 				if (Function != nullptr)
 				{
-					// functions in the widget blueprint can do anything they want, other functions have to be static functions in a BlueprintFunctionLibrary
-					const UClass* FunctionClass = Function->GetOuterUClass();
-					if ((FunctionClass->ClassGeneratedBy == WidgetBlueprint) ||
-						(FunctionClass->IsChildOf<UBlueprintFunctionLibrary>() && Function->HasAllFunctionFlags(FUNC_Static | FUNC_BlueprintPure)))
+					TValueOrError<UE::MVVM::BindingHelper::FConversionFunctionArguments, FString> FunctionProperties = UE::MVVM::BindingHelper::TryGetPropertyTypeForConversionFunction(Function);
+					if (FunctionProperties.HasError())
 					{
-						if (IsValidConversionFunction(Function, Destination, Source))
-						{
-							ConversionFunctions.Add(Function);
-						}
+						continue;
 					}
+
+					const FProperty* ArgumentProperty = FunctionProperties.GetValue().ArgumentProperty;
+					const FProperty* ReturnProperty = FunctionProperties.GetValue().ReturnProperty;
+
+					const FProperty* SourceProperty = Source.IsProperty() ? Source.GetProperty() : UE::MVVM::BindingHelper::GetReturnProperty(Source.GetFunction());
+					const FProperty* DestinationProperty = Destination.IsProperty() ? Destination.GetProperty() : UE::MVVM::BindingHelper::GetFirstArgumentProperty(Destination.GetFunction());
+
+					// check that the source -> argument is valid
+					if (!UE::MVVM::BindingHelper::ArePropertiesCompatible(SourceProperty, ArgumentProperty))
+					{
+						continue;
+					}
+					
+					// check that the return -> dest is valid
+					if (!UE::MVVM::BindingHelper::ArePropertiesCompatible(ReturnProperty, DestinationProperty))
+					{
+						continue;
+					}
+
+					ConversionFunctions.Add(Function);
 				}
 			}
 		}
diff --git a/Engine/Plugins/Runtime/ModelViewViewModel/Source/ModelViewViewModelEditor/Private/Widgets/SMVVMConversionPath.cpp b/Engine/Plugins/Runtime/ModelViewViewModel/Source/ModelViewViewModelEditor/Private/Widgets/SMVVMConversionPath.cpp
--- a/Engine/Plugins/Runtime/ModelViewViewModel/Source/ModelViewViewModelEditor/Private/Widgets/SMVVMConversionPath.cpp
+++ b/Engine/Plugins/Runtime/ModelViewViewModel/Source/ModelViewViewModelEditor/Private/Widgets/SMVVMConversionPath.cpp
@@ -224,7 +224,7 @@ TSharedRef<SWidget> SMVVMConversionPath::GetFunctionMenuContent()
 		Args.DestinationBinding = bSourceToDestination ? WidgetField : ViewModelField;
 
 		UMVVMEditorSubsystem* EditorSubsystem = GEditor->GetEditorSubsystem<UMVVMEditorSubsystem>();
-		TArray<const UFunction*> FunctionsForThis = EditorSubsystem->GetAvailableConversionFunctions(Args.SourceBinding, Args.DestinationBinding, WidgetBlueprint);
+		TArray<const UFunction*> FunctionsForThis = EditorSubsystem->GetAvailableConversionFunctions(Args.SourceBinding, Args.DestinationBinding);
 
 		if (ConversionFunctions.Num() > 0)
 		{
diff --git a/Engine/Plugins/Runtime/ModelViewViewModel/Source/ModelViewViewModelEditor/Public/MVVMEditorSubsystem.h b/Engine/Plugins/Runtime/ModelViewViewModel/Source/ModelViewViewModelEditor/Public/MVVMEditorSubsystem.h
--- a/Engine/Plugins/Runtime/ModelViewViewModel/Source/ModelViewViewModelEditor/Public/MVVMEditorSubsystem.h
+++ b/Engine/Plugins/Runtime/ModelViewViewModel/Source/ModelViewViewModelEditor/Public/MVVMEditorSubsystem.h
@@ -36,6 +36,5 @@ public:
 	void RemoveBinding(UWidgetBlueprint* WidgetBlueprint, const FMVVMBlueprintViewBinding& Binding);
 
 	TArray<UE::MVVM::FMVVMConstFieldVariant> GetChildViewModels(TSubclassOf<UMVVMViewModelBase> Class);
-	TArray<const UFunction*> GetAvailableConversionFunctions(const UE::MVVM::FMVVMConstFieldVariant& Source, const UE::MVVM::FMVVMConstFieldVariant& Destination, const UWidgetBlueprint* WidgetBlueprint) const;
-	bool IsValidConversionFunction(const UFunction* Function, const UE::MVVM::FMVVMConstFieldVariant& Source, const UE::MVVM::FMVVMConstFieldVariant& Destination) const;
+	TArray<const UFunction*> GetAvailableConversionFunctions(const UE::MVVM::FMVVMConstFieldVariant& Source, const UE::MVVM::FMVVMConstFieldVariant& Dest) const;
 };
