MVVM: Blueprint Path can now support long path. A deprecation path exist. The compiler was modified to support long path. A ViewModelCreator is generated for Manual, this allow the setter to works in standalone.

diff --git a/Engine/Plugins/Runtime/ModelViewViewModel/Source/ModelViewViewModel/Private/Bindings/MVVMCompiledBindingLibrary.cpp b/Engine/Plugins/Runtime/ModelViewViewModel/Source/ModelViewViewModel/Private/Bindings/MVVMCompiledBindingLibrary.cpp
--- a/Engine/Plugins/Runtime/ModelViewViewModel/Source/ModelViewViewModel/Private/Bindings/MVVMCompiledBindingLibrary.cpp
+++ b/Engine/Plugins/Runtime/ModelViewViewModel/Source/ModelViewViewModel/Private/Bindings/MVVMCompiledBindingLibrary.cpp
@@ -6,7 +6,7 @@
 #include "Blueprint/UserWidget.h"
 #include "Blueprint/WidgetTree.h"
 #include "UObject/PropertyAccessUtil.h"
-
+PRAGMA_DISABLE_OPTIMIZATION
 
 /**
  *
@@ -440,3 +440,4 @@ TValueOrError<UE::FieldNotification::FFieldId, void> FMVVMCompiledBindingLibrary
 	check(LoadedFieldIds.IsValidIndex(InFieldId.FieldIdIndex));
 	return MakeValue(LoadedFieldIds[InFieldId.FieldIdIndex]);
 }
+PRAGMA_ENABLE_OPTIMIZATION
\ No newline at end of file
diff --git a/Engine/Plugins/Runtime/ModelViewViewModel/Source/ModelViewViewModel/Private/View/MVVMViewClass.cpp b/Engine/Plugins/Runtime/ModelViewViewModel/Source/ModelViewViewModel/Private/View/MVVMViewClass.cpp
--- a/Engine/Plugins/Runtime/ModelViewViewModel/Source/ModelViewViewModel/Private/View/MVVMViewClass.cpp
+++ b/Engine/Plugins/Runtime/ModelViewViewModel/Source/ModelViewViewModel/Private/View/MVVMViewClass.cpp
@@ -20,6 +20,19 @@
 // 
 ///////////////////////////////////////////////////////////////////////
 
+FMVVMViewClass_SourceCreator FMVVMViewClass_SourceCreator::MakeManual(FName InName, UClass* InNotifyFieldValueChangedClass)
+{
+	FMVVMViewClass_SourceCreator Result;
+	if (ensure(InNotifyFieldValueChangedClass && InNotifyFieldValueChangedClass->ImplementsInterface(UNotifyFieldValueChanged::StaticClass())))
+	{
+		Result.PropertyName = InName;
+		Result.ExpectedSourceType = InNotifyFieldValueChangedClass;
+		Result.bCreateInstance = false;
+		Result.bOptional = true;
+	}
+	return Result;
+}
+
 FMVVMViewClass_SourceCreator FMVVMViewClass_SourceCreator::MakeInstance(FName InName, UClass* InNotifyFieldValueChangedClass)
 {
 	FMVVMViewClass_SourceCreator Result;
diff --git a/Engine/Plugins/Runtime/ModelViewViewModel/Source/ModelViewViewModel/Public/View/MVVMViewClass.h b/Engine/Plugins/Runtime/ModelViewViewModel/Source/ModelViewViewModel/Public/View/MVVMViewClass.h
--- a/Engine/Plugins/Runtime/ModelViewViewModel/Source/ModelViewViewModel/Public/View/MVVMViewClass.h
+++ b/Engine/Plugins/Runtime/ModelViewViewModel/Source/ModelViewViewModel/Public/View/MVVMViewClass.h
@@ -40,6 +40,7 @@ struct FMVVMViewClass_SourceCreator
 	friend UE::MVVM::Private::FMVVMViewBlueprintCompiler;
 
 public:
+	MODELVIEWVIEWMODEL_API static FMVVMViewClass_SourceCreator MakeManual(FName Name, UClass* NotifyFieldValueChangedClass);
 	MODELVIEWVIEWMODEL_API static FMVVMViewClass_SourceCreator MakeInstance(FName Name, UClass* NotifyFieldValueChangedClass);
 	MODELVIEWVIEWMODEL_API static FMVVMViewClass_SourceCreator MakeFieldPath(FName Name, UClass* NotifyFieldValueChangedClass, FMVVMVCompiledFieldPath FieldPath);
 	MODELVIEWVIEWMODEL_API static FMVVMViewClass_SourceCreator MakeGlobalContext(FName Name, FMVVMViewModelContext Context);
diff --git a/Engine/Plugins/Runtime/ModelViewViewModel/Source/ModelViewViewModelBlueprint/Private/MVVMBlueprintView.cpp b/Engine/Plugins/Runtime/ModelViewViewModel/Source/ModelViewViewModelBlueprint/Private/MVVMBlueprintView.cpp
--- a/Engine/Plugins/Runtime/ModelViewViewModel/Source/ModelViewViewModelBlueprint/Private/MVVMBlueprintView.cpp
+++ b/Engine/Plugins/Runtime/ModelViewViewModel/Source/ModelViewViewModelBlueprint/Private/MVVMBlueprintView.cpp
@@ -71,29 +71,16 @@ void UMVVMBlueprintView::SetViewModels(const TArray<FMVVMBlueprintViewModelConte
 
 const FMVVMBlueprintViewBinding* UMVVMBlueprintView::FindBinding(const UWidget* Widget, const FProperty* Property) const
 {
-	return FindBinding(Widget->GetFName(), FMVVMBindingName(Property->GetFName()));
+	return const_cast<UMVVMBlueprintView*>(this)->FindBinding(Widget, Property);
 }
 
 FMVVMBlueprintViewBinding* UMVVMBlueprintView::FindBinding(const UWidget* Widget, const FProperty* Property)
 {
-	return FindBinding(Widget->GetFName(), FMVVMBindingName(Property->GetFName()));
-}
-
-const FMVVMBlueprintViewBinding* UMVVMBlueprintView::FindBinding(FName WidgetName, FMVVMBindingName BindingName) const
-{
-	return Bindings.FindByPredicate([WidgetName, BindingName](const FMVVMBlueprintViewBinding& Binding)
-		{
-			return Binding.WidgetPath.WidgetName == WidgetName &&
-				Binding.WidgetPath.GetBindingName() == BindingName;
-		});
-}
-
-FMVVMBlueprintViewBinding* UMVVMBlueprintView::FindBinding(FName WidgetName, FMVVMBindingName BindingName)
-{
-	return Bindings.FindByPredicate([WidgetName, BindingName](const FMVVMBlueprintViewBinding& Binding)
+	FName WidgetName = Widget->GetFName();
+	return Bindings.FindByPredicate([WidgetName, Property](const FMVVMBlueprintViewBinding& Binding)
 		{
 			return Binding.WidgetPath.WidgetName == WidgetName &&
-				Binding.WidgetPath.GetBindingName() == BindingName;
+				Binding.WidgetPath.BasePropertyPathContains(UE::MVVM::FMVVMConstFieldVariant(Property));
 		});
 }
 
@@ -122,7 +109,7 @@ FMVVMBlueprintViewBinding& UMVVMBlueprintView::AddBinding(const UWidget* Widget,
 {
 	FMVVMBlueprintViewBinding& NewBinding = Bindings.AddDefaulted_GetRef();
 	NewBinding.WidgetPath.WidgetName = Widget->GetFName();
-	NewBinding.WidgetPath.SetBindingReference(UE::MVVM::FMVVMConstFieldVariant(Property));
+	NewBinding.WidgetPath.SetBasePropertyPath(UE::MVVM::FMVVMConstFieldVariant(Property));
 
 	OnBindingsUpdated.Broadcast();
 	return NewBinding;
diff --git a/Engine/Plugins/Runtime/ModelViewViewModel/Source/ModelViewViewModelBlueprint/Private/MVVMViewBlueprintCompiler.cpp b/Engine/Plugins/Runtime/ModelViewViewModel/Source/ModelViewViewModelBlueprint/Private/MVVMViewBlueprintCompiler.cpp
--- a/Engine/Plugins/Runtime/ModelViewViewModel/Source/ModelViewViewModelBlueprint/Private/MVVMViewBlueprintCompiler.cpp
+++ b/Engine/Plugins/Runtime/ModelViewViewModel/Source/ModelViewViewModelBlueprint/Private/MVVMViewBlueprintCompiler.cpp
@@ -550,59 +550,61 @@ bool FMVVMViewBlueprintCompiler::CompileSourceCreators(const FCompiledBindingLib
 		if (SourceCreatorContext.Type == ECompilerSourceCreatorType::ViewModel)
 		{
 			const FMVVMBlueprintViewModelContext& ViewModelContext = SourceCreatorContext.ViewModelContext;
-			if (ViewModelContext.CreationType != EMVVMBlueprintViewModelContextCreationType::Manual)
-			{
-				FMVVMViewClass_SourceCreator CompiledSourceCreator;
 
-				if (ViewModelContext.CreationType == EMVVMBlueprintViewModelContextCreationType::CreateInstance)
+			FMVVMViewClass_SourceCreator CompiledSourceCreator;
+
+			if (ViewModelContext.CreationType == EMVVMBlueprintViewModelContextCreationType::Manual)
+			{
+				CompiledSourceCreator = FMVVMViewClass_SourceCreator::MakeManual(ViewModelContext.GetViewModelName(), ViewModelContext.GetViewModelClass());
+			}
+			else if (ViewModelContext.CreationType == EMVVMBlueprintViewModelContextCreationType::CreateInstance)
+			{
+				CompiledSourceCreator = FMVVMViewClass_SourceCreator::MakeInstance(ViewModelContext.GetViewModelName(), ViewModelContext.GetViewModelClass());
+			}
+			else if (ViewModelContext.CreationType == EMVVMBlueprintViewModelContextCreationType::PropertyPath)
+			{
+				const FMVVMVCompiledFieldPath* CompiledFieldPath = CompileResult.FieldPaths.Find(SourceCreatorContext.ReadPropertyPath);
+				if (CompiledFieldPath == nullptr)
 				{
-					CompiledSourceCreator = FMVVMViewClass_SourceCreator::MakeInstance(ViewModelContext.GetViewModelName(), ViewModelContext.GetViewModelClass());
+					WidgetBlueprintCompilerContext.MessageLog.Error(*FString::Printf(TEXT("The viewmodel '%s' initialization binding was not generated.")
+						, *ViewModelContext.GetViewModelClass()->GetDisplayNameText().ToString()));
+					bAreSourcesCreatorValid = false;
+					continue;
 				}
-				else if (ViewModelContext.CreationType == EMVVMBlueprintViewModelContextCreationType::PropertyPath)
+				CompiledSourceCreator = FMVVMViewClass_SourceCreator::MakeFieldPath(ViewModelContext.GetViewModelName(), ViewModelContext.GetViewModelClass(), *CompiledFieldPath);
+			}
+			else if (ViewModelContext.CreationType == EMVVMBlueprintViewModelContextCreationType::GlobalViewModelCollection)
+			{
+				if (ViewModelContext.GlobalViewModelIdentifier.IsNone())
 				{
-					const FMVVMVCompiledFieldPath* CompiledFieldPath = CompileResult.FieldPaths.Find(SourceCreatorContext.ReadPropertyPath);
-					if (CompiledFieldPath == nullptr)
-					{
-						WidgetBlueprintCompilerContext.MessageLog.Error(*FString::Printf(TEXT("The viewmodel '%s' initialization binding was not generated.")
-							, *ViewModelContext.GetViewModelClass()->GetDisplayNameText().ToString()));
-						bAreSourcesCreatorValid = false;
-						continue;
-					}
-					CompiledSourceCreator = FMVVMViewClass_SourceCreator::MakeFieldPath(ViewModelContext.GetViewModelName(), ViewModelContext.GetViewModelClass(), *CompiledFieldPath);
+					WidgetBlueprintCompilerContext.MessageLog.Error(*FString::Printf(TEXT("The viewmodel '%s' doesn't have a valid Global identifier.")
+						, *ViewModelContext.GetViewModelClass()->GetDisplayNameText().ToString()));
+					bAreSourcesCreatorValid = false;
+					continue;
 				}
-				else if (ViewModelContext.CreationType == EMVVMBlueprintViewModelContextCreationType::GlobalViewModelCollection)
-				{
-					if (ViewModelContext.GlobalViewModelIdentifier.IsNone())
-					{
-						WidgetBlueprintCompilerContext.MessageLog.Error(*FString::Printf(TEXT("The viewmodel '%s' doesn't have a valid Global identifier.")
-							, *ViewModelContext.GetViewModelClass()->GetDisplayNameText().ToString()));
-						bAreSourcesCreatorValid = false;
-						continue;
-					}
-
-					FMVVMViewModelContext GlobalViewModelInstance;
-					GlobalViewModelInstance.ContextClass = ViewModelContext.GetViewModelClass();
-					GlobalViewModelInstance.ContextName = ViewModelContext.GlobalViewModelIdentifier;
-					if (!GlobalViewModelInstance.IsValid())
-					{
-						WidgetBlueprintCompilerContext.MessageLog.Error(*FString::Printf(TEXT("The context for viewmodel '%s' could not be created.")
-							, *ViewModelContext.GetViewModelClass()->GetDisplayNameText().ToString()));
-						bAreSourcesCreatorValid = false;
-						continue;
-					}
 
-					CompiledSourceCreator = FMVVMViewClass_SourceCreator::MakeGlobalContext(ViewModelContext.GetViewModelName(), MoveTemp(GlobalViewModelInstance));
-				}
-				else
+				FMVVMViewModelContext GlobalViewModelInstance;
+				GlobalViewModelInstance.ContextClass = ViewModelContext.GetViewModelClass();
+				GlobalViewModelInstance.ContextName = ViewModelContext.GlobalViewModelIdentifier;
+				if (!GlobalViewModelInstance.IsValid())
 				{
-					WidgetBlueprintCompilerContext.MessageLog.Error(*FString::Printf(TEXT("The viewmodel '%s' doesn't have a valid creation type.")
+					WidgetBlueprintCompilerContext.MessageLog.Error(*FString::Printf(TEXT("The context for viewmodel '%s' could not be created.")
 						, *ViewModelContext.GetViewModelClass()->GetDisplayNameText().ToString()));
 					bAreSourcesCreatorValid = false;
 					continue;
 				}
 
-				ViewExtension->SourceCreators.Add(MoveTemp(CompiledSourceCreator));
+				CompiledSourceCreator = FMVVMViewClass_SourceCreator::MakeGlobalContext(ViewModelContext.GetViewModelName(), MoveTemp(GlobalViewModelInstance));
+			}
+			else
+			{
+				WidgetBlueprintCompilerContext.MessageLog.Error(*FString::Printf(TEXT("The viewmodel '%s' doesn't have a valid creation type.")
+					, *ViewModelContext.GetViewModelClass()->GetDisplayNameText().ToString()));
+				bAreSourcesCreatorValid = false;
+				continue;
 			}
+
+			ViewExtension->SourceCreators.Add(MoveTemp(CompiledSourceCreator));
 		}
 	}
 
@@ -709,13 +711,21 @@ bool FMVVMViewBlueprintCompiler::PreCompileBindings(UWidgetBlueprintGeneratedCla
 
 		if (IsForwardBinding(Binding.BindingType))
 		{
-			// Todo change that for any type of object
-			const FMVVMBlueprintViewModelContext* SourceViewModelContext = BlueprintView->FindViewModel(Binding.ViewModelPath.ContextId);
-			check(SourceViewModelContext);
-			FName SourceName = SourceViewModelContext->GetViewModelName();
-			int32 SourceContextIndex = Self->SourceContexts.IndexOfByPredicate([SourceName](const FCompilerSourceContext& Other) { return Other.PropertyName == SourceName; });
-			check(SourceContextIndex != INDEX_NONE);
-			TValueOrError<FCompiledBindingLibraryCompiler::FFieldIdHandle, FString> AddFieldResult = AddFieldId(SourceContexts[SourceContextIndex].Class, true, Binding.BindingType, Binding.ViewModelPath.GetBindingName().ToName());
+			const FBindingSourceContext BindingSourceContext = CreateBindingSourceContext(BlueprintView, Binding.ViewModelPath);
+			if (BindingSourceContext.bIsRootWidget && BindingSourceContext.SourceClass == nullptr)
+			{
+				AddErrorForBinding(Binding, BlueprintView, TEXT("Internal error. The binding could not find its source."));
+				bIsBindingsValid = false;
+				continue;
+			}
+			else if (BindingSourceContext.bIsRootWidget && (BindingSourceContext.SourceClass == nullptr || BindingSourceContext.CompilerSourceContextIndex == INDEX_NONE))
+			{
+				AddErrorForBinding(Binding, BlueprintView, TEXT("Internal error. The binding could not find its source."));
+				bIsBindingsValid = false;
+				continue;
+			}
+
+			TValueOrError<FCompiledBindingLibraryCompiler::FFieldIdHandle, FString> AddFieldResult = AddFieldId(BindingSourceContext.SourceClass, true, Binding.BindingType, BindingSourceContext.FieldId);
 			if (AddFieldResult.HasError())
 			{
 				AddErrorForBinding(Binding, BlueprintView, *FString::Printf(TEXT("The binding could not create its source. %s"), *AddFieldResult.GetError()));
@@ -730,34 +740,24 @@ bool FMVVMViewBlueprintCompiler::PreCompileBindings(UWidgetBlueprintGeneratedCla
 			//}
 			//else
 			{
-				FString Getter = SourceContexts[SourceContextIndex].PropertyName.ToString();
-				if (!Binding.ViewModelPath.GetGetterPropertyPath().IsEmpty())
-				{
-					Getter.AppendChar(TEXT('.'));
-					Getter.Append(Binding.ViewModelPath.GetGetterPropertyPath());
-				}
-				Getters.Add(MoveTemp(Getter));
+				Getters.Add(BindingSourceContext.PropertyPath);
 			}
 
 			FString Setter;
 			{
 				FName DestinationName = Binding.WidgetPath.WidgetName;
-				checkf(!DestinationName.IsNone(), TEXT("The destination should have been checked and set bAreSourceContextsValid"));
-				if (DestinationName == Class->ClassGeneratedBy->GetFName()) // is the userwidget itself
+				checkf(!DestinationName.IsNone(), TEXT("The destination should have been checked and set bAreSourceContextsValid."));
+				const bool bSourceIsUserWidget = DestinationName == Class->ClassGeneratedBy->GetFName();
+				if (bSourceIsUserWidget)
 				{
-					Setter = Binding.WidgetPath.GetSetterPropertyPath();
+					Setter = CreatePropertyPath(FName(), Binding.WidgetPath.GetBasePropertyPath());
 				}
 				else
 				{
 					const int32 DestinationVariableContextIndex = SourceContexts.IndexOfByPredicate([DestinationName](const FCompilerSourceContext& Other) { return Other.PropertyName == DestinationName; });
 					check(DestinationVariableContextIndex != INDEX_NONE);
 
-					Setter = SourceContexts[DestinationVariableContextIndex].PropertyName.ToString();
-					if (!Binding.WidgetPath.GetSetterPropertyPath().IsEmpty())
-					{
-						Setter.AppendChar(TEXT('.'));
-						Setter.Append(Binding.WidgetPath.GetSetterPropertyPath());
-					}
+					Setter = CreatePropertyPath(SourceContexts[DestinationVariableContextIndex].PropertyName, Binding.WidgetPath.GetBasePropertyPath());
 				}
 			}
 
@@ -771,44 +771,35 @@ bool FMVVMViewBlueprintCompiler::PreCompileBindings(UWidgetBlueprintGeneratedCla
 
 			FCompilerBinding NewBinding = AddBindingResult.StealValue();
 			NewBinding.BindingIndex = Index;
-			NewBinding.SourceContextIndex = SourceContextIndex;
+			NewBinding.SourceContextIndex = BindingSourceContext.CompilerSourceContextIndex;
 			NewBinding.FieldIdHandle = AddFieldResult.StealValue();
+			NewBinding.bSourceIsUserWidget = BindingSourceContext.bIsRootWidget;
+			NewBinding.bFieldIdNeeded = !IsOneTimeBinding(Binding.BindingType);
 			Bindings.Emplace(NewBinding);
 		}
 
 		if (IsBackwardBinding(Binding.BindingType))
 		{
-			// Todo change that for any type of object
-
-			FCompiledBindingLibraryCompiler::FFieldIdHandle FieldIdHandle;
-			int32 SourceContextIndex = INDEX_NONE;
-
-			const FName SourceName = Binding.WidgetPath.WidgetName;
-			checkf(!SourceName.IsNone(), TEXT("The source should have been checked and set bAreSourceContextsValid"));
-			const bool bSourceIsUserWidget = SourceName == Class->ClassGeneratedBy->GetFName();
-			if (bSourceIsUserWidget)
+			const FBindingSourceContext BindingSourceContext = CreateBindingSourceContext(BlueprintView, Class, Binding.WidgetPath);
+			if (BindingSourceContext.bIsRootWidget && BindingSourceContext.SourceClass == nullptr)
 			{
-				TValueOrError<FCompiledBindingLibraryCompiler::FFieldIdHandle, FString> AddFieldIdResult = AddFieldId(Class->ClassGeneratedBy->GetClass(), true, Binding.BindingType, Binding.WidgetPath.GetBindingName().ToName());
-				if (AddFieldIdResult.HasError())
-				{
-					AddErrorForBinding(Binding, BlueprintView, FString::Printf(TEXT("The binding '%d' could not create it's source. %s"), *AddFieldIdResult.GetError()));
-					bIsBindingsValid = false;
-					continue;
-				}
-				FieldIdHandle = AddFieldIdResult.StealValue();
+				AddErrorForBinding(Binding, BlueprintView, TEXT("Internal error. The binding could not find its source."));
+				bIsBindingsValid = false;
+				continue;
 			}
-			else
+			else if (!BindingSourceContext.bIsRootWidget && (BindingSourceContext.SourceClass == nullptr || BindingSourceContext.CompilerSourceContextIndex == INDEX_NONE))
 			{
-				SourceContextIndex = Self->SourceContexts.IndexOfByPredicate([SourceName](const FCompilerSourceContext& Other) { return Other.PropertyName == SourceName; });
-				check(SourceContextIndex != INDEX_NONE);
-				TValueOrError<FCompiledBindingLibraryCompiler::FFieldIdHandle, FString> AddFieldIdResult = AddFieldId(SourceContexts[SourceContextIndex].Class, true, Binding.BindingType, Binding.WidgetPath.GetBindingName().ToName());
-				if (AddFieldIdResult.HasError())
-				{
-					AddErrorForBinding(Binding, BlueprintView, FString::Printf(TEXT("The binding '%d' could not create it's source. %s"), *AddFieldIdResult.GetError()));
-					bIsBindingsValid = false;
-					continue;
-				}
-				FieldIdHandle = AddFieldIdResult.StealValue();
+				AddErrorForBinding(Binding, BlueprintView, TEXT("Internal error. The binding could not find its source."));
+				bIsBindingsValid = false;
+				continue;
+			}
+
+			TValueOrError<FCompiledBindingLibraryCompiler::FFieldIdHandle, FString> AddFieldResult = AddFieldId(BindingSourceContext.SourceClass, true, Binding.BindingType, BindingSourceContext.FieldId);
+			if (AddFieldResult.HasError())
+			{
+				AddErrorForBinding(Binding, BlueprintView, FString::Printf(TEXT("The binding could not create it's source. %s"), *AddFieldResult.GetError()));
+				bIsBindingsValid = false;
+				continue;
 			}
 
 			TArray<FString> Getters;
@@ -818,20 +809,7 @@ bool FMVVMViewBlueprintCompiler::PreCompileBindings(UWidgetBlueprintGeneratedCla
 			//}
 			//else
 			{
-				if (bSourceIsUserWidget)
-				{
-					Getters.Add(Binding.WidgetPath.GetGetterPropertyPath());
-				}
-				else
-				{
-					FString Getter = SourceContexts[SourceContextIndex].PropertyName.ToString();
-					if (!Binding.WidgetPath.GetGetterPropertyPath().IsEmpty())
-					{
-						Getter.AppendChar(TEXT('.'));
-						Getter.Append(Binding.WidgetPath.GetGetterPropertyPath());
-					}
-					Getters.Add(MoveTemp(Getter));
-				}
+				Getters.Add(BindingSourceContext.PropertyPath);
 			}
 
 			FString Setter;
@@ -842,12 +820,7 @@ bool FMVVMViewBlueprintCompiler::PreCompileBindings(UWidgetBlueprintGeneratedCla
 				const int32 DestinationVariableContextIndex = SourceContexts.IndexOfByPredicate([DestinationName](const FCompilerSourceContext& Other) { return Other.PropertyName == DestinationName; });
 				check(DestinationVariableContextIndex != INDEX_NONE);
 
-				Setter = SourceContexts[DestinationVariableContextIndex].PropertyName.ToString();
-				if (!Binding.ViewModelPath.GetSetterPropertyPath().IsEmpty())
-				{
-					Setter.AppendChar(TEXT('.'));
-					Setter.Append(Binding.ViewModelPath.GetSetterPropertyPath());
-				}
+				Setter = CreatePropertyPath(SourceContexts[DestinationVariableContextIndex].PropertyName, Binding.ViewModelPath.GetBasePropertyPath());
 			}
 
 			TValueOrError<FCompilerBinding, FString> AddBindingResult = AddBinding(Class, Getters, Setter, Binding.Conversion.DestinationToSourceFunction);
@@ -860,9 +833,10 @@ bool FMVVMViewBlueprintCompiler::PreCompileBindings(UWidgetBlueprintGeneratedCla
 
 			FCompilerBinding NewBinding = AddBindingResult.StealValue();
 			NewBinding.BindingIndex = Index;
-			NewBinding.SourceContextIndex = SourceContextIndex;
-			NewBinding.bSourceIsUserWidget = bSourceIsUserWidget;
-			NewBinding.FieldIdHandle = FieldIdHandle;
+			NewBinding.SourceContextIndex = BindingSourceContext.CompilerSourceContextIndex;
+			NewBinding.FieldIdHandle = AddFieldResult.StealValue();
+			NewBinding.bSourceIsUserWidget = BindingSourceContext.bIsRootWidget;
+			NewBinding.bFieldIdNeeded = true; // one time from destination to source doesn't exist
 			Bindings.Emplace(NewBinding);
 		}
 	}
@@ -888,7 +862,7 @@ bool FMVVMViewBlueprintCompiler::CompileBindings(const FCompiledBindingLibraryCo
 		NewBinding.SourcePropertyName = SourceContexts[CompileBinding.SourceContextIndex].PropertyName;
 
 		const FMVVMVCompiledFieldId* CompiledFieldId = CompileResult.FieldIds.Find(CompileBinding.FieldIdHandle);
-		if (CompiledFieldId == nullptr)
+		if (CompiledFieldId == nullptr && CompileBinding.bFieldIdNeeded)
 		{
 			WidgetBlueprintCompilerContext.MessageLog.Error(*FString::Printf(TEXT("The field id for binding '%d' was not generated.")
 				, CompileBinding.BindingIndex));
@@ -905,7 +879,7 @@ bool FMVVMViewBlueprintCompiler::CompileBindings(const FCompiledBindingLibraryCo
 			continue;
 		}
 
-		NewBinding.FieldId = *CompiledFieldId;
+		NewBinding.FieldId = CompiledFieldId  ? *CompiledFieldId : FMVVMVCompiledFieldId();
 		NewBinding.Binding = *CompiledBinding;
 		NewBinding.UpdateMode = ViewBinding.UpdateMode;
 
@@ -927,24 +901,83 @@ const FMVVMViewBlueprintCompiler::FCompilerSourceCreatorContext* FMVVMViewBluepr
 	return SourceCreatorContexts.FindByPredicate([Id](const FCompilerSourceCreatorContext& Other) { return Other.Type == ECompilerSourceCreatorType::ViewModel ? Other.ViewModelContext.GetViewModelId() == Id : false; });
 }
 
-//void FMVVMViewBlueprintCompiler::TransformToFunction(FMVVMCachedPropertyPath& PathToTransform, const UFunction* Function) const
-//{
-//	FString StrToTransform = PathToTransform.ToString();
-//
-//	const TCHAR Delim = TEXT('.');
-//	int32 Index = INDEX_NONE;
-//	StrToTransform.FindLastChar(Delim, Index);
-//	if (Index != INDEX_NONE)
-//	{
-//		StrToTransform.RemoveAt(Index + 1, StrToTransform.Len(), false);
-//		StrToTransform.Append(Function->GetName());
-//	}
-//	else
-//	{
-//		StrToTransform = Function->GetName();
-//	}
-//
-//	PathToTransform = StrToTransform;
-//}
+
+FMVVMViewBlueprintCompiler::FBindingSourceContext FMVVMViewBlueprintCompiler::CreateBindingSourceContext(const UMVVMBlueprintView* BlueprintView, const FMVVMViewModelPropertyPath& PropertyPath) const
+{
+	FBindingSourceContext Result;
+
+	// Todo the path may contains an other INotifyFieldValueChanged
+	TArray<FName> Paths = PropertyPath.GetPaths();
+	if (Paths.Num() > 0)
+	{
+		Result.FieldId = Paths[0];
+	}
+
+	{
+		const FMVVMBlueprintViewModelContext* SourceViewModelContext = BlueprintView->FindViewModel(PropertyPath.ContextId);
+		check(SourceViewModelContext);
+		const FName SourceName = SourceViewModelContext->GetViewModelName();
+		Result.CompilerSourceContextIndex = SourceContexts.IndexOfByPredicate([SourceName](const FCompilerSourceContext& Other) { return Other.PropertyName == SourceName; });
+		check(Result.CompilerSourceContextIndex != INDEX_NONE);
+	}
+
+	Result.SourceClass = SourceContexts[Result.CompilerSourceContextIndex].Class;
+	Result.PropertyPath = CreatePropertyPath(SourceContexts[Result.CompilerSourceContextIndex].PropertyName, PropertyPath.GetBasePropertyPath());
+	Result.bIsRootWidget = false;
+
+	return Result;
+}
+
+
+FMVVMViewBlueprintCompiler::FBindingSourceContext FMVVMViewBlueprintCompiler::CreateBindingSourceContext(const UMVVMBlueprintView* BlueprintView, const UWidgetBlueprintGeneratedClass* Class, const FMVVMWidgetPropertyPath& PropertyPath) const
+{
+	FBindingSourceContext Result;
+
+	const FName SourceName = PropertyPath.WidgetName;
+	Result.bIsRootWidget = SourceName == Class->ClassGeneratedBy->GetFName();
+
+	// Todo the path may contains an other INotifyFieldValueChanged
+	TArray<FName> Paths = PropertyPath.GetPaths();
+	if (Paths.Num() > 0)
+	{
+		Result.FieldId = Paths[0];
+	}
+
+	if (Result.bIsRootWidget)
+	{
+		Result.CompilerSourceContextIndex = INDEX_NONE;
+		Result.SourceClass = Class->ClassGeneratedBy->GetClass();
+		Result.PropertyPath = CreatePropertyPath(FName(), PropertyPath.GetBasePropertyPath());
+	}
+	else
+	{
+		Result.CompilerSourceContextIndex = SourceContexts.IndexOfByPredicate([SourceName](const FCompilerSourceContext& Other) { return Other.PropertyName == SourceName; });
+		check(Result.CompilerSourceContextIndex != INDEX_NONE);
+		Result.SourceClass = SourceContexts[Result.CompilerSourceContextIndex].Class;
+		Result.PropertyPath = CreatePropertyPath(SourceContexts[Result.CompilerSourceContextIndex].PropertyName, PropertyPath.GetBasePropertyPath());
+	}
+
+	return Result;
+}
+
+
+FString FMVVMViewBlueprintCompiler::CreatePropertyPath(FName PropertyName, FString PropertyPath) const
+{
+	TStringBuilder<512> Builder;
+	if (PropertyName.IsNone())
+	{
+		Builder << PropertyPath;
+	}
+	else
+	{
+		Builder << PropertyName;
+		if (!PropertyPath.IsEmpty())
+		{
+			Builder << TEXT('.');
+			Builder << PropertyPath;
+		}
+	}
+	return Builder.ToString();
+}
 
 } //namespace
diff --git a/Engine/Plugins/Runtime/ModelViewViewModel/Source/ModelViewViewModelBlueprint/Private/MVVMViewBlueprintCompiler.h b/Engine/Plugins/Runtime/ModelViewViewModel/Source/ModelViewViewModelBlueprint/Private/MVVMViewBlueprintCompiler.h
--- a/Engine/Plugins/Runtime/ModelViewViewModel/Source/ModelViewViewModelBlueprint/Private/MVVMViewBlueprintCompiler.h
+++ b/Engine/Plugins/Runtime/ModelViewViewModel/Source/ModelViewViewModelBlueprint/Private/MVVMViewBlueprintCompiler.h
@@ -11,7 +11,7 @@
 #include "UObject/Object.h"
 #include "WidgetBlueprintCompiler.h"
 
-struct FMVVMCachedPropertyPath;
+struct FMVVMViewModelPropertyPath;
 struct FMVVMBlueprintViewBinding;
 class FWidgetBlueprintCompilerContext;
 class UEdGraph;
@@ -28,6 +28,7 @@ private:
 	struct FCompilerSourceContext;
 	struct FCompilerSourceCreatorContext;
 	struct FCompilerBinding;
+	struct FBindingSourceContext;
 
 public:
 	FMVVMViewBlueprintCompiler(FWidgetBlueprintCompilerContext& InCreationContext)
@@ -69,6 +70,11 @@ private:
 
 	void AddErrorForBinding(FMVVMBlueprintViewBinding& Binding, const UMVVMBlueprintView* BlueprintView, const FString& Message) const;
 
+	FBindingSourceContext CreateBindingSourceContext(const UMVVMBlueprintView* BlueprintView, const FMVVMViewModelPropertyPath& PropertyPath) const;
+	FBindingSourceContext CreateBindingSourceContext(const UMVVMBlueprintView* BlueprintView, const UWidgetBlueprintGeneratedClass* Class, const FMVVMWidgetPropertyPath& PropertyPath) const;
+
+	FString CreatePropertyPath(FName PropertyName, FString PropertyPath) const;
+
 private:
 	struct FCompilerSourceContext
 	{
@@ -105,6 +111,7 @@ private:
 		int32 BindingIndex = INDEX_NONE;
 		int32 SourceContextIndex = INDEX_NONE;
 		bool bSourceIsUserWidget = false;
+		bool bFieldIdNeeded = false;
 
 		FCompiledBindingLibraryCompiler::FBindingHandle BindingHandle;
 
@@ -121,6 +128,19 @@ private:
 	};
 	TArray<FCompilerBinding> Bindings;
 
+	struct FBindingSourceContext
+	{
+		UClass* SourceClass = nullptr;
+		FName FieldId;
+		// ViewModel.Field.SubProperty.SubProperty
+		// or Widget.Field.SubProperty
+		// or Field.SubProperty (if the widget is the userwidget)
+		// if its the destination then the Field is not necessary
+		FString PropertyPath;
+		int32 CompilerSourceContextIndex = INDEX_NONE;
+		bool bIsRootWidget = false;
+	};
+
 	TMap<FName, UWidget*> WidgetNameToWidgetPointerMap;
 	FWidgetBlueprintCompilerContext& WidgetBlueprintCompilerContext;
 	FCompiledBindingLibraryCompiler BindingLibraryCompiler;
diff --git a/Engine/Plugins/Runtime/ModelViewViewModel/Source/ModelViewViewModelBlueprint/Public/MVVMBlueprintView.h b/Engine/Plugins/Runtime/ModelViewViewModel/Source/ModelViewViewModelBlueprint/Public/MVVMBlueprintView.h
--- a/Engine/Plugins/Runtime/ModelViewViewModel/Source/ModelViewViewModelBlueprint/Public/MVVMBlueprintView.h
+++ b/Engine/Plugins/Runtime/ModelViewViewModel/Source/ModelViewViewModelBlueprint/Public/MVVMBlueprintView.h
@@ -36,8 +36,6 @@ public:
 
 	const FMVVMBlueprintViewBinding* FindBinding(const UWidget* Widget, const FProperty* Property) const;
 	FMVVMBlueprintViewBinding* FindBinding(const UWidget* Widget, const FProperty* Property);
-	const FMVVMBlueprintViewBinding* FindBinding(FName WidgetName, FMVVMBindingName BindingName) const;
-	FMVVMBlueprintViewBinding* FindBinding(FName WidgetName, FMVVMBindingName BindingName);
 
 	void RemoveBinding(const FMVVMBlueprintViewBinding* Binding);
 	void RemoveBindingAt(int32 Index);
diff --git a/Engine/Plugins/Runtime/ModelViewViewModel/Source/ModelViewViewModelBlueprint/Public/MVVMBlueprintViewBinding.cpp b/Engine/Plugins/Runtime/ModelViewViewModel/Source/ModelViewViewModelBlueprint/Public/MVVMBlueprintViewBinding.cpp
--- a/Engine/Plugins/Runtime/ModelViewViewModel/Source/ModelViewViewModelBlueprint/Public/MVVMBlueprintViewBinding.cpp
+++ b/Engine/Plugins/Runtime/ModelViewViewModel/Source/ModelViewViewModelBlueprint/Public/MVVMBlueprintViewBinding.cpp
@@ -2,18 +2,59 @@
 
 #include "MVVMBlueprintViewBinding.h"
 #include "MVVMBlueprintView.h"
+#include "MVVMWidgetBlueprintExtension_View.h"
+#include "WidgetBlueprintExtension.h"
+
 
 FString FMVVMBlueprintViewBinding::GetNameString(const UMVVMBlueprintView* View) const
 {
 	const FMVVMBlueprintViewModelContext* ViewModel = View->FindViewModel(ViewModelPath.ContextId);
-	FString BindingName = (ViewModel != nullptr ? ViewModel->GetViewModelName().ToString() : TEXT("<none>")) + TEXT(".") + ViewModelPath.GetBindingName().ToString();
+	TStringBuilder<512> BindingName;
+	if (ViewModel != nullptr)
+	{
+		BindingName << ViewModel->GetViewModelName();
+	}
+	else
+	{
+		BindingName << TEXT("<none>");
+	}
+	BindingName << TEXT(".");
+	BindingName << ViewModelPath.GetBasePropertyPath();
 
-	BindingName += BindingType == EMVVMBindingMode::TwoWay ? TEXT(" <-> ") :
-		UE::MVVM::IsForwardBinding(BindingType) ? TEXT(" -> ") :
-		UE::MVVM::IsBackwardBinding(BindingType) ? TEXT(" <- ") :
-		TEXT("???"); // shouldn't happen
+	if (BindingType == EMVVMBindingMode::TwoWay)
+	{
+		BindingName << TEXT(" <-> ");
+	}
+	else if (UE::MVVM::IsForwardBinding(BindingType))
+	{
+		BindingName << TEXT(" -> ");
+	}
+	else if (UE::MVVM::IsBackwardBinding(BindingType))
+	{
+		BindingName << TEXT(" <- ");
+	}
+	else
+	{
+		BindingName << TEXT(" ??? "); // shouldn't happen
+	}
 
-	BindingName += WidgetPath.WidgetName.ToString() + TEXT(".") + WidgetPath.GetBindingName().ToString();
+	if (View->GetOuterUMVVMWidgetBlueprintExtension_View()->GetWidgetBlueprint()->GetFName() == WidgetPath.WidgetName)
+	{
+		BindingName << WidgetPath.GetBasePropertyPath();
+	}
+	else
+	{
+		if (WidgetPath.WidgetName.IsNone())
+		{
+			BindingName << TEXT("<none>");
+		}
+		else
+		{
+			BindingName << WidgetPath.WidgetName;
+		}
+		BindingName << TEXT(".");
+		BindingName << WidgetPath.GetBasePropertyPath();
+	}
 
-	return BindingName;
-}
\ No newline at end of file
+	return BindingName.ToString();
+}
diff --git a/Engine/Plugins/Runtime/ModelViewViewModel/Source/ModelViewViewModelBlueprint/Public/MVVMPropertyPath.h b/Engine/Plugins/Runtime/ModelViewViewModel/Source/ModelViewViewModelBlueprint/Public/MVVMPropertyPath.h
--- a/Engine/Plugins/Runtime/ModelViewViewModel/Source/ModelViewViewModelBlueprint/Public/MVVMPropertyPath.h
+++ b/Engine/Plugins/Runtime/ModelViewViewModel/Source/ModelViewViewModelBlueprint/Public/MVVMPropertyPath.h
@@ -9,12 +9,10 @@
 #include "MVVMPropertyPath.generated.h"
 
 /**
- * Base path to properties for MVVM view models and widgets.
- * 
- * Used to associate properties within MVVM bindings in editor & during MVVM compilation
+ * A single item in a Property Path
  */
 USTRUCT()
-struct FMVVMPropertyPathBase
+struct FMVVMBlueprintFieldPath
 {
 	GENERATED_BODY()
 
@@ -29,7 +27,7 @@ private:
 
 public:
 	/** Get the binding name, resolves reference deprecation / redirectors / etc before returning */
-	FMVVMBindingName GetBindingName() const
+	FName GetFieldName() const
 	{
 		// Resolve any redirectors 
 		FName Path = FName();
@@ -45,7 +43,24 @@ public:
 			}
 		}
 
-		return FMVVMBindingName(Path);
+		return Path;
+	}
+
+	/** */
+	UE::MVVM::FMVVMConstFieldVariant GetField() const
+	{
+		if (!BindingReference.GetMemberName().IsNone())
+		{
+			if (BindingKind == EBindingKind::Property)
+			{
+				return UE::MVVM::FMVVMConstFieldVariant(BindingReference.ResolveMember<FProperty>());
+			}
+			else if (BindingKind == EBindingKind::Function)
+			{
+				return UE::MVVM::FMVVMConstFieldVariant(BindingReference.ResolveMember<UFunction>());
+			}
+		}
+		return UE::MVVM::FMVVMConstFieldVariant();
 	}
 
 	FMemberReference GetBindingReference() const
@@ -54,57 +69,171 @@ public:
 		return BindingReference;
 	}
 
-	void SetBindingReference(const UE::MVVM::FMVVMConstFieldVariant& InField)
+	void SetBindingReference(UE::MVVM::FMVVMConstFieldVariant InField)
 	{
 		if (InField.IsEmpty())
 		{
 			Reset();
-			return;
 		}
+		else
+		{
+			BindingReference = InField.CreateMemberReference();
+
+			if (InField.IsProperty())
+			{
+				BindingKind = EBindingKind::Property;
+			}
+			else if (InField.IsFunction())
+			{
+				BindingKind = EBindingKind::Function;
+			}
+			else
+			{
+				ensureAlwaysMsgf(false, TEXT("Binding to field of unknown type!"));
+			}
+		}
+	}
+
+	void Reset()
+	{
+		BindingReference = FMemberReference();
+	}
+
+	void SetDeprecatedBindingReference(const FMemberReference& InBindingReference, EBindingKind InBindingKind)
+	{
+		BindingReference = InBindingReference;
+		BindingKind = InBindingKind;
+	}
+
+public:
+	bool operator==(const FMVVMBlueprintFieldPath& Other) const
+	{
+		return BindingReference.GetMemberName() == Other.BindingReference.GetMemberName();
+	}
+	bool operator!=(const FMVVMBlueprintFieldPath& Other) const
+	{
+		return !operator==(Other);
+	}
+};
+
+
+/**
+ * Base path to properties for MVVM view models and widgets.
+ * 
+ * Used to associate properties within MVVM bindings in editor & during MVVM compilation
+ */
+USTRUCT()
+struct FMVVMBlueprintPropertyPathBase
+{
+	GENERATED_BODY()
 
-		BindingReference = InField.CreateMemberReference();
+private:
+	/** Reference to property for this binding. */
+	UPROPERTY(EditAnywhere, Category = "MVVM")
+	TArray<FMVVMBlueprintFieldPath> Paths;
+
+#if WITH_EDITORONLY_DATA
+	// Use the Paths. BindingReference and BindingKind are deprecated.
+	UPROPERTY()
+	FMemberReference BindingReference;
+	UPROPERTY()
+	EBindingKind BindingKind;
+#endif
 
-		if (InField.IsProperty())
+public:
+	/** Get the binding name, resolves reference deprecation / redirectors / etc before returning */
+	TArray<FName> GetPaths() const
+	{
+		TArray<FName> Result;
+		Result.Reserve(Paths.Num());
+
+		for (const FMVVMBlueprintFieldPath& Path : Paths)
 		{
-			BindingKind = EBindingKind::Property;
+			Result.Add(Path.GetFieldName());
 		}
-		else if (InField.IsFunction())
+
+		return Result;
+	}
+
+	bool BasePropertyPathContains(UE::MVVM::FMVVMConstFieldVariant Field) const
+	{
+		return Paths.ContainsByPredicate([Field](const FMVVMBlueprintFieldPath& FieldPath) { return FieldPath.GetField() == Field;  });
+	}
+
+	void SetBasePropertyPath(UE::MVVM::FMVVMConstFieldVariant InField)
+	{
+		Paths.Reset();
+		Paths.AddDefaulted_GetRef().SetBindingReference(InField);
+
+		if (InField.IsProperty() && InField.GetProperty()->IsA(FStructProperty::StaticClass()))
 		{
-			BindingKind = EBindingKind::Function;
+			if (FProperty* Pro = ((FStructProperty*)InField.GetProperty())->Struct->FindPropertyByName("A"))
+				Paths.AddDefaulted_GetRef().SetBindingReference(UE::MVVM::FMVVMConstFieldVariant(Pro));
 		}
-		else
+		if (InField.IsProperty() && InField.GetProperty()->IsA(FObjectPropertyBase::StaticClass()))
 		{
-			ensureAlwaysMsgf(false, TEXT("Binding to field of unknown type!"));
+			if (FProperty* Pro = ((FObjectPropertyBase*)InField.GetProperty())->PropertyClass->FindPropertyByName("Dummy"))
+				Paths.AddDefaulted_GetRef().SetBindingReference(UE::MVVM::FMVVMConstFieldVariant(Pro));
 		}
 	}
 
-	void Reset()
+	void ResetBasePropertyPath()
 	{
-		BindingReference = FMemberReference();
+		Paths.Reset();
 	}
 
-	FString GetGetterPropertyPath() const
+	/**
+	 * Get the full path without the first property name.
+	 * returns Field.SubProperty.SubProperty from ViewModel.Field.SubProeprty.SubProperty
+	 */
+	FString GetBasePropertyPath() const
 	{
-		return GetBindingName().ToString();
+		TStringBuilder<512> Result;
+		for (const FMVVMBlueprintFieldPath& Path : Paths)
+		{
+			if (Result.Len() > 0)
+			{
+				Result << TEXT('.');
+			}
+			Result << Path.GetFieldName().ToString();
+		}
+		return Result.ToString();
 	}
 
-	FString GetSetterPropertyPath() const
+public:
+	bool operator==(const FMVVMBlueprintPropertyPathBase& Other) const
 	{
-		return GetBindingName().ToString();
+		return Paths == Other.Paths;
 	}
-
-	bool operator==(const FMVVMPropertyPathBase& Other) const
+	bool operator!=(const FMVVMBlueprintPropertyPathBase& Other) const
 	{
-		return BindingReference.GetMemberName() == Other.BindingReference.GetMemberName();
+		return !operator==(Other);
 	}
-	bool operator!=(const FMVVMPropertyPathBase& Other) const
+
+public:
+	void PostSerialize(const FArchive& Ar)
 	{
-		return !operator==(Other);
+		if (Ar.IsLoading() && !BindingReference.GetMemberName().IsNone())
+		{
+			Paths.AddDefaulted_GetRef().SetDeprecatedBindingReference(BindingReference, BindingKind);
+			BindingReference = FMemberReference();
+		}
 	}
 };
 
+template<>
+struct TStructOpsTypeTraits<FMVVMBlueprintPropertyPathBase> : public TStructOpsTypeTraitsBase2<FMVVMBlueprintPropertyPathBase>
+{
+	enum
+	{
+		WithPostSerialize = true,
+	};
+};
+
+
+/** */
 USTRUCT(BlueprintType)
-struct FMVVMViewModelPropertyPath : public FMVVMPropertyPathBase
+struct FMVVMViewModelPropertyPath : public FMVVMBlueprintPropertyPathBase
 {
 	GENERATED_BODY();
 
@@ -114,18 +243,30 @@ struct FMVVMViewModelPropertyPath : public FMVVMPropertyPathBase
 	bool operator==(const FMVVMViewModelPropertyPath& Other) const
 	{
 		return ContextId == Other.ContextId &&
-			FMVVMPropertyPathBase::operator==(Other);
+			FMVVMBlueprintPropertyPathBase::operator==(Other);
 	}
 
 	bool operator!=(const FMVVMViewModelPropertyPath& Other) const
 	{
 		return ContextId != Other.ContextId &&
-			FMVVMPropertyPathBase::operator!=(Other);
+			FMVVMBlueprintPropertyPathBase::operator!=(Other);
 	}
 };
 
+
+template<>
+struct TStructOpsTypeTraits<FMVVMViewModelPropertyPath> : public TStructOpsTypeTraitsBase2<FMVVMViewModelPropertyPath>
+{
+	enum
+	{
+		WithPostSerialize = true,
+	};
+};
+
+
+/** */
 USTRUCT(BlueprintType)
-struct FMVVMWidgetPropertyPath : public FMVVMPropertyPathBase
+struct FMVVMWidgetPropertyPath : public FMVVMBlueprintPropertyPathBase
 {
 	GENERATED_BODY();
 
@@ -136,13 +277,23 @@ struct FMVVMWidgetPropertyPath : public FMVVMPropertyPathBase
 	bool operator==(const FMVVMWidgetPropertyPath& Other) const
 	{
 		return WidgetName == Other.WidgetName &&
-			FMVVMPropertyPathBase::operator==(Other);
+			FMVVMBlueprintPropertyPathBase::operator==(Other);
 	}
 
 
 	bool operator!=(const FMVVMWidgetPropertyPath& Other) const
 	{
 		return WidgetName != Other.WidgetName ||
-			FMVVMPropertyPathBase::operator!=(Other);
+			FMVVMBlueprintPropertyPathBase::operator!=(Other);
 	}
 };
+
+
+template<>
+struct TStructOpsTypeTraits<FMVVMWidgetPropertyPath> : public TStructOpsTypeTraitsBase2<FMVVMWidgetPropertyPath>
+{
+	enum
+	{
+		WithPostSerialize = true,
+	};
+};
diff --git a/Engine/Plugins/Runtime/ModelViewViewModel/Source/ModelViewViewModelEditor/Private/Customizations/MVVMPropertyBindingExtension.cpp b/Engine/Plugins/Runtime/ModelViewViewModel/Source/ModelViewViewModelEditor/Private/Customizations/MVVMPropertyBindingExtension.cpp
--- a/Engine/Plugins/Runtime/ModelViewViewModel/Source/ModelViewViewModelEditor/Private/Customizations/MVVMPropertyBindingExtension.cpp
+++ b/Engine/Plugins/Runtime/ModelViewViewModel/Source/ModelViewViewModelEditor/Private/Customizations/MVVMPropertyBindingExtension.cpp
@@ -81,10 +81,10 @@ static void ExtendBindingsMenu(FMenuBuilder& MenuBuilder, const UWidgetBlueprint
 		FMVVMBlueprintViewBinding& NewBinding = MVVMBlueprintView->AddDefaultBinding();
 
 		NewBinding.ViewModelPath.ContextId = ViewModelId;
-		NewBinding.ViewModelPath.SetBindingReference(UE::MVVM::FMVVMConstFieldVariant(ViewModelProperty));
+		NewBinding.ViewModelPath.SetBasePropertyPath(UE::MVVM::FMVVMConstFieldVariant(ViewModelProperty));
 
 		NewBinding.WidgetPath.WidgetName = Widget->GetFName();
-		NewBinding.WidgetPath.SetBindingReference(UE::MVVM::FMVVMConstFieldVariant(WidgetProperty));
+		NewBinding.WidgetPath.SetBasePropertyPath(UE::MVVM::FMVVMConstFieldVariant(WidgetProperty));
 
 		NewBinding.BindingType = EMVVMBindingMode::OneWayToDestination;
 	};
@@ -160,7 +160,12 @@ TOptional<FName> FMVVMPropertyBindingExtension::GetCurrentValue(const UWidgetBlu
 		return TOptional<FName>();
 	}
 
-	return Binding->ViewModelPath.GetBindingName().ToName();
+	TArray<FName> Names = Binding->ViewModelPath.GetPaths();
+	if (Names.Num() > 0)
+	{
+		return Names[0];
+	}
+	return TOptional<FName>();
 }
 
 
@@ -172,7 +177,7 @@ void FMVVMPropertyBindingExtension::ClearCurrentValue(const UWidgetBlueprint* Wi
 		{
 			if (FMVVMBlueprintViewBinding* Binding = MVVMBlueprintView->FindBinding(Widget, Property))
 			{
-				Binding->ViewModelPath.Reset();
+				Binding->ViewModelPath.ResetBasePropertyPath();
 			}
 		}
 	}
diff --git a/Engine/Plugins/Runtime/ModelViewViewModel/Source/ModelViewViewModelEditor/Private/MVVMPropertyPathHelpers.cpp b/Engine/Plugins/Runtime/ModelViewViewModel/Source/ModelViewViewModelEditor/Private/MVVMPropertyPathHelpers.cpp
--- a/Engine/Plugins/Runtime/ModelViewViewModel/Source/ModelViewViewModelEditor/Private/MVVMPropertyPathHelpers.cpp
+++ b/Engine/Plugins/Runtime/ModelViewViewModel/Source/ModelViewViewModelEditor/Private/MVVMPropertyPathHelpers.cpp
@@ -158,27 +158,21 @@ namespace UE::MVVM
 	{
 		if (FMVVMWidgetPropertyPath* Path = PathAttr.Get(nullptr))
 		{
-			return Path->GetBindingName();
+			TArray<FName> Paths = Path->GetPaths();
+			if (Paths.Num() > 0)
+			{
+				return FMVVMBindingName(Paths[0]);
+			}
 		}
 
 		return FMVVMBindingName();
 	}
 
-	FMemberReference FWidgetFieldPathHelper::GetBindingReference() const
-	{
-		if (FMVVMWidgetPropertyPath* Path = PathAttr.Get(nullptr))
-		{
-			return Path->GetBindingReference();
-		}
-
-		return FMemberReference();
-	}
-
 	void FWidgetFieldPathHelper::SetBindingReference(const UE::MVVM::FMVVMFieldVariant& InField) const
 	{
 		if (FMVVMWidgetPropertyPath* Path = PathAttr.Get(nullptr))
 		{
-			Path->SetBindingReference(InField);
+			Path->SetBasePropertyPath(InField);
 		}
 	}
 
@@ -186,7 +180,7 @@ namespace UE::MVVM
 	{
 		if (FMVVMWidgetPropertyPath* Path = PathAttr.Get(nullptr))
 		{
-			Path->SetBindingReference(InField);
+			Path->SetBasePropertyPath(InField);
 		}
 	}
 
@@ -194,7 +188,7 @@ namespace UE::MVVM
 	{
 		if (FMVVMWidgetPropertyPath* Path = PathAttr.Get(nullptr))
 		{
-			Path->Reset();
+			Path->ResetBasePropertyPath();
 		}
 	}
 
@@ -304,27 +298,21 @@ namespace UE::MVVM
 	{
 		if (FMVVMViewModelPropertyPath* Path = PathAttr.Get(nullptr))
 		{
-			return Path->GetBindingName();
+			TArray<FName> Paths = Path->GetPaths();
+			if (Paths.Num() > 0)
+			{
+				return FMVVMBindingName(Paths[0]);
+			}
 		}
 
 		return FMVVMBindingName();
 	}
 
-	FMemberReference FViewModelFieldPathHelper::GetBindingReference() const
-	{
-		if (FMVVMViewModelPropertyPath* Path = PathAttr.Get(nullptr))
-		{
-			return Path->GetBindingReference();
-		}
-
-		return FMemberReference();
-	}
-
 	void FViewModelFieldPathHelper::SetBindingReference(const UE::MVVM::FMVVMFieldVariant& InField) const
 	{
 		if (FMVVMViewModelPropertyPath* Path = PathAttr.Get(nullptr))
 		{
-			Path->SetBindingReference(InField);
+			Path->SetBasePropertyPath(InField);
 		}
 	}
 
@@ -332,7 +320,7 @@ namespace UE::MVVM
 	{
 		if (FMVVMViewModelPropertyPath* Path = PathAttr.Get(nullptr))
 		{
-			Path->SetBindingReference(InField);
+			Path->SetBasePropertyPath(InField);
 		}
 	}
 
@@ -340,7 +328,7 @@ namespace UE::MVVM
 	{
 		if (FMVVMViewModelPropertyPath* Path = PathAttr.Get(nullptr))
 		{
-			Path->Reset();
+			Path->ResetBasePropertyPath();
 		}
 	}
 }
diff --git a/Engine/Plugins/Runtime/ModelViewViewModel/Source/ModelViewViewModelEditor/Private/MVVMPropertyPathHelpers.h b/Engine/Plugins/Runtime/ModelViewViewModel/Source/ModelViewViewModelEditor/Private/MVVMPropertyPathHelpers.h
--- a/Engine/Plugins/Runtime/ModelViewViewModel/Source/ModelViewViewModelEditor/Private/MVVMPropertyPathHelpers.h
+++ b/Engine/Plugins/Runtime/ModelViewViewModel/Source/ModelViewViewModelEditor/Private/MVVMPropertyPathHelpers.h
@@ -25,7 +25,6 @@ namespace UE::MVVM
 		virtual void SetSelectedSource(const FBindingSource& Source) const = 0;
 
 		virtual FMVVMBindingName GetBindingName() const = 0;
-		virtual FMemberReference GetBindingReference() const = 0;
 		virtual void SetBindingReference(const UE::MVVM::FMVVMFieldVariant& InField) const = 0;
 		virtual void SetBindingReference(const UE::MVVM::FMVVMConstFieldVariant& InField) const = 0;
 		virtual void ResetBinding() const = 0;
@@ -50,7 +49,6 @@ namespace UE::MVVM
 		virtual void SetSelectedSource(const FBindingSource& Source) const override;
 
 		virtual FMVVMBindingName GetBindingName() const override;
-		virtual FMemberReference GetBindingReference() const override;
 		virtual void SetBindingReference(const UE::MVVM::FMVVMFieldVariant& InField) const override;
 		virtual void SetBindingReference(const UE::MVVM::FMVVMConstFieldVariant& InField) const override;
 		virtual void ResetBinding() const override;
@@ -79,7 +77,6 @@ namespace UE::MVVM
 		virtual void SetSelectedSource(const FBindingSource& Source) const override;
 
 		virtual FMVVMBindingName GetBindingName() const override;
-		virtual FMemberReference GetBindingReference() const override;
 		virtual void SetBindingReference(const UE::MVVM::FMVVMFieldVariant& InField) const override;
 		virtual void SetBindingReference(const UE::MVVM::FMVVMConstFieldVariant& InField) const override;
 		virtual void ResetBinding() const override;
diff --git a/Engine/Plugins/Runtime/ModelViewViewModel/Source/ModelViewViewModelEditor/Private/Widgets/SMVVMViewModelContextListWidget.cpp b/Engine/Plugins/Runtime/ModelViewViewModel/Source/ModelViewViewModelEditor/Private/Widgets/SMVVMViewModelContextListWidget.cpp
--- a/Engine/Plugins/Runtime/ModelViewViewModel/Source/ModelViewViewModelEditor/Private/Widgets/SMVVMViewModelContextListWidget.cpp
+++ b/Engine/Plugins/Runtime/ModelViewViewModel/Source/ModelViewViewModelEditor/Private/Widgets/SMVVMViewModelContextListWidget.cpp
@@ -800,7 +800,7 @@ bool SMVVMViewModelContextListWidget::ValidateRemoveViewModelContext(TSharedPtr<
 			.AutoWidth()
 			[
 				SNew(STextBlock)
-				.Text(FText::FromString(Binding.ViewModelPath.GetBindingName().ToString()))
+				.Text(FText::FromString(Binding.ViewModelPath.GetBasePropertyPath()))
 			];
 			Container->AddSlot()
 			.Padding(FMargin(3.0f, 0.0f))
@@ -821,7 +821,7 @@ bool SMVVMViewModelContextListWidget::ValidateRemoveViewModelContext(TSharedPtr<
 			.AutoWidth()
 			[
 				SNew(STextBlock)
-				.Text(FText::FormatOrdered(INVTEXT("{0}.{1}"), FText::FromString(Binding.WidgetPath.WidgetName.ToString()), FText::FromString(Binding.WidgetPath.GetBindingName().ToString())))
+				.Text(FText::FormatOrdered(INVTEXT("{0}.{1}"), FText::FromString(Binding.WidgetPath.WidgetName.ToString()), FText::FromString(Binding.WidgetPath.GetBasePropertyPath())))
 			];
 			ConflictList.Add(Container);
 		}
