MVVM: Fix bug in ViewModelBlueprint: the function graph are removed correctly and the we now use the blueprintsetter.

diff --git a/Engine/Plugins/Runtime/ModelViewViewModel/Source/ModelViewViewModelBlueprint/Private/ViewModel/MVVMViewModelBlueprintCompiler.cpp b/Engine/Plugins/Runtime/ModelViewViewModel/Source/ModelViewViewModelBlueprint/Private/ViewModel/MVVMViewModelBlueprintCompiler.cpp
--- a/Engine/Plugins/Runtime/ModelViewViewModel/Source/ModelViewViewModelBlueprint/Private/ViewModel/MVVMViewModelBlueprintCompiler.cpp
+++ b/Engine/Plugins/Runtime/ModelViewViewModel/Source/ModelViewViewModelBlueprint/Private/ViewModel/MVVMViewModelBlueprintCompiler.cpp
@@ -108,25 +108,25 @@ void FViewModelBlueprintCompilerContext::CreateFunctionList()
 {
 	for (const FGeneratedFunction& Function : GeneratedFunctions)
 	{
-		if (Function.Property == nullptr || Function.SkelProperty == nullptr)
+		if (Function.Property == nullptr && Function.SkelProperty == nullptr)
 		{
-			MessageLog.Error(TEXT("The property for the setter function for '%s' could not be generated."), *Function.PropertyName.ToString());
+			MessageLog.Error(*FString::Printf(TEXT("Internal error. The property '%s' is invalid. Setter and net rep was not generated."), *Function.PropertyName.ToString()));
 			continue;
 		}
 
 		if (Function.SetterFunction)
 		{
-			if (!FunctionGraphHelper::GenerateViewModelFieldNotifySetter(*this, Function.SetterFunction, Function.Property, TEXT("NewValue")))
+			if (!FunctionGraphHelper::GenerateViewModelFieldNotifySetter(*this, Function.SetterFunction, Function.SkelProperty, TEXT("NewValue")))
 			{
-				MessageLog.Error(TEXT("The setter function for '%s' could not be generated."), *Function.PropertyName.ToString());
+				MessageLog.Error(*FString::Printf(TEXT("The setter function for '%s' could not be generated."), *Function.PropertyName.ToString()));
 				continue;
 			}
 		}
 		if (Function.NetRepFunction)
 		{
-			if (!FunctionGraphHelper::GenerateViewModelFielNotifyBroadcast(*this, Function.NetRepFunction, Function.Property))
+			if (!FunctionGraphHelper::GenerateViewModelFielNotifyBroadcast(*this, Function.NetRepFunction, Function.SkelProperty))
 			{
-				MessageLog.Error(TEXT("The net rep function for '%s' could not be generated."), *Function.PropertyName.ToString());
+				MessageLog.Error(*FString::Printf(TEXT("The net rep function for '%s' could not be generated."), *Function.PropertyName.ToString()));
 				continue;
 			}
 		}
@@ -151,7 +151,7 @@ void FViewModelBlueprintCompilerContext::CreateClassVariablesFromBlueprint()
 		bool bGenerateReplicatedFunction = CVarAutogenerateOnRep.GetValueOnGameThread(); // Property->RepNotifyFunc.IsNone();
 		if (bGenerateSetterFunction || bGenerateReplicatedFunction)
 		{
-			Property->SetPropertyFlags(CPF_BlueprintReadOnly | CPF_DisableEditOnInstance);
+			Property->SetPropertyFlags(CPF_DisableEditOnInstance);
 
 			const FString SetterName = FString::Printf(TEXT("Set%s"), *Property->GetName());
 			const FString NetRepName = FString::Printf(TEXT("__RepNotify_%s"), *Property->GetName());
@@ -193,6 +193,12 @@ void FViewModelBlueprintCompilerContext::CreateClassVariablesFromBlueprint()
 						continue;
 					}
 
+					if (GeneratedFunction.SetterFunction->GetFName() != *SetterName)
+					{
+						MessageLog.Error(*FString::Printf(TEXT("The setter name '%s' was generated with a different name."), *SetterName));
+						continue;
+					}
+
 					FunctionGraphHelper::AddFunctionArgument(GeneratedFunction.SetterFunction, Property, TEXT("NewValue"));
 				}
 
@@ -207,15 +213,19 @@ void FViewModelBlueprintCompilerContext::CreateClassVariablesFromBlueprint()
 
 					if (GeneratedFunction.NetRepFunction == nullptr)
 					{
-						MessageLog.Error(*FString::Printf(TEXT("The net replication function name '%s' already exists and could not be autogenerated."), *NetRepName));
+						MessageLog.Error(*FString::Printf(TEXT("The netrep function name '%s' already exists and could not be autogenerated."), *NetRepName));
 						continue;
 					}
-					else
+
+					if (GeneratedFunction.NetRepFunction->GetFName() != *NetRepName)
 					{
-						Property->SetPropertyFlags(CPF_Net);
-						Property->RepNotifyFunc = GeneratedFunction.NetRepFunction->GetFName();
+						MessageLog.Error(*FString::Printf(TEXT("The netrep function name '%s' was generated with a different name."), *NetRepName));
+						continue;
 					}
 
+					Property->SetPropertyFlags(CPF_Net);
+					Property->RepNotifyFunc = GeneratedFunction.NetRepFunction->GetFName();
+
 					FunctionGraphHelper::AddFunctionArgument(GeneratedFunction.SetterFunction, Property, TEXT("NewValue"));
 				}
 
@@ -270,6 +280,16 @@ void FViewModelBlueprintCompilerContext::FinishCompilingClass(UClass* Class)
 {
 	Super::FinishCompilingClass(Class);
 
+	auto DeleteGraph = [](UEdGraph* Graph)
+	{
+		if (Graph)
+		{
+			ensureMsgf(Graph->HasAnyFlags(RF_Transient), TEXT("The graph should be temporary and should be generated automatically."));
+			// GC may not have clean the graph (GC doesn't run when bRegenerateSkeletonOnly is on)
+			Graph->Rename(nullptr, Graph->GetOuter(), REN_DoNotDirty | REN_ForceNoResetLoaders);
+		}
+	};
+
 	const bool bIsSkeletonOnly = CompileOptions.CompileType == EKismetCompileType::SkeletonOnly;
 	if (!bIsSkeletonOnly)
 	{
@@ -280,22 +300,12 @@ void FViewModelBlueprintCompilerContext::FinishCompilingClass(UClass* Class)
 				BPGClass->InitializeFieldNotification(ViewModelBase);
 			}
 		}
+	}
 
-		auto DeleteGraph = [](UEdGraph* Graph)
-		{
-			if (Graph)
-			{
-				ensureMsgf(Graph->HasAnyFlags(RF_Transient), TEXT("The graph should be temporary and should be generated automatically."));
-				// GC may not have clean the graph (GC doesn't run when bRegenerateSkeletonOnly is on)
-				Graph->Rename(nullptr, Graph->GetOuter(), REN_DoNotDirty | REN_ForceNoResetLoaders);
-			}
-		};
-
-		for (FGeneratedFunction& GeneratedFunction : GeneratedFunctions)
-		{
-			DeleteGraph(GeneratedFunction.SetterFunction);
-			DeleteGraph(GeneratedFunction.NetRepFunction);
-		}
+	for (FGeneratedFunction& GeneratedFunction : GeneratedFunctions)
+	{
+		DeleteGraph(GeneratedFunction.SetterFunction);
+		DeleteGraph(GeneratedFunction.NetRepFunction);
 	}
 }
 
