MVVM: Remove test code introduced with CL 20259322

diff --git a/Engine/Plugins/Runtime/ModelViewViewModel/Source/ModelViewViewModelBlueprint/Public/MVVMPropertyPath.h b/Engine/Plugins/Runtime/ModelViewViewModel/Source/ModelViewViewModelBlueprint/Public/MVVMPropertyPath.h
--- a/Engine/Plugins/Runtime/ModelViewViewModel/Source/ModelViewViewModelBlueprint/Public/MVVMPropertyPath.h
+++ b/Engine/Plugins/Runtime/ModelViewViewModel/Source/ModelViewViewModelBlueprint/Public/MVVMPropertyPath.h
@@ -164,17 +164,6 @@ public:
 	{
 		Paths.Reset();
 		Paths.AddDefaulted_GetRef().SetBindingReference(InField);
-
-		if (InField.IsProperty() && InField.GetProperty()->IsA(FStructProperty::StaticClass()))
-		{
-			if (FProperty* Pro = ((FStructProperty*)InField.GetProperty())->Struct->FindPropertyByName("A"))
-				Paths.AddDefaulted_GetRef().SetBindingReference(UE::MVVM::FMVVMConstFieldVariant(Pro));
-		}
-		if (InField.IsProperty() && InField.GetProperty()->IsA(FObjectPropertyBase::StaticClass()))
-		{
-			if (FProperty* Pro = ((FObjectPropertyBase*)InField.GetProperty())->PropertyClass->FindPropertyByName("Dummy"))
-				Paths.AddDefaulted_GetRef().SetBindingReference(UE::MVVM::FMVVMConstFieldVariant(Pro));
-		}
 	}
 
 	void ResetBasePropertyPath()
