MVVM: Remove test code introduced with CL 20259322

diff --git a/Engine/Plugins/Runtime/ModelViewViewModel/Source/ModelViewViewModel/Private/Bindings/MVVMCompiledBindingLibrary.cpp b/Engine/Plugins/Runtime/ModelViewViewModel/Source/ModelViewViewModel/Private/Bindings/MVVMCompiledBindingLibrary.cpp
--- a/Engine/Plugins/Runtime/ModelViewViewModel/Source/ModelViewViewModel/Private/Bindings/MVVMCompiledBindingLibrary.cpp
+++ b/Engine/Plugins/Runtime/ModelViewViewModel/Source/ModelViewViewModel/Private/Bindings/MVVMCompiledBindingLibrary.cpp
@@ -6,7 +6,6 @@
 #include "Blueprint/UserWidget.h"
 #include "Blueprint/WidgetTree.h"
 #include "UObject/PropertyAccessUtil.h"
-PRAGMA_DISABLE_OPTIMIZATION
 
 /**
  *
@@ -440,4 +439,3 @@ TValueOrError<UE::FieldNotification::FFieldId, void> FMVVMCompiledBindingLibrary
 	check(LoadedFieldIds.IsValidIndex(InFieldId.FieldIdIndex));
 	return MakeValue(LoadedFieldIds[InFieldId.FieldIdIndex]);
 }
-PRAGMA_ENABLE_OPTIMIZATION
\ No newline at end of file
