MVVM: SMVVMFieldSelector now has an optional IsFieldValid delegate, instead of passing in CounterpartHelpers.

#rnx

diff --git a/Engine/Plugins/Runtime/ModelViewViewModel/Source/ModelViewViewModelEditor/Private/Widgets/SMVVMFieldSelector.cpp b/Engine/Plugins/Runtime/ModelViewViewModel/Source/ModelViewViewModelEditor/Private/Widgets/SMVVMFieldSelector.cpp
--- a/Engine/Plugins/Runtime/ModelViewViewModel/Source/ModelViewViewModelEditor/Private/Widgets/SMVVMFieldSelector.cpp
+++ b/Engine/Plugins/Runtime/ModelViewViewModel/Source/ModelViewViewModelEditor/Private/Widgets/SMVVMFieldSelector.cpp
@@ -75,6 +75,7 @@ void SMVVMFieldSelector::Construct(const FArguments& InArgs)
 {
 	PathHelpers = InArgs._PathHelpers;
 	CounterpartHelpers = InArgs._CounterpartHelpers;
+	OnValidateFieldDelegate = InArgs._OnValidateField;
 	OnSelectionChangedDelegate = InArgs._OnSelectionChanged;
 	BindingMode = InArgs._BindingMode;
 	bIsSource = InArgs._IsSource;
@@ -118,7 +119,7 @@ void SMVVMFieldSelector::Construct(const FArguments& InArgs)
 					SAssignNew(SelectedEntry, SMVVMFieldEntry)
 					.TextStyle(TextStyle)
 					.Field(SelectedField)
-					.OnValidate(this, &SMVVMFieldSelector::ValidateField)
+					.OnValidateField(this, &SMVVMFieldSelector::ValidateField)
 				]
 			]
 		]
@@ -200,13 +201,17 @@ void SMVVMFieldSelector::Refresh()
 
 TValueOrError<bool, FString> SMVVMFieldSelector::IsValidBindingForField(const FMVVMConstFieldVariant& Field, const FMVVMConstFieldVariant& CounterpartField) const
 {
+	// if we have a delegate bound, use that instead
+	if (OnValidateFieldDelegate.IsBound())
+	{
+		return OnValidateFieldDelegate.Execute(Field);
+	}
+
 	if (CounterpartField.IsEmpty() || Field.IsEmpty())
 	{
 		return MakeValue(true);
 	}
 
-	UMVVMSubsystem* Subsystem = GEngine->GetEngineSubsystem<UMVVMSubsystem>();
-
 	if (!BindingMode.IsSet())
 	{
 		return MakeValue(true);
@@ -224,6 +229,8 @@ TValueOrError<bool, FString> SMVVMFieldSelector::IsValidBindingForField(const FM
 	Args.SourceBinding = bToDestination ? Field : CounterpartField;
 	Args.DestinationBinding = bToDestination ? CounterpartField : Field;
 
+	UMVVMSubsystem* Subsystem = GEngine->GetEngineSubsystem<UMVVMSubsystem>();
+
 	TValueOrError<bool, FString> Result = Subsystem->IsBindingValid(Args);
 
 	// check in the other direction as well if this is a two-way binding
@@ -249,13 +256,13 @@ TSharedRef<SWidget> SMVVMFieldSelector::OnGenerateFieldWidget(FMVVMConstFieldVar
 	return SNew(SMVVMFieldEntry)
 		.TextStyle(TextStyle)
 		.Field(Field)
-		.OnValidate(this, &SMVVMFieldSelector::ValidateField);
+		.OnValidateField(this, &SMVVMFieldSelector::ValidateField);
 }
 
 void SMVVMFieldEntry::Construct(const FArguments& InArgs)
 {
 	Field = InArgs._Field;
-	OnValidate = InArgs._OnValidate;
+	OnValidateField = InArgs._OnValidateField;
 
 	ChildSlot
 	[
@@ -284,9 +291,9 @@ void SMVVMFieldEntry::Refresh()
 {
 	FText ToolTipText = FText::GetEmpty();
 	bool bEnabled = true;
-	if (OnValidate.IsBound())
+	if (OnValidateField.IsBound())
 	{
-		TValueOrError<bool, FString> Result = OnValidate.Execute(Field);
+		TValueOrError<bool, FString> Result = OnValidateField.Execute(Field);
 		if (Result.HasError())
 		{
 			ToolTipText = FText::FromString(Result.GetError());
diff --git a/Engine/Plugins/Runtime/ModelViewViewModel/Source/ModelViewViewModelEditor/Private/Widgets/SMVVMFieldSelector.h b/Engine/Plugins/Runtime/ModelViewViewModel/Source/ModelViewViewModelEditor/Private/Widgets/SMVVMFieldSelector.h
--- a/Engine/Plugins/Runtime/ModelViewViewModel/Source/ModelViewViewModelEditor/Private/Widgets/SMVVMFieldSelector.h
+++ b/Engine/Plugins/Runtime/ModelViewViewModel/Source/ModelViewViewModelEditor/Private/Widgets/SMVVMFieldSelector.h
@@ -21,20 +21,19 @@ namespace UE::MVVM
 class SMVVMFieldIcon;
 class STextBlock; 
 
+using FIsFieldValidResult = TValueOrError<bool, FString>;
+DECLARE_DELEGATE_RetVal_OneParam(FIsFieldValidResult, FIsFieldValid, UE::MVVM::FMVVMConstFieldVariant);
+
 class SMVVMFieldEntry : public SCompoundWidget
 {
 public:
 
-	using FIsEntryValidResult = TValueOrError<bool, FString>; 
-
-	DECLARE_DELEGATE_RetVal_OneParam(FIsEntryValidResult, FIsEntryValid, UE::MVVM::FMVVMConstFieldVariant);
-
 	SLATE_BEGIN_ARGS(SMVVMFieldEntry) :
 		_TextStyle( &FCoreStyle::Get().GetWidgetStyle<FTextBlockStyle>( "NormalText" ) )
 		{}
 		SLATE_STYLE_ARGUMENT(FTextBlockStyle, TextStyle)
 		SLATE_ARGUMENT(UE::MVVM::FMVVMConstFieldVariant, Field)
-		SLATE_EVENT(FIsEntryValid, OnValidate)
+		SLATE_EVENT(FIsFieldValid, OnValidateField)
 	SLATE_END_ARGS()
 
 	void Construct(const FArguments& InArgs);
@@ -43,7 +42,7 @@ public:
 
 private:
 	UE::MVVM::FMVVMConstFieldVariant Field;
-	FIsEntryValid OnValidate;
+	FIsFieldValid OnValidateField;
 	TSharedPtr<SMVVMFieldIcon> Icon;
 	TSharedPtr<STextBlock> Label;
 };
@@ -64,6 +63,7 @@ public:
 		SLATE_ATTRIBUTE(EMVVMBindingMode, BindingMode)
 		SLATE_ARGUMENT(bool, IsSource)
 		SLATE_EVENT(FSelectionChanged, OnSelectionChanged)
+		SLATE_EVENT(FIsFieldValid, OnValidateField)
 	SLATE_END_ARGS()
 
 	void Construct(const FArguments& InArgs);
@@ -92,6 +92,7 @@ private:
 	TSharedPtr<SComboBox<UE::MVVM::FMVVMConstFieldVariant>> FieldComboBox;
 
 	FSelectionChanged OnSelectionChangedDelegate;
+	FIsFieldValid OnValidateFieldDelegate;
 
 	TArray<UE::MVVM::FMVVMConstFieldVariant> AvailableFields;
 }; 
