MVVM: Checkbox is now to remove the binding from being compiled.

diff --git a/Engine/Plugins/Runtime/ModelViewViewModel/Source/ModelViewViewModelEditor/Private/Widgets/SMVVMViewBindingListView.cpp b/Engine/Plugins/Runtime/ModelViewViewModel/Source/ModelViewViewModelEditor/Private/Widgets/SMVVMViewBindingListView.cpp
--- a/Engine/Plugins/Runtime/ModelViewViewModel/Source/ModelViewViewModelEditor/Private/Widgets/SMVVMViewBindingListView.cpp
+++ b/Engine/Plugins/Runtime/ModelViewViewModel/Source/ModelViewViewModelEditor/Private/Widgets/SMVVMViewBindingListView.cpp
@@ -46,6 +46,7 @@ class SMVVMViewBindingListEntryRow : public SMultiColumnTableRow<FMVVMViewBindin
 {
 public:
 	static FName EnabledColumnName;
+	static FName CompileColumnName;
 	static FName ErrorColumnName;
 	static FName ViewModelColumnName;
 	static FName ViewModelPropertyColumnName;
@@ -114,11 +115,11 @@ public:
 	{
 		FMVVMBlueprintViewBinding* ViewModelBinding = GetThisViewBinding();
 
-		if (ColumnName == EnabledColumnName)
+		if (ColumnName == CompileColumnName)
 		{
 			return SNew(SCheckBox)
-				.IsChecked(this, &SMVVMViewBindingListEntryRow::IsBindingEnabled)
-				.OnCheckStateChanged(this, &SMVVMViewBindingListEntryRow::OnIsBindingEnableChanged);
+				.IsChecked(this, &SMVVMViewBindingListEntryRow::IsBindingCompiled)
+				.OnCheckStateChanged(this, &SMVVMViewBindingListEntryRow::OnIsBindingCompileChanged);
 		}
 		else if (ColumnName == ErrorColumnName)
 		{
@@ -279,6 +280,15 @@ private:
 		return ECheckBoxState::Undetermined;
 	}
 
+	ECheckBoxState IsBindingCompiled() const
+	{
+		if (const FMVVMBlueprintViewBinding* ViewModelBinding = GetThisViewBinding())
+		{
+			return ViewModelBinding->bCompile ? ECheckBoxState::Checked : ECheckBoxState::Unchecked;
+		}
+		return ECheckBoxState::Undetermined;
+	}
+
 	EVisibility GetErrorVisibility() const
 	{
 		return GetThisViewBinding()->Errors.IsEmpty() ? EVisibility::Hidden : EVisibility::Visible;
@@ -479,6 +489,30 @@ private:
 		}
 	}
 
+	void OnIsBindingCompileChanged(ECheckBoxState NewState)
+	{
+		if (NewState == ECheckBoxState::Undetermined)
+		{
+			return;
+		}
+
+		if (FMVVMBlueprintViewBinding* ViewModelBinding = GetThisViewBinding())
+		{
+			bool bNewCompile = (NewState == ECheckBoxState::Checked);
+			if (ViewModelBinding->bCompile != bNewCompile)
+			{
+				if (UMVVMBlueprintView* BlueprintViewPtr = BlueprintView.Get())
+				{
+					BlueprintViewPtr->PreEditChange(UMVVMBlueprintView::StaticClass()->FindPropertyByName("Bindings"));
+
+					ViewModelBinding->bCompile = bNewCompile;
+
+					BlueprintViewPtr->PostEditChange();
+				}
+			}
+		}
+	}
+
 	const FSlateBrush* GetModeBrush(EMVVMBindingMode BindingMode) const
 	{
 		switch (BindingMode)
@@ -641,6 +675,7 @@ private:
 };
 
 FName SMVVMViewBindingListEntryRow::EnabledColumnName = "Enabled";
+FName SMVVMViewBindingListEntryRow::CompileColumnName = "Compile";
 FName SMVVMViewBindingListEntryRow::ErrorColumnName = "Error";
 FName SMVVMViewBindingListEntryRow::ViewModelColumnName = "ViewModel";
 FName SMVVMViewBindingListEntryRow::ViewModelPropertyColumnName = "ViewModelProperty";
@@ -675,7 +710,7 @@ void SMVVMViewBindingListView::Construct(const FArguments& InArgs, TSharedPtr<SM
 		.HeaderRow
 		(
 			SNew(SHeaderRow)
-			+ SHeaderRow::Column(SMVVMViewBindingListEntryRow::EnabledColumnName)
+			+ SHeaderRow::Column(SMVVMViewBindingListEntryRow::CompileColumnName)
 			.DefaultLabel(FText::GetEmpty())
 			.FixedWidth(25.f)
 			+ SHeaderRow::Column(SMVVMViewBindingListEntryRow::ErrorColumnName)
