MVVM: Removed FMVVMViewModelPropertyPath and FMVVMWidgetPropertyPath in favour of FMVVMBlueprintPropertyPath.

#rnx

diff --git a/Engine/Plugins/Runtime/ModelViewViewModel/Config/BaseModelViewViewModel.ini b/Engine/Plugins/Runtime/ModelViewViewModel/Config/BaseModelViewViewModel.ini
new file mode 100644
--- /dev/null
+++ b/Engine/Plugins/Runtime/ModelViewViewModel/Config/BaseModelViewViewModel.ini
@@ -0,0 +1,3 @@
+[CoreRedirects]
++StructRedirects=(OldName="MVVMViewModelPropertyPath",NewName="MVVMBlueprintPropertyPath")
++StructRedirects=(OldName="MVVMWidgetPropertyPath",NewName="MVVMBlueprintPropertyPath")
diff --git a/Engine/Plugins/Runtime/ModelViewViewModel/Source/ModelViewViewModelBlueprint/Private/MVVMBlueprintView.cpp b/Engine/Plugins/Runtime/ModelViewViewModel/Source/ModelViewViewModelBlueprint/Private/MVVMBlueprintView.cpp
--- a/Engine/Plugins/Runtime/ModelViewViewModel/Source/ModelViewViewModelBlueprint/Private/MVVMBlueprintView.cpp
+++ b/Engine/Plugins/Runtime/ModelViewViewModel/Source/ModelViewViewModelBlueprint/Private/MVVMBlueprintView.cpp
@@ -79,7 +79,7 @@ FMVVMBlueprintViewBinding* UMVVMBlueprintView::FindBinding(const UWidget* Widget
 	FName WidgetName = Widget->GetFName();
 	return Bindings.FindByPredicate([WidgetName, Property](const FMVVMBlueprintViewBinding& Binding)
 		{
-			return Binding.WidgetPath.WidgetName == WidgetName &&
+			return Binding.WidgetPath.GetWidgetName() == WidgetName &&
 				Binding.WidgetPath.BasePropertyPathContains(UE::MVVM::FMVVMConstFieldVariant(Property));
 		});
 }
@@ -108,7 +108,7 @@ void UMVVMBlueprintView::RemoveBinding(const FMVVMBlueprintViewBinding* Binding)
 FMVVMBlueprintViewBinding& UMVVMBlueprintView::AddBinding(const UWidget* Widget, const FProperty* Property)
 {
 	FMVVMBlueprintViewBinding& NewBinding = Bindings.AddDefaulted_GetRef();
-	NewBinding.WidgetPath.WidgetName = Widget->GetFName();
+	NewBinding.WidgetPath.SetWidgetName(Widget->GetFName());
 	NewBinding.WidgetPath.SetBasePropertyPath(UE::MVVM::FMVVMConstFieldVariant(Property));
 
 	OnBindingsUpdated.Broadcast();
@@ -160,9 +160,9 @@ void UMVVMBlueprintView::WidgetRenamed(FName OldObjectName, FName NewObjectName)
 	bool bRenamed = false;
 	for (FMVVMBlueprintViewBinding& Binding : Bindings)
 	{
-		if (Binding.WidgetPath.WidgetName == OldObjectName)
+		if (Binding.WidgetPath.GetWidgetName() == OldObjectName)
 		{
-			Binding.WidgetPath.WidgetName = NewObjectName;
+			Binding.WidgetPath.SetWidgetName(NewObjectName);
 			bRenamed = true;
 		}
 	}
diff --git a/Engine/Plugins/Runtime/ModelViewViewModel/Source/ModelViewViewModelBlueprint/Private/MVVMViewBlueprintCompiler.cpp b/Engine/Plugins/Runtime/ModelViewViewModel/Source/ModelViewViewModelBlueprint/Private/MVVMViewBlueprintCompiler.cpp
--- a/Engine/Plugins/Runtime/ModelViewViewModel/Source/ModelViewViewModelBlueprint/Private/MVVMViewBlueprintCompiler.cpp
+++ b/Engine/Plugins/Runtime/ModelViewViewModel/Source/ModelViewViewModelBlueprint/Private/MVVMViewBlueprintCompiler.cpp
@@ -334,7 +334,7 @@ void FMVVMViewBlueprintCompiler::CreateSourceLists(const FWidgetBlueprintCompile
 		if (IsForwardBinding(Binding.BindingType) || IsBackwardBinding(Binding.BindingType))
 		{
 			// todo support any type of UObject
-			const FMVVMBlueprintViewModelContext* SourceViewModelContext = BlueprintView->FindViewModel(Binding.ViewModelPath.ContextId);
+			const FMVVMBlueprintViewModelContext* SourceViewModelContext = BlueprintView->FindViewModel(Binding.ViewModelPath.GetViewModelId());
 			if (SourceViewModelContext == nullptr)
 			{
 				WidgetBlueprintCompilerContext.MessageLog.Error(*FString::Printf(TEXT("The binding number %d has an invalid source."), Index));
@@ -349,23 +349,23 @@ void FMVVMViewBlueprintCompiler::CreateSourceLists(const FWidgetBlueprintCompile
 				bAreSourceContextsValid = false;
 			}
 
-			if (Binding.WidgetPath.WidgetName == NAME_None)
+			if (Binding.WidgetPath.GetWidgetName().IsNone())
 			{
 				WidgetBlueprintCompilerContext.MessageLog.Error(*FString::Printf(TEXT("The widget binding is invalid.")));
 				bAreSourceContextsValid = false;
 				continue;
 			}
 
-			if (Binding.WidgetPath.WidgetName != Context.GetWidgetBlueprint()->GetFName()) // is the userwidget itself
+			if (Binding.WidgetPath.GetWidgetName() != Context.GetWidgetBlueprint()->GetFName()) // is the userwidget itself
 			{
-				if (!WidgetSources.Contains(Binding.WidgetPath.WidgetName))
+				if (!WidgetSources.Contains(Binding.WidgetPath.GetWidgetName()))
 				{
-					WidgetSources.Add(Binding.WidgetPath.WidgetName);
+					WidgetSources.Add(Binding.WidgetPath.GetWidgetName());
 
-					UWidget** WidgetPtr = WidgetNameToWidgetPointerMap.Find(Binding.WidgetPath.WidgetName);
+					UWidget** WidgetPtr = WidgetNameToWidgetPointerMap.Find(Binding.WidgetPath.GetWidgetName());
 					if (WidgetPtr == nullptr || *WidgetPtr == nullptr)
 					{
-						WidgetBlueprintCompilerContext.MessageLog.Error(*FString::Printf(TEXT("The widget '%s' could not be found."), *Binding.WidgetPath.WidgetName.ToString()));
+						WidgetBlueprintCompilerContext.MessageLog.Error(*FString::Printf(TEXT("The widget '%s' could not be found."), *Binding.WidgetPath.GetWidgetName().ToString()));
 						bAreSourceContextsValid = false;
 						continue;
 					}
@@ -373,7 +373,7 @@ void FMVVMViewBlueprintCompiler::CreateSourceLists(const FWidgetBlueprintCompile
 					UWidget* Widget = *WidgetPtr;
 					FCompilerSourceContext SourceVariable;
 					SourceVariable.Class = Widget->GetClass();
-					SourceVariable.PropertyName = Binding.WidgetPath.WidgetName;
+					SourceVariable.PropertyName = Binding.WidgetPath.GetWidgetName();
 					SourceVariable.DisplayName = Widget->GetDisplayLabel();
 					SourceVariable.CategoryName = TEXT("Widget");
 					SourceContexts.Emplace(MoveTemp(SourceVariable));
@@ -745,7 +745,7 @@ bool FMVVMViewBlueprintCompiler::PreCompileBindings(UWidgetBlueprintGeneratedCla
 
 			FString Setter;
 			{
-				FName DestinationName = Binding.WidgetPath.WidgetName;
+				FName DestinationName = Binding.WidgetPath.GetWidgetName();
 				checkf(!DestinationName.IsNone(), TEXT("The destination should have been checked and set bAreSourceContextsValid."));
 				const bool bSourceIsUserWidget = DestinationName == Class->ClassGeneratedBy->GetFName();
 				if (bSourceIsUserWidget)
@@ -814,7 +814,7 @@ bool FMVVMViewBlueprintCompiler::PreCompileBindings(UWidgetBlueprintGeneratedCla
 
 			FString Setter;
 			{
-				const FMVVMBlueprintViewModelContext* SourceViewModelContext = BlueprintView->FindViewModel(Binding.ViewModelPath.ContextId);
+				const FMVVMBlueprintViewModelContext* SourceViewModelContext = BlueprintView->FindViewModel(Binding.ViewModelPath.GetViewModelId());
 				check(SourceViewModelContext);
 				FName DestinationName = SourceViewModelContext->GetViewModelName();
 				const int32 DestinationVariableContextIndex = SourceContexts.IndexOfByPredicate([DestinationName](const FCompilerSourceContext& Other) { return Other.PropertyName == DestinationName; });
@@ -902,7 +902,7 @@ const FMVVMViewBlueprintCompiler::FCompilerSourceCreatorContext* FMVVMViewBluepr
 }
 
 
-FMVVMViewBlueprintCompiler::FBindingSourceContext FMVVMViewBlueprintCompiler::CreateBindingSourceContext(const UMVVMBlueprintView* BlueprintView, const FMVVMViewModelPropertyPath& PropertyPath) const
+FMVVMViewBlueprintCompiler::FBindingSourceContext FMVVMViewBlueprintCompiler::CreateBindingSourceContext(const UMVVMBlueprintView* BlueprintView, const FMVVMBlueprintPropertyPath& PropertyPath) const
 {
 	FBindingSourceContext Result;
 
@@ -914,7 +914,7 @@ FMVVMViewBlueprintCompiler::FBindingSourceContext FMVVMViewBlueprintCompiler::Cr
 	}
 
 	{
-		const FMVVMBlueprintViewModelContext* SourceViewModelContext = BlueprintView->FindViewModel(PropertyPath.ContextId);
+		const FMVVMBlueprintViewModelContext* SourceViewModelContext = BlueprintView->FindViewModel(PropertyPath.GetViewModelId());
 		check(SourceViewModelContext);
 		const FName SourceName = SourceViewModelContext->GetViewModelName();
 		Result.CompilerSourceContextIndex = SourceContexts.IndexOfByPredicate([SourceName](const FCompilerSourceContext& Other) { return Other.PropertyName == SourceName; });
@@ -929,11 +929,11 @@ FMVVMViewBlueprintCompiler::FBindingSourceContext FMVVMViewBlueprintCompiler::Cr
 }
 
 
-FMVVMViewBlueprintCompiler::FBindingSourceContext FMVVMViewBlueprintCompiler::CreateBindingSourceContext(const UMVVMBlueprintView* BlueprintView, const UWidgetBlueprintGeneratedClass* Class, const FMVVMWidgetPropertyPath& PropertyPath) const
+FMVVMViewBlueprintCompiler::FBindingSourceContext FMVVMViewBlueprintCompiler::CreateBindingSourceContext(const UMVVMBlueprintView* BlueprintView, const UWidgetBlueprintGeneratedClass* Class, const FMVVMBlueprintPropertyPath& PropertyPath) const
 {
 	FBindingSourceContext Result;
 
-	const FName SourceName = PropertyPath.WidgetName;
+	const FName SourceName = PropertyPath.GetWidgetName();
 	Result.bIsRootWidget = SourceName == Class->ClassGeneratedBy->GetFName();
 
 	// Todo the path may contains an other INotifyFieldValueChanged
diff --git a/Engine/Plugins/Runtime/ModelViewViewModel/Source/ModelViewViewModelBlueprint/Private/MVVMViewBlueprintCompiler.h b/Engine/Plugins/Runtime/ModelViewViewModel/Source/ModelViewViewModelBlueprint/Private/MVVMViewBlueprintCompiler.h
--- a/Engine/Plugins/Runtime/ModelViewViewModel/Source/ModelViewViewModelBlueprint/Private/MVVMViewBlueprintCompiler.h
+++ b/Engine/Plugins/Runtime/ModelViewViewModel/Source/ModelViewViewModelBlueprint/Private/MVVMViewBlueprintCompiler.h
@@ -70,8 +70,8 @@ private:
 
 	void AddErrorForBinding(FMVVMBlueprintViewBinding& Binding, const UMVVMBlueprintView* BlueprintView, const FString& Message) const;
 
-	FBindingSourceContext CreateBindingSourceContext(const UMVVMBlueprintView* BlueprintView, const FMVVMViewModelPropertyPath& PropertyPath) const;
-	FBindingSourceContext CreateBindingSourceContext(const UMVVMBlueprintView* BlueprintView, const UWidgetBlueprintGeneratedClass* Class, const FMVVMWidgetPropertyPath& PropertyPath) const;
+	FBindingSourceContext CreateBindingSourceContext(const UMVVMBlueprintView* BlueprintView, const FMVVMBlueprintPropertyPath& PropertyPath) const;
+	FBindingSourceContext CreateBindingSourceContext(const UMVVMBlueprintView* BlueprintView, const UWidgetBlueprintGeneratedClass* Class, const FMVVMBlueprintPropertyPath& PropertyPath) const;
 
 	FString CreatePropertyPath(FName PropertyName, FString PropertyPath) const;
 
diff --git a/Engine/Plugins/Runtime/ModelViewViewModel/Source/ModelViewViewModelBlueprint/Public/MVVMBlueprintViewBinding.cpp b/Engine/Plugins/Runtime/ModelViewViewModel/Source/ModelViewViewModelBlueprint/Public/MVVMBlueprintViewBinding.cpp
--- a/Engine/Plugins/Runtime/ModelViewViewModel/Source/ModelViewViewModelBlueprint/Public/MVVMBlueprintViewBinding.cpp
+++ b/Engine/Plugins/Runtime/ModelViewViewModel/Source/ModelViewViewModelBlueprint/Public/MVVMBlueprintViewBinding.cpp
@@ -8,7 +8,7 @@
 
 FString FMVVMBlueprintViewBinding::GetNameString(const UMVVMBlueprintView* View) const
 {
-	const FMVVMBlueprintViewModelContext* ViewModel = View->FindViewModel(ViewModelPath.ContextId);
+	const FMVVMBlueprintViewModelContext* ViewModel = View->FindViewModel(ViewModelPath.GetViewModelId());
 	TStringBuilder<512> BindingName;
 	if (ViewModel != nullptr)
 	{
@@ -38,19 +38,19 @@ FString FMVVMBlueprintViewBinding::GetNameString(const UMVVMBlueprintView* View)
 		BindingName << TEXT(" ??? "); // shouldn't happen
 	}
 
-	if (View->GetOuterUMVVMWidgetBlueprintExtension_View()->GetWidgetBlueprint()->GetFName() == WidgetPath.WidgetName)
+	if (View->GetOuterUMVVMWidgetBlueprintExtension_View()->GetWidgetBlueprint()->GetFName() == WidgetPath.GetWidgetName())
 	{
 		BindingName << WidgetPath.GetBasePropertyPath();
 	}
 	else
 	{
-		if (WidgetPath.WidgetName.IsNone())
+		if (WidgetPath.GetWidgetName().IsNone())
 		{
 			BindingName << TEXT("<none>");
 		}
 		else
 		{
-			BindingName << WidgetPath.WidgetName;
+			BindingName << WidgetPath.GetWidgetName();
 		}
 		BindingName << TEXT(".");
 		BindingName << WidgetPath.GetBasePropertyPath();
diff --git a/Engine/Plugins/Runtime/ModelViewViewModel/Source/ModelViewViewModelBlueprint/Public/MVVMBlueprintViewBinding.h b/Engine/Plugins/Runtime/ModelViewViewModel/Source/ModelViewViewModelBlueprint/Public/MVVMBlueprintViewBinding.h
--- a/Engine/Plugins/Runtime/ModelViewViewModel/Source/ModelViewViewModelBlueprint/Public/MVVMBlueprintViewBinding.h
+++ b/Engine/Plugins/Runtime/ModelViewViewModel/Source/ModelViewViewModelBlueprint/Public/MVVMBlueprintViewBinding.h
@@ -35,10 +35,10 @@ struct FMVVMBlueprintViewBinding
 	GENERATED_BODY()
 
 	UPROPERTY(EditAnywhere, Category = "MVVM")
-	FMVVMViewModelPropertyPath ViewModelPath;
+	FMVVMBlueprintPropertyPath ViewModelPath;
 
 	UPROPERTY(EditAnywhere, Category = "MVVM")
-	FMVVMWidgetPropertyPath WidgetPath;
+	FMVVMBlueprintPropertyPath WidgetPath;
 
 	/** */
 	UPROPERTY(EditAnywhere, Category = "MVVM")
diff --git a/Engine/Plugins/Runtime/ModelViewViewModel/Source/ModelViewViewModelBlueprint/Public/MVVMPropertyPath.h b/Engine/Plugins/Runtime/ModelViewViewModel/Source/ModelViewViewModelBlueprint/Public/MVVMPropertyPath.h
--- a/Engine/Plugins/Runtime/ModelViewViewModel/Source/ModelViewViewModelBlueprint/Public/MVVMPropertyPath.h
+++ b/Engine/Plugins/Runtime/ModelViewViewModel/Source/ModelViewViewModelBlueprint/Public/MVVMPropertyPath.h
@@ -74,23 +74,22 @@ public:
 		if (InField.IsEmpty())
 		{
 			Reset();
+			return;
+		}
+
+		BindingReference = InField.CreateMemberReference();
+
+		if (InField.IsProperty())
+		{
+			BindingKind = EBindingKind::Property;
+		}
+		else if (InField.IsFunction())
+		{
+			BindingKind = EBindingKind::Function;
 		}
 		else
 		{
-			BindingReference = InField.CreateMemberReference();
-
-			if (InField.IsProperty())
-			{
-				BindingKind = EBindingKind::Property;
-			}
-			else if (InField.IsFunction())
-			{
-				BindingKind = EBindingKind::Function;
-			}
-			else
-			{
-				ensureAlwaysMsgf(false, TEXT("Binding to field of unknown type!"));
-			}
+			ensureAlwaysMsgf(false, TEXT("Binding to field of unknown type!"));
 		}
 	}
 
@@ -123,7 +122,7 @@ public:
  * Used to associate properties within MVVM bindings in editor & during MVVM compilation
  */
 USTRUCT()
-struct FMVVMBlueprintPropertyPathBase
+struct FMVVMBlueprintPropertyPath
 {
 	GENERATED_BODY()
 
@@ -132,6 +131,12 @@ private:
 	UPROPERTY(EditAnywhere, Category = "MVVM")
 	TArray<FMVVMBlueprintFieldPath> Paths;
 
+	UPROPERTY(EditAnywhere, Category = "MVVM")
+	FName WidgetName;
+
+	UPROPERTY(EditAnywhere, Category = "MVVM")
+	FGuid ContextId;
+
 #if WITH_EDITORONLY_DATA
 	// Use the Paths. BindingReference and BindingKind are deprecated.
 	UPROPERTY()
@@ -171,6 +176,49 @@ public:
 		Paths.Reset();
 	}
 
+	void ResetSource()
+	{
+		ContextId = FGuid();
+		WidgetName = FName();
+	}
+
+	bool IsFromWidget() const
+	{
+		return !WidgetName.IsNone();
+	}
+
+	bool IsFromViewModel() const
+	{
+		return ContextId.IsValid();
+	}
+
+	FGuid GetViewModelId() const
+	{
+		return ContextId;
+	}
+
+	void SetViewModelId(FGuid InViewModelId)
+	{
+		WidgetName = FName();
+		ContextId = InViewModelId;
+	}
+
+	FName GetWidgetName() const
+	{
+		return WidgetName;
+	}
+
+	void SetWidgetName(FName InWidgetName)
+	{
+		ContextId = FGuid();
+		WidgetName = InWidgetName;
+	}
+
+	bool IsEmpty() const
+	{
+		return !IsFromWidget() && !IsFromViewModel() && BindingReference.GetMemberName() == FName();
+	}
+
 	/**
 	 * Get the full path without the first property name.
 	 * returns Field.SubProperty.SubProperty from ViewModel.Field.SubProeprty.SubProperty
@@ -189,100 +237,34 @@ public:
 		return Result.ToString();
 	}
 
-public:
-	bool operator==(const FMVVMBlueprintPropertyPathBase& Other) const
+	bool operator==(const FMVVMBlueprintPropertyPath& Other) const
 	{
 		return Paths == Other.Paths;
 	}
-	bool operator!=(const FMVVMBlueprintPropertyPathBase& Other) const
+
+	bool operator!=(const FMVVMBlueprintPropertyPath& Other) const
 	{
 		return !operator==(Other);
 	}
 
-public:
 	void PostSerialize(const FArchive& Ar)
 	{
-		if (Ar.IsLoading() && !BindingReference.GetMemberName().IsNone())
+		if (Ar.IsLoading())
 		{
-			Paths.AddDefaulted_GetRef().SetDeprecatedBindingReference(BindingReference, BindingKind);
-			BindingReference = FMemberReference();
+			if (!BindingReference.GetMemberName().IsNone())
+			{
+				Paths.AddDefaulted_GetRef().SetDeprecatedBindingReference(BindingReference, BindingKind);
+				BindingReference = FMemberReference();
+			}
 		}
 	}
 };
 
 template<>
-struct TStructOpsTypeTraits<FMVVMBlueprintPropertyPathBase> : public TStructOpsTypeTraitsBase2<FMVVMBlueprintPropertyPathBase>
-{
-	enum
-	{
-		WithPostSerialize = true,
-	};
-};
-
-
-/** */
-USTRUCT(BlueprintType)
-struct FMVVMViewModelPropertyPath : public FMVVMBlueprintPropertyPathBase
-{
-	GENERATED_BODY();
-
-	UPROPERTY(EditAnywhere, Category = "MVVM")
-	FGuid ContextId;
-
-	bool operator==(const FMVVMViewModelPropertyPath& Other) const
-	{
-		return ContextId == Other.ContextId &&
-			FMVVMBlueprintPropertyPathBase::operator==(Other);
-	}
-
-	bool operator!=(const FMVVMViewModelPropertyPath& Other) const
-	{
-		return ContextId != Other.ContextId &&
-			FMVVMBlueprintPropertyPathBase::operator!=(Other);
-	}
-};
-
-
-template<>
-struct TStructOpsTypeTraits<FMVVMViewModelPropertyPath> : public TStructOpsTypeTraitsBase2<FMVVMViewModelPropertyPath>
+struct TStructOpsTypeTraits<FMVVMBlueprintPropertyPath> : public TStructOpsTypeTraitsBase2<FMVVMBlueprintPropertyPath>
 {
 	enum
 	{
 		WithPostSerialize = true,
 	};
-};
-
-
-/** */
-USTRUCT(BlueprintType)
-struct FMVVMWidgetPropertyPath : public FMVVMBlueprintPropertyPathBase
-{
-	GENERATED_BODY();
-
-	/** The context from which the destination path will be evaluated from. */
-	UPROPERTY(EditAnywhere, Category = "MVVM")
-	FName WidgetName;
-
-	bool operator==(const FMVVMWidgetPropertyPath& Other) const
-	{
-		return WidgetName == Other.WidgetName &&
-			FMVVMBlueprintPropertyPathBase::operator==(Other);
-	}
-
-
-	bool operator!=(const FMVVMWidgetPropertyPath& Other) const
-	{
-		return WidgetName != Other.WidgetName ||
-			FMVVMBlueprintPropertyPathBase::operator!=(Other);
-	}
-};
-
-
-template<>
-struct TStructOpsTypeTraits<FMVVMWidgetPropertyPath> : public TStructOpsTypeTraitsBase2<FMVVMWidgetPropertyPath>
-{
-	enum
-	{
-		WithPostSerialize = true,
-	};
-};
+};
\ No newline at end of file
diff --git a/Engine/Plugins/Runtime/ModelViewViewModel/Source/ModelViewViewModelEditor/Private/Customizations/MVVMPropertyBindingExtension.cpp b/Engine/Plugins/Runtime/ModelViewViewModel/Source/ModelViewViewModelEditor/Private/Customizations/MVVMPropertyBindingExtension.cpp
--- a/Engine/Plugins/Runtime/ModelViewViewModel/Source/ModelViewViewModelEditor/Private/Customizations/MVVMPropertyBindingExtension.cpp
+++ b/Engine/Plugins/Runtime/ModelViewViewModel/Source/ModelViewViewModelEditor/Private/Customizations/MVVMPropertyBindingExtension.cpp
@@ -80,10 +80,10 @@ static void ExtendBindingsMenu(FMenuBuilder& MenuBuilder, const UWidgetBlueprint
 	{
 		FMVVMBlueprintViewBinding& NewBinding = MVVMBlueprintView->AddDefaultBinding();
 
-		NewBinding.ViewModelPath.ContextId = ViewModelId;
+		NewBinding.ViewModelPath.SetViewModelId(ViewModelId);
 		NewBinding.ViewModelPath.SetBasePropertyPath(UE::MVVM::FMVVMConstFieldVariant(ViewModelProperty));
 
-		NewBinding.WidgetPath.WidgetName = Widget->GetFName();
+		NewBinding.WidgetPath.SetWidgetName(Widget->GetFName());
 		NewBinding.WidgetPath.SetBasePropertyPath(UE::MVVM::FMVVMConstFieldVariant(WidgetProperty));
 
 		NewBinding.BindingType = EMVVMBindingMode::OneWayToDestination;
diff --git a/Engine/Plugins/Runtime/ModelViewViewModel/Source/ModelViewViewModelEditor/Private/Customizations/MVVMPropertyPathCustomization.cpp b/Engine/Plugins/Runtime/ModelViewViewModel/Source/ModelViewViewModelEditor/Private/Customizations/MVVMPropertyPathCustomization.cpp
--- a/Engine/Plugins/Runtime/ModelViewViewModel/Source/ModelViewViewModelEditor/Private/Customizations/MVVMPropertyPathCustomization.cpp
+++ b/Engine/Plugins/Runtime/ModelViewViewModel/Source/ModelViewViewModelEditor/Private/Customizations/MVVMPropertyPathCustomization.cpp
@@ -29,7 +29,7 @@ namespace UE::MVVM
 	{
 		// figure out what the selection is - if any is different in the case of a multiselection, then return null
 		// this is required for getting the selection of the other pair
-		TArray<TSharedPtr<IFieldPathHelper>> GetPathHelpersForProperty(const TSharedPtr<IPropertyHandle>& PropertyHandle, UWidgetBlueprint* WidgetBlueprint, EPropertyPathType PathType)
+		TArray<TSharedPtr<IFieldPathHelper>> GetPathHelpersForProperty(const TSharedPtr<IPropertyHandle>& PropertyHandle, UWidgetBlueprint* WidgetBlueprint, bool bIsWidget)
 		{
 			TArray<TSharedPtr<IFieldPathHelper>> Result;
 
@@ -39,13 +39,14 @@ namespace UE::MVVM
 			for (void* Datum : RawData)
 			{
 				const FProperty* ThisSelected = nullptr;
-				if (PathType == EPropertyPathType::ViewModel)
+				
+				if (bIsWidget)
 				{
-					Result.Add(MakeShared<FViewModelFieldPathHelper>(reinterpret_cast<FMVVMViewModelPropertyPath*>(Datum), WidgetBlueprint));
+					Result.Add(MakeShared<FWidgetFieldPathHelper>(reinterpret_cast<FMVVMBlueprintPropertyPath*>(Datum), WidgetBlueprint));
 				}
 				else
 				{
-					Result.Add(MakeShared<FWidgetFieldPathHelper>(reinterpret_cast<FMVVMWidgetPropertyPath*>(Datum), WidgetBlueprint));
+					Result.Add(MakeShared<FViewModelFieldPathHelper>(reinterpret_cast<FMVVMBlueprintPropertyPath*>(Datum), WidgetBlueprint));
 				}
 			}
 			return Result;
@@ -62,41 +63,53 @@ namespace UE::MVVM
 		}
 	}
 
-	FPropertyPathCustomizationBase::FPropertyPathCustomizationBase(UWidgetBlueprint* InWidgetBlueprint, EPropertyPathType InPathType) :
-		PathType(InPathType),
+	FPropertyPathCustomization::FPropertyPathCustomization(UWidgetBlueprint* InWidgetBlueprint) :
 		WidgetBlueprint(InWidgetBlueprint)
 	{
 		check(WidgetBlueprint != nullptr);
 
-		OnBlueprintChangedHandle = WidgetBlueprint->OnChanged().AddRaw(this, &FPropertyPathCustomizationBase::HandleBlueprintChanged);
+		OnBlueprintChangedHandle = WidgetBlueprint->OnChanged().AddRaw(this, &FPropertyPathCustomization::HandleBlueprintChanged);
 	}
 
-	FPropertyPathCustomizationBase::~FPropertyPathCustomizationBase()
+	FPropertyPathCustomization::~FPropertyPathCustomization()
 	{
 		WidgetBlueprint->OnChanged().Remove(OnBlueprintChangedHandle);
 	}
 
-	void FPropertyPathCustomizationBase::CustomizeHeader(TSharedRef<IPropertyHandle> InPropertyHandle, FDetailWidgetRow& HeaderRow, IPropertyTypeCustomizationUtils& CustomizationUtils)
+	void FPropertyPathCustomization::CustomizeHeader(TSharedRef<IPropertyHandle> InPropertyHandle, FDetailWidgetRow& HeaderRow, IPropertyTypeCustomizationUtils& CustomizationUtils)
 	{
 		PropertyHandle = InPropertyHandle;
-		PathHelpers = Private::GetPathHelpersForProperty(PropertyHandle, WidgetBlueprint, PathType);
+
+		bool bIsWidget = false;
+
+		const FName PropertyName = PropertyHandle->GetProperty()->GetFName();
+		if (PropertyName == GET_MEMBER_NAME_CHECKED(FMVVMBlueprintViewBinding, WidgetPath))
+		{
+			bIsWidget = true;
+		}
+		else if (PropertyName == GET_MEMBER_NAME_CHECKED(FMVVMBlueprintViewBinding, ViewModelPath))
+		{
+			bIsWidget = false;
+		}
+
+		PathHelpers = Private::GetPathHelpersForProperty(PropertyHandle, WidgetBlueprint, bIsWidget);
 
 		TSharedPtr<IPropertyHandle> ParentHandle = PropertyHandle->GetParentHandle();
 		TSharedPtr<IPropertyHandle> OtherHandle;
 
-		if (PathType == EPropertyPathType::ViewModel)
+		if (bIsWidget)
 		{
-			OtherHandle = ParentHandle->GetChildHandle(GET_MEMBER_NAME_CHECKED(FMVVMBlueprintViewBinding, WidgetPath));
+			OtherHandle = ParentHandle->GetChildHandle(GET_MEMBER_NAME_CHECKED(FMVVMBlueprintViewBinding, ViewModelPath));
 		}
 		else
 		{
-			OtherHandle = ParentHandle->GetChildHandle(GET_MEMBER_NAME_CHECKED(FMVVMBlueprintViewBinding, ViewModelPath));
+			OtherHandle = ParentHandle->GetChildHandle(GET_MEMBER_NAME_CHECKED(FMVVMBlueprintViewBinding, WidgetPath));
 		}
 
 		if (OtherHandle.IsValid() && OtherHandle->IsValidHandle())
 		{
-			OtherHandle->SetOnPropertyValueChanged(FSimpleDelegate::CreateSP(this, &FPropertyPathCustomizationBase::OnOtherPropertyChanged));
-			OtherHelpers = Private::GetPathHelpersForProperty(OtherHandle, WidgetBlueprint, PathType == EPropertyPathType::Widget ? EPropertyPathType::ViewModel : EPropertyPathType::Widget);
+			OtherHandle->SetOnPropertyValueChanged(FSimpleDelegate::CreateSP(this, &FPropertyPathCustomization::OnOtherPropertyChanged));
+			OtherHelpers = Private::GetPathHelpersForProperty(OtherHandle, WidgetBlueprint, !bIsWidget);
 		}
 		else
 		{
@@ -106,7 +119,7 @@ namespace UE::MVVM
 		TSharedPtr<IPropertyHandle> BindingTypeHandle = ParentHandle->GetChildHandle(GET_MEMBER_NAME_CHECKED(FMVVMBlueprintViewBinding, BindingType));
 		if (BindingTypeHandle.IsValid() && BindingTypeHandle->IsValidHandle())
 		{
-			BindingTypeHandle->SetOnPropertyValueChanged(FSimpleDelegate::CreateSP(this, &FPropertyPathCustomizationBase::OnOtherPropertyChanged));
+			BindingTypeHandle->SetOnPropertyValueChanged(FSimpleDelegate::CreateSP(this, &FPropertyPathCustomization::OnOtherPropertyChanged));
 		}
 		else
 		{
@@ -132,8 +145,8 @@ namespace UE::MVVM
 				[
 					SAssignNew(SourceSelector, SMVVMSourceSelector)
 					.TextStyle(FAppStyle::Get(), "SmallText")
-					.PathHelpers(this, &FPropertyPathCustomizationBase::GetRawPathHelpers)
-					.OnSelectionChanged(this, &FPropertyPathCustomizationBase::OnSourceSelectionChanged)
+					.PathHelpers(this, &FPropertyPathCustomization::GetRawPathHelpers)
+					.OnSelectionChanged(this, &FPropertyPathCustomization::OnSourceSelectionChanged)
 				]
 			+ SHorizontalBox::Slot()
 				.Padding(5, 0, 0, 0)
@@ -143,28 +156,28 @@ namespace UE::MVVM
 					.TextStyle(FAppStyle::Get(), "SmallText")
 					.PathHelpers(GetRawPathHelpers())
 					.CounterpartHelpers(GetRawOtherHelpers())
-					.BindingMode(this, &FPropertyPathCustomizationBase::GetCurrentBindingMode)
-					.OnSelectionChanged(this, &FPropertyPathCustomizationBase::OnPropertySelectionChanged)
+					.BindingMode(this, &FPropertyPathCustomization::GetCurrentBindingMode)
+					.OnSelectionChanged(this, &FPropertyPathCustomization::OnPropertySelectionChanged)
 				]				
 			];
 	}
 
-	TArray<IFieldPathHelper*> FPropertyPathCustomizationBase::GetRawPathHelpers() const
+	TArray<IFieldPathHelper*> FPropertyPathCustomization::GetRawPathHelpers() const
 	{
 		return Private::GetRawPathHelperPointers(PathHelpers);
 	}
 
-	TArray<IFieldPathHelper*> FPropertyPathCustomizationBase::GetRawOtherHelpers() const
+	TArray<IFieldPathHelper*> FPropertyPathCustomization::GetRawOtherHelpers() const
 	{
 		return Private::GetRawPathHelperPointers(OtherHelpers);
 	}
 
-	void FPropertyPathCustomizationBase::OnOtherPropertyChanged()
+	void FPropertyPathCustomization::OnOtherPropertyChanged()
 	{
 		FieldSelector->Refresh();
 	}
 
-	EMVVMBindingMode FPropertyPathCustomizationBase::GetCurrentBindingMode() const
+	EMVVMBindingMode FPropertyPathCustomization::GetCurrentBindingMode() const
 	{
 		TSharedPtr<IPropertyHandle> BindingTypeHandle = PropertyHandle->GetParentHandle()->GetChildHandle(GET_MEMBER_NAME_CHECKED(FMVVMBlueprintViewBinding, BindingType));
 		uint8 BindingMode;
@@ -176,7 +189,7 @@ namespace UE::MVVM
 		return EMVVMBindingMode::OneWayToDestination;
 	}
 
-	void FPropertyPathCustomizationBase::OnSourceSelectionChanged(FBindingSource Selected)
+	void FPropertyPathCustomization::OnSourceSelectionChanged(FBindingSource Selected)
 	{
 		FScopedTransaction Transaction(LOCTEXT("ChangeSource", "Change binding source."));
 
@@ -193,7 +206,7 @@ namespace UE::MVVM
 		FieldSelector->Refresh();
 	}
 
-	void FPropertyPathCustomizationBase::OnPropertySelectionChanged(FMVVMConstFieldVariant Selected)
+	void FPropertyPathCustomization::OnPropertySelectionChanged(FMVVMConstFieldVariant Selected)
 	{
 		if (bPropertySelectionChanging)
 		{
@@ -215,7 +228,7 @@ namespace UE::MVVM
 		PropertyHandle->NotifyPostChange(EPropertyChangeType::ValueSet);
 	}
 
-	void FPropertyPathCustomizationBase::HandleBlueprintChanged(UBlueprint* Blueprint)
+	void FPropertyPathCustomization::HandleBlueprintChanged(UBlueprint* Blueprint)
 	{
 		SourceSelector->Refresh();
 		FieldSelector->Refresh();
diff --git a/Engine/Plugins/Runtime/ModelViewViewModel/Source/ModelViewViewModelEditor/Private/Customizations/MVVMPropertyPathCustomization.h b/Engine/Plugins/Runtime/ModelViewViewModel/Source/ModelViewViewModelEditor/Private/Customizations/MVVMPropertyPathCustomization.h
--- a/Engine/Plugins/Runtime/ModelViewViewModel/Source/ModelViewViewModelEditor/Private/Customizations/MVVMPropertyPathCustomization.h
+++ b/Engine/Plugins/Runtime/ModelViewViewModel/Source/ModelViewViewModelEditor/Private/Customizations/MVVMPropertyPathCustomization.h
@@ -15,22 +15,22 @@ namespace UE::MVVM
 {
 	struct FBindingSource;
 
-	enum class EPropertyPathType
+	class FPropertyPathCustomization : public IPropertyTypeCustomization
 	{
-		Widget,
-		ViewModel
-	};
-
-	class FPropertyPathCustomizationBase : public IPropertyTypeCustomization
-	{
-	protected:
-		FPropertyPathCustomizationBase(UWidgetBlueprint* WidgetBlueprint, EPropertyPathType PathType);
-		~FPropertyPathCustomizationBase();
-
 	public:
+		FPropertyPathCustomization(UWidgetBlueprint* WidgetBlueprint);
+		~FPropertyPathCustomization();
+
 		virtual void CustomizeHeader(TSharedRef<IPropertyHandle> InPropertyHandle, FDetailWidgetRow& HeaderRow, IPropertyTypeCustomizationUtils& CustomizationUtils) override;
 		virtual void CustomizeChildren(TSharedRef<IPropertyHandle> InPropertyHandle, IDetailChildrenBuilder& ChildBuilder, IPropertyTypeCustomizationUtils& CustomizationUtils) override {}
 
+		static TSharedRef<IPropertyTypeCustomization> MakeInstance(UWidgetBlueprint* InWidgetBlueprint)
+		{
+			check(InWidgetBlueprint != nullptr);
+
+			return MakeShared<FPropertyPathCustomization>(InWidgetBlueprint);
+		}
+
 	private:
 		void OnSourceSelectionChanged(FBindingSource Selected);
 		void OnPropertySelectionChanged(FMVVMConstFieldVariant Selected);
@@ -46,7 +46,6 @@ namespace UE::MVVM
 
 	private:
 		TSharedPtr<IPropertyHandle> PropertyHandle;
-		EPropertyPathType PathType;
 		UWidgetBlueprint* WidgetBlueprint = nullptr;
 
 		TSharedPtr<SMVVMSourceSelector> SourceSelector;
@@ -59,36 +58,4 @@ namespace UE::MVVM
 
 		FDelegateHandle OnBlueprintChangedHandle;
 	};
-
-	class FWidgetPropertyPathCustomization : public FPropertyPathCustomizationBase
-	{
-	public:
-		static TSharedRef<IPropertyTypeCustomization> MakeInstance(UWidgetBlueprint* InWidgetBlueprint)
-		{
-			check(InWidgetBlueprint != nullptr);
-			TSharedRef<FWidgetPropertyPathCustomization> Customization = MakeShared<FWidgetPropertyPathCustomization>(InWidgetBlueprint);
-			return Customization;
-		}
-	
-		FWidgetPropertyPathCustomization(UWidgetBlueprint* WidgetBlueprint) : 
-			FPropertyPathCustomizationBase(WidgetBlueprint, EPropertyPathType::Widget)
-		{
-		}
-	};
-
-	class FViewModelPropertyPathCustomization : public FPropertyPathCustomizationBase
-	{
-	public:
-		static TSharedRef<IPropertyTypeCustomization> MakeInstance(UWidgetBlueprint* InWidgetBlueprint)
-		{
-			check(InWidgetBlueprint != nullptr);
-			TSharedRef<FViewModelPropertyPathCustomization> Customization = MakeShared<FViewModelPropertyPathCustomization>(InWidgetBlueprint);
-			return Customization;
-		}
-
-		FViewModelPropertyPathCustomization(UWidgetBlueprint* WidgetBlueprint) :
-			FPropertyPathCustomizationBase(WidgetBlueprint, EPropertyPathType::ViewModel)
-		{
-		}
-	};
 }
\ No newline at end of file
diff --git a/Engine/Plugins/Runtime/ModelViewViewModel/Source/ModelViewViewModelEditor/Private/MVVMPropertyPathHelpers.cpp b/Engine/Plugins/Runtime/ModelViewViewModel/Source/ModelViewViewModelEditor/Private/MVVMPropertyPathHelpers.cpp
--- a/Engine/Plugins/Runtime/ModelViewViewModel/Source/ModelViewViewModelEditor/Private/MVVMPropertyPathHelpers.cpp
+++ b/Engine/Plugins/Runtime/ModelViewViewModel/Source/ModelViewViewModelEditor/Private/MVVMPropertyPathHelpers.cpp
@@ -15,12 +15,6 @@ namespace UE::MVVM
 {
 	void FWidgetFieldPathHelper::GetAvailableSources(TSet<FBindingSource>& OutSources) const
 	{
-		FMVVMWidgetPropertyPath* Path = PathAttr.Get(nullptr);
-		if (Path == nullptr)
-		{
-			return;
-		}
-
 		const UWidgetTree* WidgetTree = WidgetBlueprint->WidgetTree;
 		if (WidgetTree == nullptr)
 		{
@@ -44,7 +38,7 @@ namespace UE::MVVM
 				Source.Name = WidgetBlueprint->GetFName();
 				Source.DisplayName = FText::FromName(WidgetBlueprint->GetFName());
 				Source.Class = BPClass;
-				Source.IsSelected = Source.Name == Path->WidgetName;
+				Source.IsSelected = Source.Name == Path->GetWidgetName();
 				OutSources.Add(Source);
 			}
 		}
@@ -60,7 +54,7 @@ namespace UE::MVVM
 				Source.Name = Widget->GetFName();
 				Source.DisplayName = Widget->GetLabelText();
 				Source.Class = Widget->GetClass();
-				Source.IsSelected = Source.Name == Path->WidgetName;
+				Source.IsSelected = Source.Name == Path->GetWidgetName();
 				OutSources.Add(Source);
 			}
 		}
@@ -84,12 +78,6 @@ namespace UE::MVVM
 
 	void FWidgetFieldPathHelper::GetAvailableFields(TSet<FMVVMConstFieldVariant>& OutFields) const
 	{
-		FMVVMWidgetPropertyPath* Path = PathAttr.Get(nullptr);
-		if (Path == nullptr)
-		{
-			return;
-		}
-
 		const UWidgetTree* WidgetTree = WidgetBlueprint->WidgetTree;
 		if (WidgetTree == nullptr)
 		{
@@ -97,11 +85,11 @@ namespace UE::MVVM
 		}
 
 		TSubclassOf<UWidget> WidgetClass = nullptr;
-		if (const UWidget* Widget = WidgetTree->FindWidget(Path->WidgetName))
+		if (const UWidget* Widget = WidgetTree->FindWidget(Path->GetWidgetName()))
 		{
 			WidgetClass = Widget->GetClass();
 		}
-		else if (WidgetBlueprint->GetFName() == Path->WidgetName)
+		else if (WidgetBlueprint->GetFName() == Path->GetWidgetName())
 		{
 			WidgetClass = WidgetBlueprint->GeneratedClass;
 		}
@@ -138,68 +126,39 @@ namespace UE::MVVM
 
 	void FWidgetFieldPathHelper::SetSelectedSource(const FBindingSource& Source) const
 	{
-		FMVVMWidgetPropertyPath* Path = PathAttr.Get(nullptr);
-		if (Path == nullptr)
-		{
-			return;
-		}
-
 		if (Source.IsValid())
 		{
-			Path->WidgetName = Source.Name;
+			Path->SetWidgetName(Source.Name);
 		}
 		else
 		{
-			Path->WidgetName = FName();
+			Path->SetWidgetName(FName());
 		}
 	}
 
 	FMVVMBindingName FWidgetFieldPathHelper::GetBindingName() const
 	{
-		if (FMVVMWidgetPropertyPath* Path = PathAttr.Get(nullptr))
+		TArray<FName> Paths = Path->GetPaths();
+		if (Paths.Num() > 0)
 		{
-			TArray<FName> Paths = Path->GetPaths();
-			if (Paths.Num() > 0)
-			{
-				return FMVVMBindingName(Paths[0]);
-			}
+			return FMVVMBindingName(Paths[0]);
 		}
 
 		return FMVVMBindingName();
 	}
 
-	void FWidgetFieldPathHelper::SetBindingReference(const UE::MVVM::FMVVMFieldVariant& InField) const
-	{
-		if (FMVVMWidgetPropertyPath* Path = PathAttr.Get(nullptr))
-		{
-			Path->SetBasePropertyPath(InField);
-		}
-	}
-
 	void FWidgetFieldPathHelper::SetBindingReference(const UE::MVVM::FMVVMConstFieldVariant& InField) const
 	{
-		if (FMVVMWidgetPropertyPath* Path = PathAttr.Get(nullptr))
-		{
-			Path->SetBasePropertyPath(InField);
-		}
+		Path->SetBasePropertyPath(InField);
 	}
 
 	void FWidgetFieldPathHelper::ResetBinding() const
 	{
-		if (FMVVMWidgetPropertyPath* Path = PathAttr.Get(nullptr))
-		{
-			Path->ResetBasePropertyPath();
-		}
+		Path->ResetBasePropertyPath();
 	}
 
 	void FViewModelFieldPathHelper::GetAvailableSources(TSet<FBindingSource>& OutSources) const
 	{
-		FMVVMViewModelPropertyPath* Path = PathAttr.Get(nullptr);
-		if (Path == nullptr)
-		{
-			return;
-		}
-
 		const UMVVMWidgetBlueprintExtension_View* MVVMExtensionPtr = UMVVMWidgetBlueprintExtension_View::GetExtension<UMVVMWidgetBlueprintExtension_View>(WidgetBlueprint);
 		const UMVVMBlueprintView* MVVMBlueprintView = MVVMExtensionPtr->GetBlueprintView();
 		const TArrayView<const FMVVMBlueprintViewModelContext>& ViewModels = MVVMBlueprintView->GetViewModels();
@@ -212,7 +171,7 @@ namespace UE::MVVM
 			Source.SourceGuid = ViewModel.GetViewModelId();
 			Source.DisplayName = ViewModel.GetDisplayName();
 			Source.Class = ViewModel.GetViewModelClass().Get();
-			Source.IsSelected = (Source.SourceGuid == Path->ContextId);
+			Source.IsSelected = (Source.SourceGuid == Path->GetViewModelId());
 
 			OutSources.Add(Source);
 		}
@@ -236,16 +195,10 @@ namespace UE::MVVM
 
 	void FViewModelFieldPathHelper::GetAvailableFields(TSet<FMVVMConstFieldVariant>& OutFields) const
 	{
-		FMVVMViewModelPropertyPath* Path = PathAttr.Get(nullptr);
-		if (Path == nullptr)
-		{
-			return;
-		}
-
 		const UMVVMWidgetBlueprintExtension_View* MVVMExtensionPtr = UMVVMWidgetBlueprintExtension_View::GetExtension<UMVVMWidgetBlueprintExtension_View>(WidgetBlueprint);
 		const UMVVMBlueprintView* MVVMBlueprintView = MVVMExtensionPtr->GetBlueprintView();
 
-		const FMVVMBlueprintViewModelContext* ViewModelContext = MVVMBlueprintView->FindViewModel(Path->ContextId);
+		const FMVVMBlueprintViewModelContext* ViewModelContext = MVVMBlueprintView->FindViewModel(Path->GetViewModelId());
 		if (ViewModelContext == nullptr)
 		{
 			return;
@@ -278,57 +231,33 @@ namespace UE::MVVM
 
 	void FViewModelFieldPathHelper::SetSelectedSource(const FBindingSource& Source) const
 	{
-		FMVVMViewModelPropertyPath* Path = PathAttr.Get(nullptr);
-		if (Path == nullptr)
-		{
-			return;
-		}
-
 		if (Source.IsValid())
 		{
-			Path->ContextId = Source.SourceGuid;
+			Path->SetViewModelId(Source.SourceGuid);
 		}
 		else
 		{
-			Path->ContextId = FGuid();
+			Path->SetViewModelId(FGuid());
 		}
 	}
 
 	FMVVMBindingName FViewModelFieldPathHelper::GetBindingName() const
 	{
-		if (FMVVMViewModelPropertyPath* Path = PathAttr.Get(nullptr))
+		TArray<FName> Paths = Path->GetPaths();
+		if (Paths.Num() > 0)
 		{
-			TArray<FName> Paths = Path->GetPaths();
-			if (Paths.Num() > 0)
-			{
-				return FMVVMBindingName(Paths[0]);
-			}
+			return FMVVMBindingName(Paths[0]);
 		}
-
 		return FMVVMBindingName();
 	}
 
-	void FViewModelFieldPathHelper::SetBindingReference(const UE::MVVM::FMVVMFieldVariant& InField) const
-	{
-		if (FMVVMViewModelPropertyPath* Path = PathAttr.Get(nullptr))
-		{
-			Path->SetBasePropertyPath(InField);
-		}
-	}
-
 	void FViewModelFieldPathHelper::SetBindingReference(const UE::MVVM::FMVVMConstFieldVariant& InField) const
 	{
-		if (FMVVMViewModelPropertyPath* Path = PathAttr.Get(nullptr))
-		{
-			Path->SetBasePropertyPath(InField);
-		}
+		Path->SetBasePropertyPath(InField);
 	}
 
 	void FViewModelFieldPathHelper::ResetBinding() const
 	{
-		if (FMVVMViewModelPropertyPath* Path = PathAttr.Get(nullptr))
-		{
-			Path->ResetBasePropertyPath();
-		}
+		Path->ResetBasePropertyPath();
 	}
 }
diff --git a/Engine/Plugins/Runtime/ModelViewViewModel/Source/ModelViewViewModelEditor/Private/MVVMPropertyPathHelpers.h b/Engine/Plugins/Runtime/ModelViewViewModel/Source/ModelViewViewModelEditor/Private/MVVMPropertyPathHelpers.h
--- a/Engine/Plugins/Runtime/ModelViewViewModel/Source/ModelViewViewModelEditor/Private/MVVMPropertyPathHelpers.h
+++ b/Engine/Plugins/Runtime/ModelViewViewModel/Source/ModelViewViewModelEditor/Private/MVVMPropertyPathHelpers.h
@@ -17,6 +17,8 @@ namespace UE::MVVM
 	class IFieldPathHelper
 	{
 	public:
+		virtual ~IFieldPathHelper() {}
+
 		virtual void GetAvailableFields(TSet<FMVVMConstFieldVariant>& Fields) const = 0;
 		virtual FMVVMConstFieldVariant GetSelectedField() const = 0;
 
@@ -25,8 +27,7 @@ namespace UE::MVVM
 		virtual void SetSelectedSource(const FBindingSource& Source) const = 0;
 
 		virtual FMVVMBindingName GetBindingName() const = 0;
-		virtual void SetBindingReference(const UE::MVVM::FMVVMFieldVariant& InField) const = 0;
-		virtual void SetBindingReference(const UE::MVVM::FMVVMConstFieldVariant& InField) const = 0;
+		virtual void SetBindingReference(const FMVVMConstFieldVariant& InField) const = 0;
 		virtual void ResetBinding() const = 0;
 	};
 
@@ -35,10 +36,12 @@ namespace UE::MVVM
 	public:
 		FWidgetFieldPathHelper() {}
 		virtual ~FWidgetFieldPathHelper() {}
-		FWidgetFieldPathHelper(const TAttribute<FMVVMWidgetPropertyPath*>& WidgetPath, const UWidgetBlueprint* InWidgetBlueprint) :
-			PathAttr(WidgetPath),
+		FWidgetFieldPathHelper(FMVVMBlueprintPropertyPath* WidgetPath, const UWidgetBlueprint* InWidgetBlueprint) :
+			Path(WidgetPath),
 			WidgetBlueprint(InWidgetBlueprint)
 		{
+			check(Path);
+			check(Path->IsFromWidget() || Path->IsEmpty());
 		}
 
 		virtual void GetAvailableFields(TSet<FMVVMConstFieldVariant>& Fields) const override;
@@ -49,12 +52,11 @@ namespace UE::MVVM
 		virtual void SetSelectedSource(const FBindingSource& Source) const override;
 
 		virtual FMVVMBindingName GetBindingName() const override;
-		virtual void SetBindingReference(const UE::MVVM::FMVVMFieldVariant& InField) const override;
-		virtual void SetBindingReference(const UE::MVVM::FMVVMConstFieldVariant& InField) const override;
+		virtual void SetBindingReference(const FMVVMConstFieldVariant& InField) const override;
 		virtual void ResetBinding() const override;
 
 	private:
-		TAttribute<FMVVMWidgetPropertyPath*> PathAttr;
+		FMVVMBlueprintPropertyPath* Path;
 		const UWidgetBlueprint* WidgetBlueprint = nullptr;
 	};
 
@@ -63,10 +65,12 @@ namespace UE::MVVM
 	public:
 		FViewModelFieldPathHelper() {}
 		virtual ~FViewModelFieldPathHelper() {}
-		FViewModelFieldPathHelper(const TAttribute<FMVVMViewModelPropertyPath*>& ViewModelPath, const UWidgetBlueprint* InWidgetBlueprint) :
-			PathAttr(ViewModelPath),
+		FViewModelFieldPathHelper(FMVVMBlueprintPropertyPath* ViewModelPath, const UWidgetBlueprint* InWidgetBlueprint) :
+			Path(ViewModelPath),
 			WidgetBlueprint(InWidgetBlueprint)
 		{
+			check(Path);
+			check(Path->IsFromViewModel() || Path->IsEmpty());
 		}
 
 		virtual void GetAvailableFields(TSet<FMVVMConstFieldVariant>& Fields) const override;
@@ -77,12 +81,11 @@ namespace UE::MVVM
 		virtual void SetSelectedSource(const FBindingSource& Source) const override;
 
 		virtual FMVVMBindingName GetBindingName() const override;
-		virtual void SetBindingReference(const UE::MVVM::FMVVMFieldVariant& InField) const override;
-		virtual void SetBindingReference(const UE::MVVM::FMVVMConstFieldVariant& InField) const override;
+		virtual void SetBindingReference(const FMVVMConstFieldVariant& InField) const override;
 		virtual void ResetBinding() const override;
 
 	private:
-		TAttribute<FMVVMViewModelPropertyPath*> PathAttr;
+		FMVVMBlueprintPropertyPath* Path;
 		const UWidgetBlueprint* WidgetBlueprint = nullptr;
 	};
 }
diff --git a/Engine/Plugins/Runtime/ModelViewViewModel/Source/ModelViewViewModelEditor/Private/Widgets/SMVVMPropertyPath.h b/Engine/Plugins/Runtime/ModelViewViewModel/Source/ModelViewViewModelEditor/Private/Widgets/SMVVMPropertyPath.h
--- a/Engine/Plugins/Runtime/ModelViewViewModel/Source/ModelViewViewModelEditor/Private/Widgets/SMVVMPropertyPath.h
+++ b/Engine/Plugins/Runtime/ModelViewViewModel/Source/ModelViewViewModelEditor/Private/Widgets/SMVVMPropertyPath.h
@@ -27,12 +27,11 @@ private:
 	TOptional<UE::MVVM::FBindingSource> SelectedSource;
 };
 
-
 class SMVVMWidgetPropertyPath : public SMVVMPropertyPathBase
 {
 public:
 	SLATE_BEGIN_ARGS(SMVVMWidgetPropertyPath) {}
-		SLATE_ATTRIBUTE(FMVVMWidgetPropertyPath*, WidgetPath)
+		SLATE_ARGUMENT(FMVVMBlueprintPropertyPath*, WidgetPath)
 	SLATE_END_ARGS()
 
 	void Construct(const FArguments& InArgs, const UWidgetBlueprint* WidgetBlueprint);
@@ -51,7 +50,7 @@ class SMVVMViewModelPropertyPath : public SMVVMPropertyPathBase
 {
 public:
 	SLATE_BEGIN_ARGS(SMVVMViewModelPropertyPath) {}
-		SLATE_ATTRIBUTE(FMVVMViewModelPropertyPath*, ViewModelPath)
+		SLATE_ARGUMENT(FMVVMBlueprintPropertyPath*, ViewModelPath)
 	SLATE_END_ARGS()
 
 	void Construct(const FArguments& InArgs, const UWidgetBlueprint* WidgetBlueprint);
diff --git a/Engine/Plugins/Runtime/ModelViewViewModel/Source/ModelViewViewModelEditor/Private/Widgets/SMVVMViewBindingPanel.cpp b/Engine/Plugins/Runtime/ModelViewViewModel/Source/ModelViewViewModelEditor/Private/Widgets/SMVVMViewBindingPanel.cpp
--- a/Engine/Plugins/Runtime/ModelViewViewModel/Source/ModelViewViewModelEditor/Private/Widgets/SMVVMViewBindingPanel.cpp
+++ b/Engine/Plugins/Runtime/ModelViewViewModel/Source/ModelViewViewModelEditor/Private/Widgets/SMVVMViewBindingPanel.cpp
@@ -70,8 +70,7 @@ void SMVVMViewBindingPanel::Construct(const FArguments& InArgs, TSharedPtr<FWidg
 		//DetailsView->OnFinishedChangingProperties().AddSP(this, &SLiveLinkClientPanel::OnPropertyChanged);
 		FStructureDetailsViewArgs StructureDetailsViewArgs;
 		StructDetailsView = PropertyEditorModule.CreateStructureDetailView(DetailsViewArgs, StructureDetailsViewArgs, TSharedPtr<FStructOnScope>());
-		StructDetailsView->GetDetailsView()->RegisterInstancedCustomPropertyTypeLayout(FMVVMViewModelPropertyPath::StaticStruct()->GetFName(), FOnGetPropertyTypeCustomizationInstance::CreateStatic(&UE::MVVM::FViewModelPropertyPathCustomization::MakeInstance, WidgetBlueprint));
-		StructDetailsView->GetDetailsView()->RegisterInstancedCustomPropertyTypeLayout(FMVVMWidgetPropertyPath::StaticStruct()->GetFName(), FOnGetPropertyTypeCustomizationInstance::CreateStatic(&UE::MVVM::FWidgetPropertyPathCustomization::MakeInstance, WidgetBlueprint));
+		StructDetailsView->GetDetailsView()->RegisterInstancedCustomPropertyTypeLayout(FMVVMBlueprintPropertyPath::StaticStruct()->GetFName(), FOnGetPropertyTypeCustomizationInstance::CreateStatic(&UE::MVVM::FPropertyPathCustomization::MakeInstance, WidgetBlueprint));
 		StructDetailsView->GetDetailsView()->RegisterInstancedCustomPropertyTypeLayout(FMVVMBlueprintViewConversionPath::StaticStruct()->GetFName(), FOnGetPropertyTypeCustomizationInstance::CreateStatic(&UE::MVVM::FConversionPathCustomization::MakeInstance, WidgetBlueprint));
 	}
 
diff --git a/Engine/Plugins/Runtime/ModelViewViewModel/Source/ModelViewViewModelEditor/Private/Widgets/SMVVMViewModelContextListWidget.cpp b/Engine/Plugins/Runtime/ModelViewViewModel/Source/ModelViewViewModelEditor/Private/Widgets/SMVVMViewModelContextListWidget.cpp
--- a/Engine/Plugins/Runtime/ModelViewViewModel/Source/ModelViewViewModelEditor/Private/Widgets/SMVVMViewModelContextListWidget.cpp
+++ b/Engine/Plugins/Runtime/ModelViewViewModel/Source/ModelViewViewModelEditor/Private/Widgets/SMVVMViewModelContextListWidget.cpp
@@ -792,7 +792,7 @@ bool SMVVMViewModelContextListWidget::ValidateRemoveViewModelContext(TSharedPtr<
 	TArray<TSharedPtr<SHorizontalBox>> ConflictList;
 	for (const FMVVMBlueprintViewBinding& Binding : Bindings)
 	{
-		if (Binding.ViewModelPath.ContextId == ContextToRemove->GetViewModelId())
+		if (Binding.ViewModelPath.GetViewModelId() == ContextToRemove->GetViewModelId())
 		{
 			TSharedRef<SHorizontalBox> Container = SNew(SHorizontalBox);
 			Container->AddSlot()
@@ -802,6 +802,7 @@ bool SMVVMViewModelContextListWidget::ValidateRemoveViewModelContext(TSharedPtr<
 				SNew(STextBlock)
 				.Text(FText::FromString(Binding.ViewModelPath.GetBasePropertyPath()))
 			];
+
 			Container->AddSlot()
 			.Padding(FMargin(3.0f, 0.0f))
 			.AutoWidth()
@@ -816,12 +817,13 @@ bool SMVVMViewModelContextListWidget::ValidateRemoveViewModelContext(TSharedPtr<
 					.Image(UE::MVVM::Private::ManageViewModelsGetModeBrush(Binding.BindingType))
 				]
 			];
+
 			Container->AddSlot()
 			.Padding(FMargin(3.0f, 0.0f))
 			.AutoWidth()
 			[
 				SNew(STextBlock)
-				.Text(FText::FormatOrdered(INVTEXT("{0}.{1}"), FText::FromString(Binding.WidgetPath.WidgetName.ToString()), FText::FromString(Binding.WidgetPath.GetBasePropertyPath())))
+				.Text(FText::FormatOrdered(INVTEXT("{0}.{1}"), FText::FromString(Binding.WidgetPath.GetWidgetName().ToString()), FText::FromString(Binding.WidgetPath.GetBasePropertyPath())))
 			];
 			ConflictList.Add(Container);
 		}
