MVVM: Added transactions and property chain notifications to SMVVMViewBindingListEntryRow.

#rnx

diff --git a/Engine/Plugins/Runtime/ModelViewViewModel/Source/ModelViewViewModelBlueprint/Private/MVVMBlueprintView.cpp b/Engine/Plugins/Runtime/ModelViewViewModel/Source/ModelViewViewModelBlueprint/Private/MVVMBlueprintView.cpp
--- a/Engine/Plugins/Runtime/ModelViewViewModel/Source/ModelViewViewModelBlueprint/Private/MVVMBlueprintView.cpp
+++ b/Engine/Plugins/Runtime/ModelViewViewModel/Source/ModelViewViewModelBlueprint/Private/MVVMBlueprintView.cpp
@@ -155,6 +155,19 @@ void UMVVMBlueprintView::PostEditChangeProperty(FPropertyChangedEvent& PropertyC
 	}
 }
 
+void UMVVMBlueprintView::PostEditChangeChainProperty(FPropertyChangedChainEvent& PropertyChainEvent)
+{
+	Super::PostEditChangeChainProperty(PropertyChainEvent);
+	if (PropertyChainEvent.PropertyChain.Contains(UMVVMBlueprintView::StaticClass()->FindPropertyByName(GET_MEMBER_NAME_CHECKED(UMVVMBlueprintView, Bindings))))
+	{
+		OnBindingsUpdated.Broadcast();
+	}
+	if (PropertyChainEvent.PropertyChain.Contains(UMVVMBlueprintView::StaticClass()->FindPropertyByName(GET_MEMBER_NAME_CHECKED(UMVVMBlueprintView, AvailableViewModels))))
+	{
+		OnViewModelsUpdated.Broadcast();
+	}
+}
+
 void UMVVMBlueprintView::WidgetRenamed(FName OldObjectName, FName NewObjectName)
 {
 	bool bRenamed = false;
diff --git a/Engine/Plugins/Runtime/ModelViewViewModel/Source/ModelViewViewModelBlueprint/Public/MVVMBlueprintView.h b/Engine/Plugins/Runtime/ModelViewViewModel/Source/ModelViewViewModelBlueprint/Public/MVVMBlueprintView.h
--- a/Engine/Plugins/Runtime/ModelViewViewModel/Source/ModelViewViewModelBlueprint/Public/MVVMBlueprintView.h
+++ b/Engine/Plugins/Runtime/ModelViewViewModel/Source/ModelViewViewModelBlueprint/Public/MVVMBlueprintView.h
@@ -63,6 +63,8 @@ public:
 
 #if WITH_EDITOR
 	virtual void PostEditChangeProperty(FPropertyChangedEvent& PropertyChangedEvent) override;
+	virtual void PostEditChangeChainProperty(FPropertyChangedChainEvent& PropertyChainEvent) override;
+
 	void WidgetRenamed(FName OldObjectName, FName NewObjectName);
 #endif
 
diff --git a/Engine/Plugins/Runtime/ModelViewViewModel/Source/ModelViewViewModelEditor/Private/Widgets/SMVVMViewBindingListView.cpp b/Engine/Plugins/Runtime/ModelViewViewModel/Source/ModelViewViewModelEditor/Private/Widgets/SMVVMViewBindingListView.cpp
--- a/Engine/Plugins/Runtime/ModelViewViewModel/Source/ModelViewViewModelEditor/Private/Widgets/SMVVMViewBindingListView.cpp
+++ b/Engine/Plugins/Runtime/ModelViewViewModel/Source/ModelViewViewModelEditor/Private/Widgets/SMVVMViewBindingListView.cpp
@@ -16,6 +16,7 @@
 #include "MVVMPropertyPathHelpers.h"
 #include "MVVMSubsystem.h"
 #include "MVVMWidgetBlueprintExtension_View.h"
+#include "ScopedTransaction.h"
 #include "SEnumCombo.h"
 #include "SSimpleButton.h" 
 #include "Styling/MVVMEditorStyle.h"
@@ -239,6 +240,7 @@ public:
 				[
 					SNew(SMVVMConversionPath, WidgetBlueprint, false)
 					.Bindings(this, &SMVVMViewBindingListEntryRow::GetThisViewBindingAsArray)
+					.OnFunctionChanged(this, &SMVVMViewBindingListEntryRow::OnConversionFunctionChanged, false)
 				]
 				+ SHorizontalBox::Slot()
 				.VAlign(VAlign_Center)
@@ -247,6 +249,7 @@ public:
 				[
 					SNew(SMVVMConversionPath, WidgetBlueprint, true)
 					.Bindings(this, &SMVVMViewBindingListEntryRow::GetThisViewBindingAsArray)
+					.OnFunctionChanged(this, &SMVVMViewBindingListEntryRow::OnConversionFunctionChanged, true)
 				];
 		}
 		else if (ColumnName == DropDownOptionsColumnName)
@@ -371,9 +374,43 @@ private:
 			];
 	}
 
+	void OnPreEditChange(FName PropertyName) const
+	{
+		if (UMVVMBlueprintView* BlueprintViewPtr = BlueprintView.Get())
+		{
+			FProperty* ChangedProperty = FMVVMBlueprintViewBinding::StaticStruct()->FindPropertyByName(PropertyName);
+			check(ChangedProperty != nullptr);
+
+			FEditPropertyChain EditChain;
+			EditChain.AddTail(UMVVMBlueprintView::StaticClass()->FindPropertyByName("Bindings"));
+			EditChain.AddTail(ChangedProperty);
+			EditChain.SetActivePropertyNode(ChangedProperty);
+
+			BlueprintViewPtr->PreEditChange(EditChain);
+		}
+	}
+
+	void OnPostEditChange(FName PropertyName) const
+	{
+		if (UMVVMBlueprintView* BlueprintViewPtr = BlueprintView.Get())
+		{
+			FProperty* ChangedProperty = FMVVMBlueprintViewBinding::StaticStruct()->FindPropertyByName(PropertyName);
+			check(ChangedProperty != nullptr);
+
+			FEditPropertyChain EditChain;
+			EditChain.AddTail(UMVVMBlueprintView::StaticClass()->FindPropertyByName("Bindings"));
+			EditChain.AddTail(ChangedProperty);
+			EditChain.SetActivePropertyNode(ChangedProperty);
+
+			FPropertyChangedEvent ChangeEvent(ChangedProperty, EPropertyChangeType::ValueSet);
+			FPropertyChangedChainEvent ChainEvent(EditChain, ChangeEvent);
+			BlueprintViewPtr->PostEditChangeChainProperty(ChainEvent);
+		}
+	}
+
 	void OnViewModelSelectionChanged(UE::MVVM::FBindingSource Source)
 	{
-		OnSourceSelectionChanged(Source, ViewModelHelper.Get());
+		OnSourceSelectionChanged(Source, ViewModelHelper.Get(), false);
 		if (ViewModelFieldSelector.IsValid())
 		{
 			ViewModelFieldSelector->Refresh();
@@ -382,23 +419,35 @@ private:
 
 	void OnWidgetSelectionChanged(UE::MVVM::FBindingSource Source)
 	{
-		OnSourceSelectionChanged(Source, WidgetHelper.Get());
+		OnSourceSelectionChanged(Source, WidgetHelper.Get(), true);
 		if (WidgetFieldSelector.IsValid())
 		{
 			WidgetFieldSelector->Refresh();
 		}
 	}
 
-	void OnSourceSelectionChanged(UE::MVVM::FBindingSource SelectedSource, UE::MVVM::IFieldPathHelper* PathHelper)
+	void OnSourceSelectionChanged(UE::MVVM::FBindingSource SelectedSource, UE::MVVM::IFieldPathHelper* PathHelper, bool bIsWidget)
 	{
 		if (UMVVMBlueprintView* BlueprintViewPtr = BlueprintView.Get())
 		{
-			BlueprintViewPtr->PreEditChange(UMVVMBlueprintView::StaticClass()->FindPropertyByName("Bindings"));
+			FScopedTransaction Transaction(LOCTEXT("SetBindingSource", "Set Binding Source"));
+
+			FName ChangedProperty;
+			if (bIsWidget)
+			{
+				ChangedProperty = GET_MEMBER_NAME_CHECKED(FMVVMBlueprintViewBinding, WidgetPath);
+			}
+			else
+			{
+				ChangedProperty = GET_MEMBER_NAME_CHECKED(FMVVMBlueprintViewBinding, ViewModelPath);
+			}
+
+			OnPreEditChange(ChangedProperty);
 
 			PathHelper->SetSelectedSource(SelectedSource);
 			PathHelper->ResetBinding(); // Might make sense to keep this around in case we retarget to a compatible widget or switch back.
 
-			BlueprintViewPtr->PostEditChange();
+			OnPostEditChange(ChangedProperty);
 		}
 	}
 
@@ -407,7 +456,7 @@ private:
 		UE::MVVM::FMVVMConstFieldVariant CurrentField = ViewModelHelper->GetSelectedField();
 		if (CurrentField != SelectedField)
 		{
-			OnPropertySelectionChanged(SelectedField, ViewModelHelper.Get());
+			OnPropertySelectionChanged(SelectedField, ViewModelHelper.Get(), false);
 
 			if (WidgetFieldSelector.IsValid())
 			{
@@ -421,7 +470,7 @@ private:
 		UE::MVVM::FMVVMConstFieldVariant CurrentField = WidgetHelper->GetSelectedField();
 		if (CurrentField != SelectedField)
 		{
-			OnPropertySelectionChanged(SelectedField, WidgetHelper.Get());
+			OnPropertySelectionChanged(SelectedField, WidgetHelper.Get(), true);
 
 			if (ViewModelFieldSelector.IsValid())
 			{
@@ -430,16 +479,25 @@ private:
 		}
 	}
 
-	void OnPropertySelectionChanged(UE::MVVM::FMVVMConstFieldVariant SelectedField, UE::MVVM::IFieldPathHelper* PathHelper)
+	void OnPropertySelectionChanged(UE::MVVM::FMVVMConstFieldVariant SelectedField, UE::MVVM::IFieldPathHelper* PathHelper, bool bIsWidget)
 	{
-		if (UMVVMBlueprintView* BlueprintViewPtr = BlueprintView.Get())
+		FScopedTransaction Transaction(LOCTEXT("SetBindingProperty", "Set Binding Property"));
+
+		FName ChangedProperty;
+		if (bIsWidget)
 		{
-			BlueprintViewPtr->PreEditChange(UMVVMBlueprintView::StaticClass()->FindPropertyByName("Bindings"));
+			ChangedProperty = GET_MEMBER_NAME_CHECKED(FMVVMBlueprintViewBinding, WidgetPath);
+		}
+		else
+		{
+			ChangedProperty = GET_MEMBER_NAME_CHECKED(FMVVMBlueprintViewBinding, ViewModelPath);
+		}
 
-			PathHelper->SetBindingReference(SelectedField);
+		OnPreEditChange(ChangedProperty);
 
-			BlueprintViewPtr->PostEditChange();
-		}
+		PathHelper->SetBindingReference(SelectedField);
+
+		OnPostEditChange(ChangedProperty);
 	}
 
 	void OnUpdateModeSelectionChanged(int32 Value, ESelectInfo::Type)
@@ -448,18 +506,45 @@ private:
 		{
 			if (ViewModelBinding->UpdateMode != (EMVVMViewBindingUpdateMode) Value)
 			{
-				if (UMVVMBlueprintView* BlueprintViewPtr = BlueprintView.Get())
-				{
-					BlueprintViewPtr->PreEditChange(UMVVMBlueprintView::StaticClass()->FindPropertyByName("Bindings"));
+				FScopedTransaction Transaction(LOCTEXT("SetUpdateMode", "Set Update Mode"));
 
-					ViewModelBinding->UpdateMode = (EMVVMViewBindingUpdateMode) Value;
+				OnPreEditChange(GET_MEMBER_NAME_CHECKED(FMVVMBlueprintViewBinding, UpdateMode));
 
-					BlueprintViewPtr->PostEditChange();
-				}
+				ViewModelBinding->UpdateMode = (EMVVMViewBindingUpdateMode) Value;
+
+				OnPostEditChange(GET_MEMBER_NAME_CHECKED(FMVVMBlueprintViewBinding, UpdateMode));
 			}
 		}
 	}
 
+	void OnConversionFunctionChanged(const UFunction* Function, bool bSourceToDestination)
+	{
+		if (FMVVMBlueprintViewBinding* ViewModelBinding = GetThisViewBinding())
+		{
+			FScopedTransaction Transaction(LOCTEXT("SetConversionFunction", "Set Conversion Function"));
+
+			FMemberReference NewReference;
+			if (Function != nullptr)
+			{
+				NewReference.SetFromField<UFunction>(Function, WidgetBlueprint->SkeletonGeneratedClass);
+			}
+				
+			OnPreEditChange(GET_MEMBER_NAME_CHECKED(FMVVMBlueprintViewBinding, Conversion));
+
+			if (bSourceToDestination)
+			{
+				ViewModelBinding->Conversion.SourceToDestinationFunction = NewReference;
+			}
+			else
+			{
+				ViewModelBinding->Conversion.DestinationToSourceFunction = NewReference;
+			}
+
+			OnPostEditChange(GET_MEMBER_NAME_CHECKED(FMVVMBlueprintViewBinding, Conversion));
+		}
+	}
+
+
 	int32 GetUpdateModeValue() const
 	{
 		return (int32) GetThisViewBinding()->UpdateMode;
@@ -477,14 +562,11 @@ private:
 			bool bNewEnabled = (NewState == ECheckBoxState::Checked);
 			if (ViewModelBinding->bEnabled != bNewEnabled)
 			{
-				if (UMVVMBlueprintView* BlueprintViewPtr = BlueprintView.Get())
-				{
-					BlueprintViewPtr->PreEditChange(UMVVMBlueprintView::StaticClass()->FindPropertyByName("Bindings"));
+				OnPreEditChange(GET_MEMBER_NAME_CHECKED(FMVVMBlueprintViewBinding, bEnabled));
 
-					ViewModelBinding->bEnabled = bNewEnabled;
+				ViewModelBinding->bEnabled = bNewEnabled;
 
-					BlueprintViewPtr->PostEditChange();
-				}
+				OnPostEditChange(GET_MEMBER_NAME_CHECKED(FMVVMBlueprintViewBinding, bEnabled));
 			}
 		}
 	}
@@ -600,14 +682,11 @@ private:
 
 			if (ViewModelBinding->BindingType != NewMode)
 			{
-				if (UMVVMBlueprintView* BlueprintViewPtr = BlueprintView.Get())
-				{
-					BlueprintViewPtr->PreEditChange(UMVVMBlueprintView::StaticClass()->FindPropertyByName("Bindings"));
+				OnPreEditChange(GET_MEMBER_NAME_CHECKED(FMVVMBlueprintViewBinding, BindingType));
 
-					ViewModelBinding->BindingType = NewMode;
+				ViewModelBinding->BindingType = NewMode;
 
-					BlueprintViewPtr->PostEditChange();
-				}
+				OnPostEditChange(GET_MEMBER_NAME_CHECKED(FMVVMBlueprintViewBinding, BindingType));
 
 				ViewModelFieldSelector->Refresh();
 				WidgetFieldSelector->Refresh();
@@ -658,7 +737,7 @@ private:
 private:
 	FMVVMViewBindingListEntryPtr Entry;
 	TWeakObjectPtr<UMVVMBlueprintView> BlueprintView;
-	UWidgetBlueprint* WidgetBlueprint { nullptr };
+	UWidgetBlueprint* WidgetBlueprint = nullptr;
 	TUniquePtr<UE::MVVM::FViewModelFieldPathHelper> ViewModelHelper;
 	TUniquePtr<UE::MVVM::FWidgetFieldPathHelper> WidgetHelper;
 	TSharedPtr<SMVVMSourceSelector> ViewModelSourceSelector;
@@ -667,8 +746,6 @@ private:
 	TSharedPtr<SMVVMFieldSelector> WidgetFieldSelector;
 	TSharedPtr<SWidget> ContextMenuOptionHelper;
 	TSharedPtr<SCustomDialog> ErrorDialog;
-	TSharedPtr<SMenuAnchor> GetterConversionFunctionAnchor;
-	TSharedPtr<SMenuAnchor> SetterConversionFunctionAnchor;
 	TArray<TSharedPtr<FText>> ErrorItems;
 	TArray<FName> ModeNames;
 	FDelegateHandle OnBlueprintChangedHandle;
@@ -769,6 +846,23 @@ SMVVMViewBindingListView::~SMVVMViewBindingListView()
 
 void SMVVMViewBindingListView::RequestListRefresh()
 {
+	int32 SelectedIndex = -1;
+	if (ListView.IsValid())
+	{
+		if (const UMVVMWidgetBlueprintExtension_View* MVVMExtensionPtr = MVVMExtension.Get())
+		{
+			if (const UMVVMBlueprintView* BlueprintView = MVVMExtensionPtr->GetBlueprintView())
+			{
+				TArray<FMVVMViewBindingListEntryPtr> SelectedItems = ListView->GetSelectedItems();
+				for (const FMVVMViewBindingListEntryPtr& Entry : SelectedItems)
+				{
+					SelectedIndex = Entry->Index;
+					break;
+				}
+			}
+		}
+	}
+
 	SourceData.Reset();
 	if (const UMVVMWidgetBlueprintExtension_View* MVVMExtensionPtr = MVVMExtension.Get())
 	{
@@ -785,6 +879,11 @@ void SMVVMViewBindingListView::RequestListRefresh()
 	if (ListView.IsValid())
 	{
 		ListView->RequestListRefresh();
+
+		if (SourceData.IsValidIndex(SelectedIndex))
+		{
+			ListView->SetItemSelection(SourceData[SelectedIndex], true);
+		}
 	}
 }
 
