MVVM: Add message when no path exist to get a viewmodel at runtime

diff --git a/Engine/Plugins/Runtime/ModelViewViewModel/Source/ModelViewViewModelEditor/Private/Widgets/SMVVMViewModelContextListWidget.cpp b/Engine/Plugins/Runtime/ModelViewViewModel/Source/ModelViewViewModelEditor/Private/Widgets/SMVVMViewModelContextListWidget.cpp
--- a/Engine/Plugins/Runtime/ModelViewViewModel/Source/ModelViewViewModelEditor/Private/Widgets/SMVVMViewModelContextListWidget.cpp
+++ b/Engine/Plugins/Runtime/ModelViewViewModel/Source/ModelViewViewModelEditor/Private/Widgets/SMVVMViewModelContextListWidget.cpp
@@ -392,6 +392,9 @@ namespace UE::MVVM::Private
 			UClass* InBindableClass = Cast<UClass>(InOwnerStruct);
 			TSubclassOf<UMVVMViewModelBase> ViewModelClass = Entry->GetViewModelClass();
 
+			check(ViewModelClass.Get());
+
+			bool bHasFunctionEntry = false;
 			MenuBuilder.BeginSection("Functions", LOCTEXT("Functions", "Functions"));
 			for (TFieldIterator<UFunction> FuncIt(InBindableClass, EFieldIteratorFlags::IncludeSuper); FuncIt; ++FuncIt)
 			{
@@ -408,6 +411,7 @@ namespace UE::MVVM::Private
 						FUIAction(FExecuteAction::CreateSP(this, &SMVVMManageViewModelsListEntryRow::HandleAddBinding, NewBindingChain)),
 						MakeFieldWidget(FMVVMConstFieldVariant(Function))
 					);
+					bHasFunctionEntry = true;
 				}
 				else
 				{
@@ -421,6 +425,7 @@ namespace UE::MVVM::Private
 								MakeFieldWidget(FMVVMConstFieldVariant(Function)),
 								FNewMenuDelegate::CreateSP(this, &SMVVMManageViewModelsListEntryRow::GeneratePropertyPathMenuContent, static_cast<UStruct*>(ObjectPropertyBase->PropertyClass), NewBindingChain)
 							);
+							bHasFunctionEntry = true;
 						}
 					}
 					else if (const FStructProperty* StructProperty = CastField<FStructProperty>(ReturnProperty))
@@ -431,12 +436,14 @@ namespace UE::MVVM::Private
 								MakeFieldWidget(FMVVMConstFieldVariant(Function)),
 								FNewMenuDelegate::CreateSP(this, &SMVVMManageViewModelsListEntryRow::GeneratePropertyPathMenuContent, static_cast<UStruct*>(StructProperty->Struct), NewBindingChain)
 							);
+							bHasFunctionEntry = true;
 						}
 					}
 				}
 			}
 			MenuBuilder.EndSection(); //Functions
 
+			bool bHasPropertyEntry = false;
 			MenuBuilder.BeginSection("Properties", LOCTEXT("Properties", "Properties"));
 			for (TFieldIterator<FProperty> PropIt(InBindableClass, EFieldIteratorFlags::IncludeSuper); PropIt; ++PropIt)
 			{
@@ -452,6 +459,7 @@ namespace UE::MVVM::Private
 						FUIAction(FExecuteAction::CreateSP(this, &SMVVMManageViewModelsListEntryRow::HandleAddBinding, NewBindingChain)),
 						MakeFieldWidget(FMVVMConstFieldVariant(Property))
 					);
+					bHasPropertyEntry = true;
 				}
 				else
 				{
@@ -465,6 +473,7 @@ namespace UE::MVVM::Private
 								MakeFieldWidget(FMVVMConstFieldVariant(Property)),
 								FNewMenuDelegate::CreateSP(this, &SMVVMManageViewModelsListEntryRow::GeneratePropertyPathMenuContent, static_cast<UStruct*>(ObjectPropertyBase->PropertyClass), NewBindingChain)
 							);
+							bHasPropertyEntry = true;
 						}
 					}
 					else if (const FStructProperty* StructProperty = CastField<FStructProperty>(Property))
@@ -475,11 +484,22 @@ namespace UE::MVVM::Private
 								MakeFieldWidget(FMVVMConstFieldVariant(Property)),
 								FNewMenuDelegate::CreateSP(this, &SMVVMManageViewModelsListEntryRow::GeneratePropertyPathMenuContent, static_cast<UStruct*>(StructProperty->Struct), NewBindingChain)
 							);
+							bHasPropertyEntry = true;
 						}
 					}
 				}
 			}
 			MenuBuilder.EndSection(); //Properties
+
+			if (!bHasFunctionEntry && !bHasPropertyEntry)
+			{
+				FMenuEntryParams MenuEntry;
+				MenuEntry.DirectActions.ExecuteAction = FExecuteAction::CreateLambda([](){});
+				MenuEntry.DirectActions.CanExecuteAction = FCanExecuteAction::CreateLambda([](){ return false; });
+				MenuEntry.LabelOverride = LOCTEXT("NoPropertyPathEntry", "No Available Path");
+				MenuEntry.ToolTipOverride = FText::Format(LOCTEXT("NoPropertyPathEntryToolTip", "No available path that returns a viewmodel of the '{0}' type."), ViewModelClass->GetDisplayNameText());
+				MenuBuilder.AddMenuEntry(MenuEntry);
+			}
 		}
 
 		TSharedPtr<FMVVMBlueprintViewModelContext> Entry;
