MVVM: Generate full path for Object property. It now load properly and support downcast.
#prefligh 628cd4005665463c2104be4b

diff --git a/Engine/Plugins/Runtime/ModelViewViewModel/Source/ModelViewViewModel/Private/Bindings/MVVMCompiledBindingLibrary.cpp b/Engine/Plugins/Runtime/ModelViewViewModel/Source/ModelViewViewModel/Private/Bindings/MVVMCompiledBindingLibrary.cpp
--- a/Engine/Plugins/Runtime/ModelViewViewModel/Source/ModelViewViewModel/Private/Bindings/MVVMCompiledBindingLibrary.cpp
+++ b/Engine/Plugins/Runtime/ModelViewViewModel/Source/ModelViewViewModel/Private/Bindings/MVVMCompiledBindingLibrary.cpp
@@ -205,8 +205,14 @@
 	check(!Destination.GetObjectVariant().IsNull());
 	check(!Source.GetFieldVariant().IsEmpty());
 	check(!Destination.GetFieldVariant().IsEmpty());
-	check(Source.GetObjectVariant().GetOwner()->IsChildOf(Source.GetFieldVariant().GetOwner()));
-	check(Destination.GetObjectVariant().GetOwner()->IsChildOf(Destination.GetFieldVariant().GetOwner()));
+	if (Source.GetObjectVariant().GetOwner()->IsChildOf(Source.GetFieldVariant().GetOwner()))
+	{
+		return MakeError(EExecutionFailingReason::InvalidCast);
+	}
+	if (Destination.GetObjectVariant().GetOwner()->IsChildOf(Destination.GetFieldVariant().GetOwner()))
+	{
+		return MakeError(EExecutionFailingReason::InvalidCast);
+	}
 
 	if (ConversionFunction.GetFunction())
 	{
diff --git a/Engine/Plugins/Runtime/ModelViewViewModel/Source/ModelViewViewModel/Private/Bindings/MVVMFieldPathHelper.cpp b/Engine/Plugins/Runtime/ModelViewViewModel/Source/ModelViewViewModel/Private/Bindings/MVVMFieldPathHelper.cpp
--- a/Engine/Plugins/Runtime/ModelViewViewModel/Source/ModelViewViewModel/Private/Bindings/MVVMFieldPathHelper.cpp
+++ b/Engine/Plugins/Runtime/ModelViewViewModel/Source/ModelViewViewModel/Private/Bindings/MVVMFieldPathHelper.cpp
@@ -1,214 +1,306 @@
-// Copyright Epic Games, Inc. All Rights Reserved.
-
-#include "Bindings/MVVMFieldPathHelper.h"
-
-#include "Bindings/MVVMBindingHelper.h"
-
-namespace UE::MVVM::FieldPathHelper
-{
-
-namespace Private
-{
-
-static const FName NAME_BlueprintGetter = "BlueprintGetter";
-static const FName NAME_BlueprintSetter = "BlueprintSetter";
-
-} // namespace
-
-
-TValueOrError<TArray<FMVVMFieldVariant>, FString> GenerateFieldPathList(TSubclassOf<UObject> InFrom, FStringView InFieldPath, bool bForReading)
-{
-	if (InFrom.Get() == nullptr)
-	{
-		return MakeError(TEXT("The source class is invalid."));
-	}
-	if (InFieldPath.IsEmpty())
-	{
-		return MakeError(TEXT("The FieldPath is empty."));
-	}
-	if (InFieldPath[InFieldPath.Len() - 1] == TEXT('.'))
-	{
-		return MakeError(TEXT("The field path cannot end with a '.' character."));
-	}
-
-	auto FindContainer = [](const FProperty* Property) -> TValueOrError<UStruct*, FString>
-	{
-		const FObjectPropertyBase* ObjectProperty = CastField<const FObjectPropertyBase>(Property);
-		const FStructProperty* StructProperty = CastField<const FStructProperty>(Property);
-
-		if (ObjectProperty)
-		{
-			return MakeValue(ObjectProperty->PropertyClass);
-		}
-		else if (StructProperty)
-		{
-			return MakeValue(StructProperty->Struct);
-		}
-		return MakeError(FString::Printf(TEXT("Only object or struct properties can be used as source paths. '%s' is a '%s'."), *Property->GetName(), *Property->GetClass()->GetName()));
-	};
-
-	auto TransformWithAccessor = [](UStruct* CurrentContainer, FMVVMFieldVariant CurrentField, bool bForReading) -> TValueOrError<FMVVMFieldVariant, FString>
-	{
-#if WITH_EDITORONLY_DATA
-		if (bForReading)
-		{
-			if (!CurrentField.GetProperty()->HasGetter())
-			{
-				const FString& BlueprintGetter = CurrentField.GetProperty()->GetMetaData(Private::NAME_BlueprintGetter);
-				if (!BlueprintGetter.IsEmpty())
-				{
-					FMVVMFieldVariant NewField = BindingHelper::FindFieldByName(CurrentContainer, FMVVMBindingName(*BlueprintGetter));
-					if (NewField.IsFunction())
-					{
-						CurrentField = NewField;
-					}
-					else
-					{
-						return MakeError(FString::Printf(TEXT("The BlueprintGetter '%s' could not be found on object '%s'."), *BlueprintGetter, *CurrentContainer->GetName()));
-					}
-				}
-			}
-		}
-		else
-		{
-			if (!CurrentField.GetProperty()->HasSetter())
-			{
-				const FString& BlueprintSetter = CurrentField.GetProperty()->GetMetaData(Private::NAME_BlueprintSetter);
-				if (!BlueprintSetter.IsEmpty())
-				{
-					FMVVMFieldVariant NewField = BindingHelper::FindFieldByName(CurrentContainer, FMVVMBindingName(*BlueprintSetter));
-					if (NewField.IsFunction())
-					{
-						CurrentField = NewField;
-					}
-					else
-					{
-						return MakeError(FString::Printf(TEXT("The BlueprintSetter '%s' could not be found on object %s."), *BlueprintSetter, *CurrentContainer->GetName()));
-					}
-				}
-			}
-		}
-#endif
-		return MakeValue(CurrentField);
-	};
-
-	TArray<FMVVMFieldVariant> Result;
-	UStruct* CurrentContainer = InFrom.Get();
-
-	// Split the string into property or function names
-	//ie. myvar.myfunction.myvar
-	int32 FoundIndex = INDEX_NONE;
-	while (InFieldPath.FindChar(TEXT('.'), FoundIndex))
-	{
-		FMVVMFieldVariant Field = BindingHelper::FindFieldByName(CurrentContainer, FMVVMBindingName(FName(FoundIndex, InFieldPath.GetData())));
-		if (Field.IsEmpty())
-		{
-			return MakeError(FString::Printf(TEXT("The field '%s' does not exist in the struct '%s'."), *FName(FoundIndex, InFieldPath.GetData()).ToString(), CurrentContainer ? *CurrentContainer->GetName() : TEXT("<none>")));
-		}
-		else if (Field.IsProperty())
-		{
-			TValueOrError<FMVVMFieldVariant, FString> TransformedField = TransformWithAccessor(CurrentContainer, Field, true);
-			if (TransformedField.HasError())
-			{
-				return MakeError(TransformedField.StealError());
-			}
-			Field = TransformedField.StealValue();
-			check(!Field.IsEmpty());
-
-			if (Field.IsProperty())
-			{
-				TValueOrError<UStruct*, FString> FoundContainer = FindContainer(Field.GetProperty());
-				if (FoundContainer.HasError())
-				{
-					return MakeError(FoundContainer.StealError());
-				}
-				CurrentContainer = FoundContainer.GetValue();
-			}
-		}
-		
-		if (Field.IsFunction())
-		{
-			const FProperty* ReturnProperty = BindingHelper::GetReturnProperty(Field.GetFunction());
-			TValueOrError<UStruct*, FString> FoundContainer = FindContainer(ReturnProperty);
-			if (FoundContainer.HasError())
-			{
-				return MakeError(FoundContainer.StealError());
-			}
-			CurrentContainer = FoundContainer.GetValue();
-		}
-
-		InFieldPath = InFieldPath.RightChop(FoundIndex + 1);
-		Result.Add(Field);
-	}
-
-	// The last field can be anything (that is what we are going to bind to)
-	if (InFieldPath.Len() > 0)
-	{
-		FMVVMFieldVariant Field = BindingHelper::FindFieldByName(CurrentContainer, FMVVMBindingName(InFieldPath.GetData()));
-		if (Field.IsEmpty())
-		{
-			return MakeError(FString::Printf(TEXT("The field '%s' does not exist in the struct '%s'."), InFieldPath.GetData(), CurrentContainer ? *CurrentContainer->GetName() : TEXT("<none>")));
-		}
-
-		if (Field.IsProperty())
-		{
-			TValueOrError<FMVVMFieldVariant, FString> TransformedField = TransformWithAccessor(CurrentContainer, Field, bForReading);
-			if (TransformedField.HasError())
-			{
-				return MakeError(TransformedField.StealError());
-			}
-			Field = TransformedField.StealValue();
-		}
-		check(!Field.IsEmpty());
-		Result.Add(Field);
-	}
-
-	return MakeValue(MoveTemp(Result));
-}
-
-
-TValueOrError<TArray<FMVVMFieldVariant>, FString> GenerateConversionFunctionFieldPathList(TSubclassOf<UObject> From, FStringView FieldPath)
-{
-	return MakeError(TEXT("Not implemented yet"));
-}
-
-
-TValueOrError<UObject*, void> EvaluateObjectProperty(const FFieldContext& InSource)
-{
-	if (InSource.GetObjectVariant().IsNull())
-	{
-		return MakeError();
-	}
-
-	const bool bIsProperty = InSource.GetFieldVariant().IsProperty();
-	const FProperty* GetterType = bIsProperty ? InSource.GetFieldVariant().GetProperty() : BindingHelper::GetReturnProperty(InSource.GetFieldVariant().GetFunction());
-	check(GetterType);
-
-	const FObjectPropertyBase* ObjectProperty = CastField<const FObjectPropertyBase>(GetterType);
-	if (ObjectProperty == nullptr)
-	{
-		return MakeError();
-	}
-
-	if (bIsProperty)
-	{
-		return MakeValue(ObjectProperty->GetObjectPropertyValue_InContainer(InSource.GetObjectVariant().GetData()));
-	}
-	else
-	{
-		check(InSource.GetObjectVariant().IsUObject());
-		check(InSource.GetObjectVariant().GetUObject());
-
-		UFunction* Function = InSource.GetFieldVariant().GetFunction();
-		check(InSource.GetObjectVariant().GetUObject()->GetClass()->IsChildOf(Function->GetOuterUClass()));
-
-		void* DataPtr = FMemory_Alloca_Aligned(Function->ParmsSize, Function->GetMinAlignment());
-		ObjectProperty->InitializeValue(DataPtr);
-		InSource.GetObjectVariant().GetUObject()->ProcessEvent(Function, DataPtr);
-		UObject* Result = ObjectProperty->GetObjectPropertyValue_InContainer(DataPtr);
-		ObjectProperty->DestroyValue(DataPtr);
-		return MakeValue(Result);
-	}
-}
-
-} // namespace
+// Copyright Epic Games, Inc. All Rights Reserved.
+
+#include "Bindings/MVVMFieldPathHelper.h"
+
+#include "Bindings/MVVMBindingHelper.h"
+#include "Misc/MemStack.h"
+
+namespace UE::MVVM::FieldPathHelper
+{
+
+namespace Private
+{
+
+static const FName NAME_BlueprintGetter = "BlueprintGetter";
+static const FName NAME_BlueprintSetter = "BlueprintSetter";
+
+
+TValueOrError<UStruct*, FString> FindContainer(const FProperty* Property)
+{
+	const FObjectPropertyBase* ObjectProperty = CastField<const FObjectPropertyBase>(Property);
+	const FStructProperty* StructProperty = CastField<const FStructProperty>(Property);
+
+	if (ObjectProperty)
+	{
+		return MakeValue(ObjectProperty->PropertyClass);
+	}
+	else if (StructProperty)
+	{
+		return MakeValue(StructProperty->Struct);
+	}
+	return MakeError(FString::Printf(TEXT("Only object or struct properties can be used as source paths. '%s' is a '%s'."), *Property->GetName(), *Property->GetClass()->GetName()));
+}
+
+
+TValueOrError<FMVVMConstFieldVariant, FString> TransformWithAccessor(UStruct* CurrentContainer, FMVVMConstFieldVariant CurrentField, bool bForReading)
+{
+#if WITH_EDITORONLY_DATA
+	if (bForReading)
+	{
+		if (!CurrentField.GetProperty()->HasGetter())
+		{
+			const FString& BlueprintGetter = CurrentField.GetProperty()->GetMetaData(Private::NAME_BlueprintGetter);
+			if (!BlueprintGetter.IsEmpty())
+			{
+				FMVVMConstFieldVariant NewField = BindingHelper::FindFieldByName(CurrentContainer, FMVVMBindingName(*BlueprintGetter));
+				if (NewField.IsFunction())
+				{
+					CurrentField = NewField;
+				}
+				else
+				{
+					return MakeError(FString::Printf(TEXT("The BlueprintGetter '%s' could not be found on object '%s'."), *BlueprintGetter, *CurrentContainer->GetName()));
+				}
+			}
+		}
+	}
+	else
+	{
+		if (!CurrentField.GetProperty()->HasSetter())
+		{
+			const FString& BlueprintSetter = CurrentField.GetProperty()->GetMetaData(Private::NAME_BlueprintSetter);
+			if (!BlueprintSetter.IsEmpty())
+			{
+				FMVVMConstFieldVariant NewField = BindingHelper::FindFieldByName(CurrentContainer, FMVVMBindingName(*BlueprintSetter));
+				if (NewField.IsFunction())
+				{
+					CurrentField = NewField;
+				}
+				else
+				{
+					return MakeError(FString::Printf(TEXT("The BlueprintSetter '%s' could not be found on object %s."), *BlueprintSetter, *CurrentContainer->GetName()));
+				}
+			}
+		}
+	}
+#endif
+	return MakeValue(CurrentField);
+}
+
+} // namespace
+
+
+TValueOrError<TArray<FMVVMConstFieldVariant>, FString> GenerateFieldPathList(TSubclassOf<UObject> InFrom, FStringView InFieldPath, bool bForReading)
+{
+	if (InFrom.Get() == nullptr)
+	{
+		return MakeError(TEXT("The source class is invalid."));
+	}
+	if (InFieldPath.IsEmpty())
+	{
+		return MakeError(TEXT("The FieldPath is empty."));
+	}
+	if (InFieldPath[InFieldPath.Len() - 1] == TEXT('.'))
+	{
+		return MakeError(TEXT("The field path cannot end with a '.' character."));
+	}
+
+	FMemMark Mark(FMemStack::Get());
+	TArray<FMVVMConstFieldVariant, TMemStackAllocator<>> Result;
+
+	UStruct* CurrentContainer = InFrom.Get();
+
+	// Split the string into property or function names
+	//ie. myvar.myfunction.myvar
+	int32 FoundIndex = INDEX_NONE;
+	while (InFieldPath.FindChar(TEXT('.'), FoundIndex))
+	{
+		FMVVMConstFieldVariant Field = BindingHelper::FindFieldByName(CurrentContainer, FMVVMBindingName(FName(FoundIndex, InFieldPath.GetData())));
+		if (Field.IsEmpty())
+		{
+			return MakeError(FString::Printf(TEXT("The field '%s' does not exist in the struct '%s'."), *FName(FoundIndex, InFieldPath.GetData()).ToString(), CurrentContainer ? *CurrentContainer->GetName() : TEXT("<none>")));
+		}
+		else if (Field.IsProperty())
+		{
+			TValueOrError<UStruct*, FString> FoundContainer = Private::FindContainer(Field.GetProperty());
+			if (FoundContainer.HasError())
+			{
+				return MakeError(FoundContainer.StealError());
+			}
+			CurrentContainer = FoundContainer.GetValue();
+		}
+		
+		if (Field.IsFunction())
+		{
+			const FProperty* ReturnProperty = BindingHelper::GetReturnProperty(Field.GetFunction());
+			TValueOrError<UStruct*, FString> FoundContainer = Private::FindContainer(ReturnProperty);
+			if (FoundContainer.HasError())
+			{
+				return MakeError(FoundContainer.StealError());
+			}
+			CurrentContainer = FoundContainer.GetValue();
+		}
+
+		InFieldPath = InFieldPath.RightChop(FoundIndex + 1);
+		Result.Add(Field);
+	}
+
+	// The last field can be anything (that is what we are going to bind to)
+	if (InFieldPath.Len() > 0)
+	{
+		FMVVMConstFieldVariant Field = BindingHelper::FindFieldByName(CurrentContainer, FMVVMBindingName(InFieldPath.GetData()));
+		if (Field.IsEmpty())
+		{
+			return MakeError(FString::Printf(TEXT("The field '%s' does not exist in the struct '%s'."), InFieldPath.GetData(), CurrentContainer ? *CurrentContainer->GetName() : TEXT("<none>")));
+		}
+
+		Result.Add(Field);
+	}
+
+	return GenerateFieldPathList(Result, bForReading);
+}
+
+
+/**
+ * Rules for reading:
+ *		* Build path using Getter or BlueprintGetter if needed.
+ *		* If the FProperty is a FStructProperty and a function was used, then the runtime may use dynamic memory instead of stack memory.
+ * Rules for writing:
+ *		* Build path using Getter or BlueprintGetter, the last element should use Setter or BlueprintSetter.
+ *		* If one of the element in the path has a Setter or BlueprintSetter, then the path needs to stop there and be divided in 3 paths
+ *			* ex: PropertyA.PropertyB.PropertyC.PropertyD.PropertyE and PropertyC has a Setter/BlueprintSetter
+ *			* 1. To write: Property.Property.SetPropertyC()
+ *			* 2. To read: PropertyA.PropertyB.PropertyC or PropertyA.PropertyB.GetPropertyC()
+ *			* 3. To continue the reading: PropertyD.PropertyE
+ *				PropertyD and PropertyE cannot have Getter/BlueprintGetter if they are FStructProperty
+ *		* We can only have one Setter/BlueprintSetter in the path
+ */
+TValueOrError<TArray<FMVVMConstFieldVariant>, FString> GenerateFieldPathList(const TArrayView<FMVVMConstFieldVariant> InFieldPath, bool bForSourceBinding)
+{
+	if (InFieldPath.Num() == 0)
+	{
+		return MakeError(TEXT("The FieldPath is empty."));
+	}
+
+	TArray<FMVVMConstFieldVariant> Result;
+	Result.Reserve(InFieldPath.Num());
+	UStruct* CurrentContainer = InFieldPath[0].GetOwner();
+
+	for (int32 Index = 0; Index < InFieldPath.Num(); ++Index)
+	{
+		bool bLastField = Index == InFieldPath.Num() - 1;
+		FMVVMConstFieldVariant Field = InFieldPath[Index];
+
+		if (Field.IsEmpty())
+		{
+			return MakeError(FString::Printf(TEXT("The field '%d' does not exist."), Index));
+		}
+
+		bool bIsChild = Field.GetOwner()->IsChildOf(CurrentContainer);
+		bool bIsDownCast = CurrentContainer->IsChildOf(Field.GetOwner());
+		if (CurrentContainer == nullptr || !(bIsChild || bIsDownCast))
+		{
+			return MakeError(FString::Printf(TEXT("The field '%s' does not exist in the struct '%s'."), *Field.GetName().ToString(), CurrentContainer ? *CurrentContainer->GetName() : TEXT("<none>")));
+		}
+		
+		if (Field.IsProperty())
+		{
+			TValueOrError<FMVVMConstFieldVariant, FString> TransformedField = Private::TransformWithAccessor(CurrentContainer, Field, (bForSourceBinding || !bLastField));
+			if (TransformedField.HasError())
+			{
+				return MakeError(TransformedField.StealError());
+			}
+			Field = TransformedField.StealValue();
+			check(!Field.IsEmpty());
+
+			if (!bLastField && Field.IsProperty())
+			{
+				TValueOrError<UStruct*, FString> FoundContainer = Private::FindContainer(Field.GetProperty());
+				if (FoundContainer.HasError())
+				{
+					return MakeError(FoundContainer.StealError());
+				}
+				CurrentContainer = FoundContainer.GetValue();
+			}
+		}
+
+		if (!bLastField && Field.IsFunction())
+		{
+			const FProperty* ReturnProperty = BindingHelper::GetReturnProperty(Field.GetFunction());
+			TValueOrError<UStruct*, FString> FoundContainer = Private::FindContainer(ReturnProperty);
+			if (FoundContainer.HasError())
+			{
+				return MakeError(FoundContainer.StealError());
+			}
+			CurrentContainer = FoundContainer.GetValue();
+		}
+
+		Result.Add(Field);
+	}
+
+	return MakeValue(MoveTemp(Result));
+}
+
+
+FString ToString(const TArrayView<FMVVMFieldVariant> Fields)
+{
+	TStringBuilder<512> Builder;
+	for (int32 Index = 0; Index < Fields.Num(); ++Index)
+	{
+		if (Index != 0)
+		{
+			Builder << TEXT(".");
+		}
+		Builder << Fields[Index].GetName();
+	}
+	return Builder.ToString();
+}
+
+
+FString ToString(const TArrayView<FMVVMConstFieldVariant> Fields)
+{
+	TStringBuilder<512> Builder;
+	for (int32 Index = 0; Index < Fields.Num(); ++Index)
+	{
+		if (Index != 0)
+		{
+			Builder << TEXT(".");
+		}
+		Builder << Fields[Index].GetName();
+	}
+	return Builder.ToString();
+}
+
+
+TValueOrError<UObject*, void> EvaluateObjectProperty(const FFieldContext& InSource)
+{
+	if (InSource.GetObjectVariant().IsNull())
+	{
+		return MakeError();
+	}
+
+	const bool bIsProperty = InSource.GetFieldVariant().IsProperty();
+	const FProperty* GetterType = bIsProperty ? InSource.GetFieldVariant().GetProperty() : BindingHelper::GetReturnProperty(InSource.GetFieldVariant().GetFunction());
+	check(GetterType);
+
+	const FObjectPropertyBase* ObjectProperty = CastField<const FObjectPropertyBase>(GetterType);
+	if (ObjectProperty == nullptr)
+	{
+		return MakeError();
+	}
+
+	if (bIsProperty)
+	{
+		return MakeValue(ObjectProperty->GetObjectPropertyValue_InContainer(InSource.GetObjectVariant().GetData()));
+	}
+	else
+	{
+		check(InSource.GetObjectVariant().IsUObject());
+		check(InSource.GetObjectVariant().GetUObject());
+
+		UFunction* Function = InSource.GetFieldVariant().GetFunction();
+		if (InSource.GetObjectVariant().GetUObject()->GetClass()->IsChildOf(Function->GetOuterUClass()))
+		{
+			return MakeError();
+		}
+
+		void* DataPtr = FMemory_Alloca_Aligned(Function->ParmsSize, Function->GetMinAlignment());
+		ObjectProperty->InitializeValue(DataPtr);
+		InSource.GetObjectVariant().GetUObject()->ProcessEvent(Function, DataPtr);
+		UObject* Result = ObjectProperty->GetObjectPropertyValue_InContainer(DataPtr);
+		ObjectProperty->DestroyValue(DataPtr);
+		return MakeValue(Result);
+	}
+}
+
+} // namespace
diff --git a/Engine/Plugins/Runtime/ModelViewViewModel/Source/ModelViewViewModel/Public/Bindings/MVVMCompiledBindingLibrary.h b/Engine/Plugins/Runtime/ModelViewViewModel/Source/ModelViewViewModel/Public/Bindings/MVVMCompiledBindingLibrary.h
--- a/Engine/Plugins/Runtime/ModelViewViewModel/Source/ModelViewViewModel/Public/Bindings/MVVMCompiledBindingLibrary.h
+++ b/Engine/Plugins/Runtime/ModelViewViewModel/Source/ModelViewViewModel/Public/Bindings/MVVMCompiledBindingLibrary.h
@@ -255,6 +255,7 @@
 		InvalidSource,
 		InvalidDestination,
 		InvalidConversionFunction,
+		InvalidCast,
 	};
 
 	/** Fetch the FProperty and UFunction. */
diff --git a/Engine/Plugins/Runtime/ModelViewViewModel/Source/ModelViewViewModel/Public/Bindings/MVVMFieldPathHelper.h b/Engine/Plugins/Runtime/ModelViewViewModel/Source/ModelViewViewModel/Public/Bindings/MVVMFieldPathHelper.h
--- a/Engine/Plugins/Runtime/ModelViewViewModel/Source/ModelViewViewModel/Public/Bindings/MVVMFieldPathHelper.h
+++ b/Engine/Plugins/Runtime/ModelViewViewModel/Source/ModelViewViewModel/Public/Bindings/MVVMFieldPathHelper.h
@@ -1,24 +1,28 @@
-// Copyright Epic Games, Inc. All Rights Reserved.
-
-#pragma once
-
-#include "CoreMinimal.h"
-#include "Templates/SubclassOf.h"
-#include "Templates/ValueOrError.h"
-#include "Types/MVVMFieldContext.h"
-#include "UObject/FieldPath.h"
-#include "UObject/Object.h"
-
-
-namespace UE::MVVM::FieldPathHelper
-{
-	/** */
-	UE_NODISCARD MODELVIEWVIEWMODEL_API TValueOrError<TArray<FMVVMFieldVariant>, FString> GenerateFieldPathList(TSubclassOf<UObject> From, FStringView FieldPath, bool bForSourceBinding);
-
-	/** */
-	UE_NODISCARD MODELVIEWVIEWMODEL_API TValueOrError<TArray<FMVVMFieldVariant>, FString> GenerateConversionFunctionFieldPathList(TSubclassOf<UObject> From, FStringView FieldPath);
-
-	/** */
-	UE_NODISCARD MODELVIEWVIEWMODEL_API TValueOrError<UObject*, void> EvaluateObjectProperty(const FFieldContext& InSource);
-
-} // namespace
+// Copyright Epic Games, Inc. All Rights Reserved.
+
+#pragma once
+
+#include "CoreMinimal.h"
+#include "Templates/SubclassOf.h"
+#include "Templates/ValueOrError.h"
+#include "Types/MVVMFieldContext.h"
+#include "UObject/FieldPath.h"
+#include "UObject/Object.h"
+
+
+namespace UE::MVVM::FieldPathHelper
+{
+	/** */
+	UE_NODISCARD MODELVIEWVIEWMODEL_API TValueOrError<TArray<FMVVMConstFieldVariant>, FString> GenerateFieldPathList(TSubclassOf<UObject> From, FStringView FieldPath, bool bForSourceBinding);
+
+	/** */
+	UE_NODISCARD MODELVIEWVIEWMODEL_API TValueOrError<TArray<FMVVMConstFieldVariant>, FString> GenerateFieldPathList(const TArrayView<FMVVMConstFieldVariant>, bool bForSourceBinding);
+
+	/** */
+	UE_NODISCARD MODELVIEWVIEWMODEL_API FString ToString(const TArrayView<FMVVMFieldVariant> Fields);
+	UE_NODISCARD MODELVIEWVIEWMODEL_API FString ToString(const TArrayView<FMVVMConstFieldVariant> Fields);
+
+	/** */
+	UE_NODISCARD MODELVIEWVIEWMODEL_API TValueOrError<UObject*, void> EvaluateObjectProperty(const FFieldContext& InSource);
+
+} // namespace
diff --git a/Engine/Plugins/Runtime/ModelViewViewModel/Source/ModelViewViewModel/Public/Types/MVVMFieldVariant.h b/Engine/Plugins/Runtime/ModelViewViewModel/Source/ModelViewViewModel/Public/Types/MVVMFieldVariant.h
--- a/Engine/Plugins/Runtime/ModelViewViewModel/Source/ModelViewViewModel/Public/Types/MVVMFieldVariant.h
+++ b/Engine/Plugins/Runtime/ModelViewViewModel/Source/ModelViewViewModel/Public/Types/MVVMFieldVariant.h
@@ -131,33 +131,6 @@
 		{
 			return !(*this == B);
 		}
-		
-		/**
-		 * Create a serializable member reference from this field.
-		 * 
-		 * Generally should not need to set self context to true. If you are running into reference issues
-		 * see param comment & consider reviewing UBlueprintVariableNodeSpawner::CreateFromMemberOrParam
-		 * 
-		 * @param  bInSelfContext	Pass true if a BP skeleton class owns the referenced property (Class not BP generated or otherwise).
-		 * @return A newly allocated FMemberReference representing this field variant.
-		 */
-		FMemberReference CreateMemberReference(bool bInSelfContext = false) const
-		{
-			FMemberReference BindingReference = FMemberReference();
-			if (IsProperty())
-			{
-				PropertyType* Property = GetProperty();
-				BindingReference.SetFromField<FProperty>(Property, bInSelfContext);
-			}
-			else if (IsFunction())
-			{
-				// Functions should never be self context references
-				FunctionType* Function = GetFunction();
-				bool bSelfContext = false;
-				BindingReference.SetFromField<UFunction>(Function, bSelfContext);
-			}
-			return BindingReference;
-		}
 
 		friend int32 GetTypeHash(const TMVVMFieldVariant<bConst>& Variant)
 		{
diff --git a/Engine/Plugins/Runtime/ModelViewViewModel/Source/ModelViewViewModelBlueprint/Private/Bindings/MVVMCompiledBindingLibraryCompiler.cpp b/Engine/Plugins/Runtime/ModelViewViewModel/Source/ModelViewViewModelBlueprint/Private/Bindings/MVVMCompiledBindingLibraryCompiler.cpp
--- a/Engine/Plugins/Runtime/ModelViewViewModel/Source/ModelViewViewModelBlueprint/Private/Bindings/MVVMCompiledBindingLibraryCompiler.cpp
+++ b/Engine/Plugins/Runtime/ModelViewViewModel/Source/ModelViewViewModelBlueprint/Private/Bindings/MVVMCompiledBindingLibraryCompiler.cpp
@@ -140,6 +140,11 @@
 {
 	Impl->bCompiled = false;
 
+	if (FieldId.IsNone())
+	{
+		return MakeError(TEXT("The Field/Property is not defined."));
+	}
+
 	if (!SourceClass->ImplementsInterface(UNotifyFieldValueChanged::StaticClass()))
 	{
 		return MakeError(FString::Printf(TEXT("'%s' doesn't implement the NotifyFieldValueChanged interface.")
@@ -181,20 +186,34 @@
 {
 	Impl->bCompiled = false;
 
-	TValueOrError<TArray<FMVVMFieldVariant>, FString> GeneratedField = FieldPathHelper::GenerateFieldPathList(InSourceClass, InFieldPath, bInRead);
+	TValueOrError<TArray<FMVVMConstFieldVariant>, FString> GeneratedField = FieldPathHelper::GenerateFieldPathList(InSourceClass, InFieldPath, bInRead);
 	if (GeneratedField.HasError())
 	{
 		return MakeError(GeneratedField.StealError());
 	}
 
-	return AddFieldPath(MakeArrayView(GeneratedField.GetValue()), bInRead);
+	return AddFieldPathImpl(MakeArrayView(GeneratedField.GetValue()), bInRead);
 }
 
 
-TValueOrError<FCompiledBindingLibraryCompiler::FFieldPathHandle, FString> FCompiledBindingLibraryCompiler::AddFieldPath(TArrayView<FMVVMFieldVariant> InFieldPath, bool bInRead)
+TValueOrError<FCompiledBindingLibraryCompiler::FFieldPathHandle, FString> FCompiledBindingLibraryCompiler::AddFieldPath(TArrayView<FMVVMConstFieldVariant> InFieldPath, bool bInRead)
 {
 	Impl->bCompiled = false;
 
+	TValueOrError<TArray<FMVVMConstFieldVariant>, FString> GeneratedField = FieldPathHelper::GenerateFieldPathList(InFieldPath, bInRead);
+	if (GeneratedField.HasError())
+	{
+		return MakeError(GeneratedField.StealError());
+	}
+
+	return AddFieldPathImpl(MakeArrayView(GeneratedField.GetValue()), bInRead);
+}
+
+
+TValueOrError<FCompiledBindingLibraryCompiler::FFieldPathHandle, FString> FCompiledBindingLibraryCompiler::AddFieldPathImpl(TArrayView<FMVVMConstFieldVariant> InFieldPath, bool bInRead)
+{
+	Impl->bCompiled = false;
+
 	auto ValidateContainer = [](const FProperty* Property, bool bShouldBeInsideContainer, bool bIsObjectOrScriptStruct) -> FString
 	{	
 		const UStruct* OwnerStruct = Property->GetOwnerStruct();
@@ -322,7 +341,7 @@
 
 	check(ExpectedType);
 
-	TValueOrError<TArray<FMVVMFieldVariant>, FString> GeneratedField = FieldPathHelper::GenerateFieldPathList(InSourceClass, InFieldPath, bInRead);
+	TValueOrError<TArray<FMVVMConstFieldVariant>, FString> GeneratedField = FieldPathHelper::GenerateFieldPathList(InSourceClass, InFieldPath, bInRead);
 	if (GeneratedField.HasError())
 	{
 		return MakeError(GeneratedField.StealError());
@@ -351,7 +370,7 @@
 		return MakeError(FString::Printf(TEXT("The field does not return a '%s'."), *ExpectedType->GetName()));
 	}
 
-	return AddFieldPath(MakeArrayView(GeneratedField.GetValue()), bInRead);
+	return AddFieldPathImpl(MakeArrayView(GeneratedField.GetValue()), bInRead);
 }
 
 
diff --git a/Engine/Plugins/Runtime/ModelViewViewModel/Source/ModelViewViewModelBlueprint/Private/MVVMViewBlueprintCompiler.cpp b/Engine/Plugins/Runtime/ModelViewViewModel/Source/ModelViewViewModelBlueprint/Private/MVVMViewBlueprintCompiler.cpp
--- a/Engine/Plugins/Runtime/ModelViewViewModel/Source/ModelViewViewModelBlueprint/Private/MVVMViewBlueprintCompiler.cpp
+++ b/Engine/Plugins/Runtime/ModelViewViewModel/Source/ModelViewViewModelBlueprint/Private/MVVMViewBlueprintCompiler.cpp
@@ -639,29 +639,29 @@
 		}
 
 		FMVVMViewBlueprintCompiler* Self = this;
-		auto AddBinding = [Self](UWidgetBlueprintGeneratedClass* Class, const TArrayView<FString> Getters, const FString& Setter, const FMemberReference& ConversionFunctionReference) -> TValueOrError<FCompilerBinding, FString>
+		auto AddBinding = [Self](UWidgetBlueprintGeneratedClass* Class, const TArrayView<TArray<UE::MVVM::FMVVMConstFieldVariant>> GetterFields, TArrayView<UE::MVVM::FMVVMConstFieldVariant> SetterFields, const FMemberReference& ConversionFunctionReference) -> TValueOrError<FCompilerBinding, FString>
 		{
 			FCompilerBinding Result;
 
-			for (const FString& Getter : Getters)
+			for (const TArrayView<UE::MVVM::FMVVMConstFieldVariant>& GetterField : GetterFields)
 			{
 				// Generate a path to read the value at runtime
-				TValueOrError<FCompiledBindingLibraryCompiler::FFieldPathHandle, FString> FieldPathResult = Self->BindingLibraryCompiler.AddFieldPath(Class, Getter, true);
+				TValueOrError<FCompiledBindingLibraryCompiler::FFieldPathHandle, FString> FieldPathResult = Self->BindingLibraryCompiler.AddFieldPath(GetterField, true);
 				if (FieldPathResult.HasError())
 				{
 					return MakeError(FString::Printf(TEXT("Couldn't create the source field path '%s'. %s")
-						, *Getter
+						, *::UE::MVVM::FieldPathHelper::ToString(GetterField)
 						, *FieldPathResult.GetError()));
 				}
 				Result.SourceRead.Add(FieldPathResult.GetValue());
 			}
 
 			{
-				TValueOrError<FCompiledBindingLibraryCompiler::FFieldPathHandle, FString> FieldPathResult = Self->BindingLibraryCompiler.AddFieldPath(Class, Setter, false);
+				TValueOrError<FCompiledBindingLibraryCompiler::FFieldPathHandle, FString> FieldPathResult = Self->BindingLibraryCompiler.AddFieldPath(SetterFields, false);
 				if (FieldPathResult.HasError())
 				{
 					return MakeError(FString::Printf(TEXT("Couldn't create the destination field path '%s'. %s")
-						, *Setter
+						, *::UE::MVVM::FieldPathHelper::ToString(SetterFields)
 						, *FieldPathResult.GetError()));
 				}
 				Result.DestinationWrite = FieldPathResult.GetValue();
@@ -711,7 +711,7 @@
 
 		if (IsForwardBinding(Binding.BindingType))
 		{
-			const FBindingSourceContext BindingSourceContext = CreateBindingSourceContext(BlueprintView, Binding.ViewModelPath);
+			const FBindingSourceContext BindingSourceContext = CreateBindingSourceContext(BlueprintView, Class, Binding.ViewModelPath);
 			if (BindingSourceContext.bIsRootWidget && BindingSourceContext.SourceClass == nullptr)
 			{
 				AddErrorForBinding(Binding, BlueprintView, TEXT("Internal error. The binding could not find its source."));
@@ -733,35 +733,22 @@
 				continue;
 			}
 
-			TArray<FString> Getters;
+			TArray<TArray<UE::MVVM::FMVVMConstFieldVariant>> GetterFields;
 			//if (Binding.Conversion.HasConversionFunction && has read values defined (that are not the default one))
 			//{
 			//  Add the read values for the conversion function
 			//}
 			//else
 			{
-				Getters.Add(BindingSourceContext.PropertyPath);
+				GetterFields.Add(BindingSourceContext.PropertyPath);
 			}
 
-			FString Setter;
+			TArray<UE::MVVM::FMVVMConstFieldVariant> SetterPath;
 			{
-				FName DestinationName = Binding.WidgetPath.GetWidgetName();
-				checkf(!DestinationName.IsNone(), TEXT("The destination should have been checked and set bAreSourceContextsValid."));
-				const bool bSourceIsUserWidget = DestinationName == Class->ClassGeneratedBy->GetFName();
-				if (bSourceIsUserWidget)
-				{
-					Setter = CreatePropertyPath(FName(), Binding.WidgetPath.GetBasePropertyPath());
-				}
-				else
-				{
-					const int32 DestinationVariableContextIndex = SourceContexts.IndexOfByPredicate([DestinationName](const FCompilerSourceContext& Other) { return Other.PropertyName == DestinationName; });
-					check(DestinationVariableContextIndex != INDEX_NONE);
-
-					Setter = CreatePropertyPath(SourceContexts[DestinationVariableContextIndex].PropertyName, Binding.WidgetPath.GetBasePropertyPath());
-				}
+				SetterPath = CreateBindingDestinationPath(BlueprintView, Class, Binding.WidgetPath);
 			}
 
-			TValueOrError<FCompilerBinding, FString> AddBindingResult = AddBinding(Class, Getters, Setter, Binding.Conversion.SourceToDestinationFunction);
+			TValueOrError<FCompilerBinding, FString> AddBindingResult = AddBinding(Class, GetterFields, SetterPath, Binding.Conversion.SourceToDestinationFunction);
 			if (AddBindingResult.HasError())
 			{
 				AddErrorForBinding(Binding, BlueprintView, FString::Printf(TEXT("The binding could not be created. %s"), *AddBindingResult.GetError()));
@@ -802,28 +789,22 @@
 				continue;
 			}
 
-			TArray<FString> Getters;
+			TArray<TArray<UE::MVVM::FMVVMConstFieldVariant>> GetterFields;
 			//if (Binding.Conversion.HasConversionFunction && has read values defined (that are not the default one))
 			//{
 			//  Add the read values for the conversion function
 			//}
 			//else
 			{
-				Getters.Add(BindingSourceContext.PropertyPath);
+				GetterFields.Add(BindingSourceContext.PropertyPath);
 			}
 
-			FString Setter;
+			TArray<UE::MVVM::FMVVMConstFieldVariant> SetterPath;
 			{
-				const FMVVMBlueprintViewModelContext* SourceViewModelContext = BlueprintView->FindViewModel(Binding.ViewModelPath.GetViewModelId());
-				check(SourceViewModelContext);
-				FName DestinationName = SourceViewModelContext->GetViewModelName();
-				const int32 DestinationVariableContextIndex = SourceContexts.IndexOfByPredicate([DestinationName](const FCompilerSourceContext& Other) { return Other.PropertyName == DestinationName; });
-				check(DestinationVariableContextIndex != INDEX_NONE);
-
-				Setter = CreatePropertyPath(SourceContexts[DestinationVariableContextIndex].PropertyName, Binding.ViewModelPath.GetBasePropertyPath());
+				SetterPath = CreateBindingDestinationPath(BlueprintView, Class, Binding.ViewModelPath);
 			}
 
-			TValueOrError<FCompilerBinding, FString> AddBindingResult = AddBinding(Class, Getters, Setter, Binding.Conversion.DestinationToSourceFunction);
+			TValueOrError<FCompilerBinding, FString> AddBindingResult = AddBinding(Class, GetterFields, SetterPath, Binding.Conversion.DestinationToSourceFunction);
 			if (AddBindingResult.HasError())
 			{
 				AddErrorForBinding(Binding, BlueprintView, FString::Printf(TEXT("The binding could not be created. %s"), *AddBindingResult.GetError()));
@@ -902,65 +883,106 @@
 }
 
 
-FMVVMViewBlueprintCompiler::FBindingSourceContext FMVVMViewBlueprintCompiler::CreateBindingSourceContext(const UMVVMBlueprintView* BlueprintView, const FMVVMBlueprintPropertyPath& PropertyPath) const
-{
-	FBindingSourceContext Result;
-
-	// Todo the path may contains an other INotifyFieldValueChanged
-	TArray<FName> Paths = PropertyPath.GetPaths();
-	if (Paths.Num() > 0)
-	{
-		Result.FieldId = Paths[0];
-	}
-
-	{
-		const FMVVMBlueprintViewModelContext* SourceViewModelContext = BlueprintView->FindViewModel(PropertyPath.GetViewModelId());
-		check(SourceViewModelContext);
-		const FName SourceName = SourceViewModelContext->GetViewModelName();
-		Result.CompilerSourceContextIndex = SourceContexts.IndexOfByPredicate([SourceName](const FCompilerSourceContext& Other) { return Other.PropertyName == SourceName; });
-		check(Result.CompilerSourceContextIndex != INDEX_NONE);
-	}
-
-	Result.SourceClass = SourceContexts[Result.CompilerSourceContextIndex].Class;
-	Result.PropertyPath = CreatePropertyPath(SourceContexts[Result.CompilerSourceContextIndex].PropertyName, PropertyPath.GetBasePropertyPath());
-	Result.bIsRootWidget = false;
-
-	return Result;
-}
-
-
 FMVVMViewBlueprintCompiler::FBindingSourceContext FMVVMViewBlueprintCompiler::CreateBindingSourceContext(const UMVVMBlueprintView* BlueprintView, const UWidgetBlueprintGeneratedClass* Class, const FMVVMBlueprintPropertyPath& PropertyPath) const
 {
 	FBindingSourceContext Result;
 
-	const FName SourceName = PropertyPath.GetWidgetName();
-	Result.bIsRootWidget = SourceName == Class->ClassGeneratedBy->GetFName();
+	if (PropertyPath.IsEmpty())
+	{
+		ensureAlwaysMsgf(false, TEXT("Not supported yet."));
+		return Result;
+	}
 
-	// Todo the path may contains an other INotifyFieldValueChanged
+	// Todo the path may contains another INotifyFieldValueChanged
 	TArray<FName> Paths = PropertyPath.GetPaths();
 	if (Paths.Num() > 0)
 	{
 		Result.FieldId = Paths[0];
 	}
 
-	if (Result.bIsRootWidget)
+	if (PropertyPath.IsFromViewModel())
 	{
-		Result.CompilerSourceContextIndex = INDEX_NONE;
-		Result.SourceClass = Class->ClassGeneratedBy->GetClass();
-		Result.PropertyPath = CreatePropertyPath(FName(), PropertyPath.GetBasePropertyPath());
-	}
-	else
-	{
+		Result.bIsRootWidget = false;
+
+		const FMVVMBlueprintViewModelContext* SourceViewModelContext = BlueprintView->FindViewModel(PropertyPath.GetViewModelId());
+		check(SourceViewModelContext);
+		const FName SourceName = SourceViewModelContext->GetViewModelName();
 		Result.CompilerSourceContextIndex = SourceContexts.IndexOfByPredicate([SourceName](const FCompilerSourceContext& Other) { return Other.PropertyName == SourceName; });
 		check(Result.CompilerSourceContextIndex != INDEX_NONE);
+
 		Result.SourceClass = SourceContexts[Result.CompilerSourceContextIndex].Class;
-		Result.PropertyPath = CreatePropertyPath(SourceContexts[Result.CompilerSourceContextIndex].PropertyName, PropertyPath.GetBasePropertyPath());
+		Result.PropertyPath = CreatePropertyPath(Class, SourceContexts[Result.CompilerSourceContextIndex].PropertyName, PropertyPath.GetFields());
+	}
+	else if (PropertyPath.IsFromWidget())
+	{
+		const FName SourceName = PropertyPath.GetWidgetName();
+		Result.bIsRootWidget = SourceName == Class->ClassGeneratedBy->GetFName();
+		if (Result.bIsRootWidget)
+		{
+			Result.CompilerSourceContextIndex = INDEX_NONE;
+			Result.SourceClass = Class->ClassGeneratedBy->GetClass();
+			Result.PropertyPath = CreatePropertyPath(Class, FName(), PropertyPath.GetFields());
+		}
+		else
+		{
+			Result.CompilerSourceContextIndex = SourceContexts.IndexOfByPredicate([SourceName](const FCompilerSourceContext& Other) { return Other.PropertyName == SourceName; });
+			check(Result.CompilerSourceContextIndex != INDEX_NONE);
+			Result.SourceClass = SourceContexts[Result.CompilerSourceContextIndex].Class;
+			Result.PropertyPath = CreatePropertyPath(Class, SourceContexts[Result.CompilerSourceContextIndex].PropertyName, PropertyPath.GetFields());
+		}
+	}
+	else
+	{
+		ensureAlwaysMsgf(false, TEXT("Not supported yet."));
 	}
 
 	return Result;
 }
 
 
+TArray<FMVVMConstFieldVariant> FMVVMViewBlueprintCompiler::CreateBindingDestinationPath(const UMVVMBlueprintView* BlueprintView, const UWidgetBlueprintGeneratedClass* Class, const FMVVMBlueprintPropertyPath& PropertyPath) const
+{
+	if (PropertyPath.IsEmpty())
+	{
+		ensureAlwaysMsgf(false, TEXT("Not supported yet."));
+		return TArray<FMVVMConstFieldVariant>();
+	}
+
+	if (PropertyPath.IsFromViewModel())
+	{
+		const FMVVMBlueprintViewModelContext* SourceViewModelContext = BlueprintView->FindViewModel(PropertyPath.GetViewModelId());
+		check(SourceViewModelContext);
+		FName DestinationName = SourceViewModelContext->GetViewModelName();
+		const int32 DestinationVariableContextIndex = SourceContexts.IndexOfByPredicate([DestinationName](const FCompilerSourceContext& Other) { return Other.PropertyName == DestinationName; });
+		check(DestinationVariableContextIndex != INDEX_NONE);
+
+		return CreatePropertyPath(Class, SourceContexts[DestinationVariableContextIndex].PropertyName, PropertyPath.GetFields());
+	}
+	else if (PropertyPath.IsFromWidget())
+	{
+		FName DestinationName = PropertyPath.GetWidgetName();
+		checkf(!DestinationName.IsNone(), TEXT("The destination should have been checked and set bAreSourceContextsValid."));
+		const bool bSourceIsUserWidget = DestinationName == Class->ClassGeneratedBy->GetFName();
+		if (bSourceIsUserWidget)
+		{
+			return CreatePropertyPath(Class, FName(), PropertyPath.GetFields());
+		}
+		else
+		{
+			const int32 DestinationVariableContextIndex = SourceContexts.IndexOfByPredicate([DestinationName](const FCompilerSourceContext& Other) { return Other.PropertyName == DestinationName; });
+			check(DestinationVariableContextIndex != INDEX_NONE);
+
+			return CreatePropertyPath(Class, SourceContexts[DestinationVariableContextIndex].PropertyName, PropertyPath.GetFields());
+		}
+	}
+	else
+	{
+		ensureAlwaysMsgf(false, TEXT("Not supported yet."));
+		return CreatePropertyPath(Class, FName(), PropertyPath.GetFields());
+	}
+}
+
+
 FString FMVVMViewBlueprintCompiler::CreatePropertyPath(FName PropertyName, FString PropertyPath) const
 {
 	TStringBuilder<512> Builder;
@@ -980,4 +1002,18 @@
 	return Builder.ToString();
 }
 
+
+TArray<FMVVMConstFieldVariant> FMVVMViewBlueprintCompiler::CreatePropertyPath(const UClass* Class, FName PropertyName, TArray<FMVVMConstFieldVariant> Properties) const
+{
+	if (PropertyName.IsNone())
+	{
+		return Properties;
+	}
+
+	check(Class);
+	FMVVMConstFieldVariant NewProperty = BindingHelper::FindFieldByName(Class, FMVVMBindingName(PropertyName));
+	Properties.Insert(NewProperty, 0);
+	return Properties;
+}
+
 } //namespace
diff --git a/Engine/Plugins/Runtime/ModelViewViewModel/Source/ModelViewViewModelBlueprint/Private/MVVMViewBlueprintCompiler.h b/Engine/Plugins/Runtime/ModelViewViewModel/Source/ModelViewViewModelBlueprint/Private/MVVMViewBlueprintCompiler.h
--- a/Engine/Plugins/Runtime/ModelViewViewModel/Source/ModelViewViewModelBlueprint/Private/MVVMViewBlueprintCompiler.h
+++ b/Engine/Plugins/Runtime/ModelViewViewModel/Source/ModelViewViewModelBlueprint/Private/MVVMViewBlueprintCompiler.h
@@ -70,10 +70,11 @@
 
 	void AddErrorForBinding(FMVVMBlueprintViewBinding& Binding, const UMVVMBlueprintView* BlueprintView, const FString& Message) const;
 
-	FBindingSourceContext CreateBindingSourceContext(const UMVVMBlueprintView* BlueprintView, const FMVVMBlueprintPropertyPath& PropertyPath) const;
 	FBindingSourceContext CreateBindingSourceContext(const UMVVMBlueprintView* BlueprintView, const UWidgetBlueprintGeneratedClass* Class, const FMVVMBlueprintPropertyPath& PropertyPath) const;
+	TArray<FMVVMConstFieldVariant> CreateBindingDestinationPath(const UMVVMBlueprintView* BlueprintView, const UWidgetBlueprintGeneratedClass* Class, const FMVVMBlueprintPropertyPath& PropertyPath) const;
 
 	FString CreatePropertyPath(FName PropertyName, FString PropertyPath) const;
+	TArray<FMVVMConstFieldVariant> CreatePropertyPath(const UClass* Class, FName PropertyName, TArray<FMVVMConstFieldVariant> Properties) const;
 
 private:
 	struct FCompilerSourceContext
@@ -136,7 +137,7 @@
 		// or Widget.Field.SubProperty
 		// or Field.SubProperty (if the widget is the userwidget)
 		// if its the destination then the Field is not necessary
-		FString PropertyPath;
+		TArray<UE::MVVM::FMVVMConstFieldVariant> PropertyPath;
 		int32 CompilerSourceContextIndex = INDEX_NONE;
 		bool bIsRootWidget = false;
 	};
diff --git a/Engine/Plugins/Runtime/ModelViewViewModel/Source/ModelViewViewModelBlueprint/Public/Bindings/MVVMCompiledBindingLibraryCompiler.h b/Engine/Plugins/Runtime/ModelViewViewModel/Source/ModelViewViewModelBlueprint/Public/Bindings/MVVMCompiledBindingLibraryCompiler.h
--- a/Engine/Plugins/Runtime/ModelViewViewModel/Source/ModelViewViewModelBlueprint/Public/Bindings/MVVMCompiledBindingLibraryCompiler.h
+++ b/Engine/Plugins/Runtime/ModelViewViewModel/Source/ModelViewViewModelBlueprint/Public/Bindings/MVVMCompiledBindingLibraryCompiler.h
@@ -162,7 +162,7 @@
 	TValueOrError<FFieldPathHandle, FString> AddFieldPath(TSubclassOf<UObject> SourceClass, FStringView FieldPath, bool bRead);
 	
 	/** */
-	TValueOrError<FFieldPathHandle, FString> AddFieldPath(TArrayView<UE::MVVM::FMVVMFieldVariant> FieldPath, bool bRead);
+	TValueOrError<FFieldPathHandle, FString> AddFieldPath(TArrayView<UE::MVVM::FMVVMConstFieldVariant> FieldPath, bool bRead);
 
 	/** */
 	TValueOrError<FFieldPathHandle, FString> AddObjectFieldPath(TSubclassOf<UObject> SourceClass, FStringView FieldPath, UClass* ExpectedType, bool bRead);
@@ -191,6 +191,9 @@
 	TValueOrError<FCompileResult, FString> Compile();
 
 private:
+	/** */
+	TValueOrError<FFieldPathHandle, FString> AddFieldPathImpl(TArrayView<UE::MVVM::FMVVMConstFieldVariant> FieldPath, bool bRead);
+
 	TPimplPtr<Private::FCompiledBindingLibraryCompilerImpl> Impl;
 };
 
diff --git a/Engine/Plugins/Runtime/ModelViewViewModel/Source/ModelViewViewModelBlueprint/Public/MVVMPropertyPath.h b/Engine/Plugins/Runtime/ModelViewViewModel/Source/ModelViewViewModelBlueprint/Public/MVVMPropertyPath.h
--- a/Engine/Plugins/Runtime/ModelViewViewModel/Source/ModelViewViewModelBlueprint/Public/MVVMPropertyPath.h
+++ b/Engine/Plugins/Runtime/ModelViewViewModel/Source/ModelViewViewModelBlueprint/Public/MVVMPropertyPath.h
@@ -5,6 +5,7 @@
 #include "Blueprint/WidgetBlueprintGeneratedClass.h"
 #include "Types/MVVMBindingName.h"
 #include "Types/MVVMFieldVariant.h"
+#include "UObject/Class.h"
 
 #include "MVVMPropertyPath.generated.h"
 
@@ -29,21 +30,40 @@
 	/** Get the binding name, resolves reference deprecation / redirectors / etc before returning */
 	FName GetFieldName() const
 	{
-		// Resolve any redirectors 
-		FName Path = FName();
+		// Resolve any redirectors
 		if (!BindingReference.GetMemberName().IsNone())
 		{
-			if (BindingKind == EBindingKind::Property && BindingReference.ResolveMember<FProperty>())
+			if (BindingKind == EBindingKind::Property)
 			{
-				Path = BindingReference.GetMemberName();
+				if (BindingReference.IsLocalScope())
+				{
+					if (UPackage* Package = BindingReference.GetMemberParentPackage())
+					{
+						UObjectBase* FoundObject = FindObjectWithOuter(Package, UScriptStruct::StaticClass(), *BindingReference.GetMemberScopeName());
+						if (UScriptStruct* Struct = Cast<UScriptStruct>(static_cast<UObject*>(FoundObject)))
+						{
+							if (const FProperty* FoundProperty = FindUFieldOrFProperty<FProperty>(Struct, BindingReference.GetMemberName(), EFieldIterationFlags::IncludeAll))
+							{
+								return FoundProperty->GetFName();
+							}
+						}
+					}
+				}
+				else if (BindingReference.ResolveMember<FProperty>())
+				{
+					return BindingReference.GetMemberName();
+				}
 			}
-			else if (BindingKind == EBindingKind::Function && BindingReference.ResolveMember<UFunction>())
+			else if (BindingKind == EBindingKind::Function)
 			{
-				Path = BindingReference.GetMemberName();
+				if (BindingReference.ResolveMember<UFunction>())
+				{
+					return BindingReference.GetMemberName();
+				}
 			}
 		}
 
-		return Path;
+		return FName();
 	}
 
 	/** */
@@ -53,7 +73,24 @@
 		{
 			if (BindingKind == EBindingKind::Property)
 			{
-				return UE::MVVM::FMVVMConstFieldVariant(BindingReference.ResolveMember<FProperty>());
+				if (BindingReference.IsLocalScope())
+				{
+					if (UPackage* Package = BindingReference.GetMemberParentPackage())
+					{
+						UObjectBase* FoundObject = FindObjectWithOuter(Package, UScriptStruct::StaticClass(), *BindingReference.GetMemberScopeName());
+						if (UScriptStruct* Struct = Cast<UScriptStruct>(static_cast<UObject*>(FoundObject)))
+						{
+							if (const FProperty* FoundProperty = FindUFieldOrFProperty<FProperty>(Struct, BindingReference.GetMemberName(), EFieldIterationFlags::IncludeAll))
+							{
+								return UE::MVVM::FMVVMConstFieldVariant(FoundProperty);
+							}
+						}
+					}
+				}
+				else
+				{
+					return UE::MVVM::FMVVMConstFieldVariant(BindingReference.ResolveMember<FProperty>());
+				}
 			}
 			else if (BindingKind == EBindingKind::Function)
 			{
@@ -63,12 +100,7 @@
 		return UE::MVVM::FMVVMConstFieldVariant();
 	}
 
-	FMemberReference GetBindingReference() const
-	{
-		// Note: Prefer copy since FName & FGuid, FString should be empty based on our usage
-		return BindingReference;
-	}
-
+	/** */
 	void SetBindingReference(UE::MVVM::FMVVMConstFieldVariant InField)
 	{
 		if (InField.IsEmpty())
@@ -77,7 +109,7 @@
 			return;
 		}
 
-		BindingReference = InField.CreateMemberReference();
+		BindingReference = CreateMemberReference(InField);
 
 		if (InField.IsProperty())
 		{
@@ -93,11 +125,13 @@
 		}
 	}
 
+	/** */
 	void Reset()
 	{
 		BindingReference = FMemberReference();
 	}
 
+	/** */
 	void SetDeprecatedBindingReference(const FMemberReference& InBindingReference, EBindingKind InBindingKind)
 	{
 		BindingReference = InBindingReference;
@@ -113,9 +147,38 @@
 	{
 		return !operator==(Other);
 	}
+
+private:
+	/**
+	 * Create a serializable member reference from this field.
+	 */
+	static FMemberReference CreateMemberReference(UE::MVVM::FMVVMConstFieldVariant InField)
+	{
+		FMemberReference BindingReference = FMemberReference();
+		if (InField.IsProperty())
+		{
+			const FProperty* Property = InField.GetProperty();
+			if (Property->GetOwnerClass())
+			{
+				BindingReference.SetFromField<FProperty>(Property, false);
+			}
+			else
+			{
+				BindingReference.SetGlobalField(Property->GetFName(), Property->GetOutermost());
+				BindingReference.SetLocalMember(Property->GetFName(), Property->GetOwnerStruct(), FGuid());
+			}
+		}
+		else if (InField.IsFunction())
+		{
+			// Functions should never be self context references
+			const UFunction* Function = InField.GetFunction();
+			bool bSelfContext = false;
+			BindingReference.SetFromField<UFunction>(Function, bSelfContext);
+		}
+		return BindingReference;
+	}
 };
 
-
 /**
  * Base path to properties for MVVM view models and widgets.
  * 
@@ -160,6 +223,20 @@
 		return Result;
 	}
 
+	/** Get the binding name, resolves reference deprecation / redirectors / etc before returning */
+	TArray<UE::MVVM::FMVVMConstFieldVariant> GetFields() const
+	{
+		TArray<UE::MVVM::FMVVMConstFieldVariant> Result;
+		Result.Reserve(Paths.Num());
+
+		for (const FMVVMBlueprintFieldPath& Path : Paths)
+		{
+			Result.Add(Path.GetField());
+		}
+
+		return Result;
+	}
+
 	bool BasePropertyPathContains(UE::MVVM::FMVVMConstFieldVariant Field) const
 	{
 		return Paths.ContainsByPredicate([Field](const FMVVMBlueprintFieldPath& FieldPath) { return FieldPath.GetField() == Field;  });
@@ -168,7 +245,10 @@
 	void SetBasePropertyPath(UE::MVVM::FMVVMConstFieldVariant InField)
 	{
 		Paths.Reset();
-		Paths.AddDefaulted_GetRef().SetBindingReference(InField);
+		if (!InField.IsEmpty() && !InField.GetName().IsNone())
+		{
+			Paths.AddDefaulted_GetRef().SetBindingReference(InField);
+		}
 	}
 
 	void ResetBasePropertyPath()
