MVVM: Fix for range loop bind to temporary constructed type.
#rnx none

diff --git a/Engine/Plugins/Runtime/ModelViewViewModel/Source/ModelViewViewModel/Private/Bindings/MVVMFieldPathHelper.cpp b/Engine/Plugins/Runtime/ModelViewViewModel/Source/ModelViewViewModel/Private/Bindings/MVVMFieldPathHelper.cpp
--- a/Engine/Plugins/Runtime/ModelViewViewModel/Source/ModelViewViewModel/Private/Bindings/MVVMFieldPathHelper.cpp
+++ b/Engine/Plugins/Runtime/ModelViewViewModel/Source/ModelViewViewModel/Private/Bindings/MVVMFieldPathHelper.cpp
@@ -165,7 +165,7 @@ TValueOrError<TArray<FMVVMConstFieldVariant>, FString> GenerateFieldPathList(TSu
  *				PropertyD and PropertyE cannot have Getter/BlueprintGetter if they are FStructProperty
  *		* We can only have one Setter/BlueprintSetter in the path
  */
-TValueOrError<TArray<FMVVMConstFieldVariant>, FString> GenerateFieldPathList(const TArrayView<FMVVMConstFieldVariant> InFieldPath, bool bForSourceBinding)
+TValueOrError<TArray<FMVVMConstFieldVariant>, FString> GenerateFieldPathList(TArrayView<const FMVVMConstFieldVariant> InFieldPath, bool bForSourceBinding)
 {
 	if (InFieldPath.Num() == 0)
 	{
@@ -232,7 +232,7 @@ TValueOrError<TArray<FMVVMConstFieldVariant>, FString> GenerateFieldPathList(con
 }
 
 
-FString ToString(const TArrayView<FMVVMFieldVariant> Fields)
+FString ToString(TArrayView<const FMVVMFieldVariant> Fields)
 {
 	TStringBuilder<512> Builder;
 	for (int32 Index = 0; Index < Fields.Num(); ++Index)
@@ -247,7 +247,7 @@ FString ToString(const TArrayView<FMVVMFieldVariant> Fields)
 }
 
 
-FString ToString(const TArrayView<FMVVMConstFieldVariant> Fields)
+FString ToString(TArrayView<const FMVVMConstFieldVariant> Fields)
 {
 	TStringBuilder<512> Builder;
 	for (int32 Index = 0; Index < Fields.Num(); ++Index)
diff --git a/Engine/Plugins/Runtime/ModelViewViewModel/Source/ModelViewViewModel/Public/Bindings/MVVMFieldPathHelper.h b/Engine/Plugins/Runtime/ModelViewViewModel/Source/ModelViewViewModel/Public/Bindings/MVVMFieldPathHelper.h
--- a/Engine/Plugins/Runtime/ModelViewViewModel/Source/ModelViewViewModel/Public/Bindings/MVVMFieldPathHelper.h
+++ b/Engine/Plugins/Runtime/ModelViewViewModel/Source/ModelViewViewModel/Public/Bindings/MVVMFieldPathHelper.h
@@ -16,11 +16,11 @@ namespace UE::MVVM::FieldPathHelper
 	UE_NODISCARD MODELVIEWVIEWMODEL_API TValueOrError<TArray<FMVVMConstFieldVariant>, FString> GenerateFieldPathList(TSubclassOf<UObject> From, FStringView FieldPath, bool bForSourceBinding);
 
 	/** */
-	UE_NODISCARD MODELVIEWVIEWMODEL_API TValueOrError<TArray<FMVVMConstFieldVariant>, FString> GenerateFieldPathList(const TArrayView<FMVVMConstFieldVariant>, bool bForSourceBinding);
+	UE_NODISCARD MODELVIEWVIEWMODEL_API TValueOrError<TArray<FMVVMConstFieldVariant>, FString> GenerateFieldPathList(TArrayView<const FMVVMConstFieldVariant>, bool bForSourceBinding);
 
 	/** */
-	UE_NODISCARD MODELVIEWVIEWMODEL_API FString ToString(const TArrayView<FMVVMFieldVariant> Fields);
-	UE_NODISCARD MODELVIEWVIEWMODEL_API FString ToString(const TArrayView<FMVVMConstFieldVariant> Fields);
+	UE_NODISCARD MODELVIEWVIEWMODEL_API FString ToString(TArrayView<const FMVVMFieldVariant> Fields);
+	UE_NODISCARD MODELVIEWVIEWMODEL_API FString ToString(TArrayView<const FMVVMConstFieldVariant> Fields);
 
 	/** */
 	UE_NODISCARD MODELVIEWVIEWMODEL_API TValueOrError<UObject*, void> EvaluateObjectProperty(const FFieldContext& InSource);
diff --git a/Engine/Plugins/Runtime/ModelViewViewModel/Source/ModelViewViewModelBlueprint/Private/Bindings/MVVMCompiledBindingLibraryCompiler.cpp b/Engine/Plugins/Runtime/ModelViewViewModel/Source/ModelViewViewModelBlueprint/Private/Bindings/MVVMCompiledBindingLibraryCompiler.cpp
--- a/Engine/Plugins/Runtime/ModelViewViewModel/Source/ModelViewViewModelBlueprint/Private/Bindings/MVVMCompiledBindingLibraryCompiler.cpp
+++ b/Engine/Plugins/Runtime/ModelViewViewModel/Source/ModelViewViewModelBlueprint/Private/Bindings/MVVMCompiledBindingLibraryCompiler.cpp
@@ -196,7 +196,7 @@ TValueOrError<FCompiledBindingLibraryCompiler::FFieldPathHandle, FString> FCompi
 }
 
 
-TValueOrError<FCompiledBindingLibraryCompiler::FFieldPathHandle, FString> FCompiledBindingLibraryCompiler::AddFieldPath(TArrayView<FMVVMConstFieldVariant> InFieldPath, bool bInRead)
+TValueOrError<FCompiledBindingLibraryCompiler::FFieldPathHandle, FString> FCompiledBindingLibraryCompiler::AddFieldPath(TArrayView<const FMVVMConstFieldVariant> InFieldPath, bool bInRead)
 {
 	Impl->bCompiled = false;
 
@@ -210,7 +210,7 @@ TValueOrError<FCompiledBindingLibraryCompiler::FFieldPathHandle, FString> FCompi
 }
 
 
-TValueOrError<FCompiledBindingLibraryCompiler::FFieldPathHandle, FString> FCompiledBindingLibraryCompiler::AddFieldPathImpl(TArrayView<FMVVMConstFieldVariant> InFieldPath, bool bInRead)
+TValueOrError<FCompiledBindingLibraryCompiler::FFieldPathHandle, FString> FCompiledBindingLibraryCompiler::AddFieldPathImpl(TArrayView<const FMVVMConstFieldVariant> InFieldPath, bool bInRead)
 {
 	Impl->bCompiled = false;
 
@@ -434,7 +434,7 @@ TValueOrError<FCompiledBindingLibraryCompiler::FBindingHandle, FString> FCompile
 }
 
 
-TValueOrError<FCompiledBindingLibraryCompiler::FBindingHandle, FString> FCompiledBindingLibraryCompiler::AddBinding(TArrayView<FFieldPathHandle> InSourceHandles, FFieldPathHandle InDestinationHandle, FFieldPathHandle InConversionFunctionHandle)
+TValueOrError<FCompiledBindingLibraryCompiler::FBindingHandle, FString> FCompiledBindingLibraryCompiler::AddBinding(TArrayView<const FFieldPathHandle> InSourceHandles, FFieldPathHandle InDestinationHandle, FFieldPathHandle InConversionFunctionHandle)
 {
 	Impl->bCompiled = false;
 
diff --git a/Engine/Plugins/Runtime/ModelViewViewModel/Source/ModelViewViewModelBlueprint/Private/MVVMViewBlueprintCompiler.cpp b/Engine/Plugins/Runtime/ModelViewViewModel/Source/ModelViewViewModelBlueprint/Private/MVVMViewBlueprintCompiler.cpp
--- a/Engine/Plugins/Runtime/ModelViewViewModel/Source/ModelViewViewModelBlueprint/Private/MVVMViewBlueprintCompiler.cpp
+++ b/Engine/Plugins/Runtime/ModelViewViewModel/Source/ModelViewViewModelBlueprint/Private/MVVMViewBlueprintCompiler.cpp
@@ -643,7 +643,7 @@ bool FMVVMViewBlueprintCompiler::PreCompileBindings(UWidgetBlueprintGeneratedCla
 		{
 			FCompilerBinding Result;
 
-			for (const TArrayView<UE::MVVM::FMVVMConstFieldVariant>& GetterField : GetterFields)
+			for (const TArray<UE::MVVM::FMVVMConstFieldVariant>& GetterField : GetterFields)
 			{
 				// Generate a path to read the value at runtime
 				TValueOrError<FCompiledBindingLibraryCompiler::FFieldPathHandle, FString> FieldPathResult = Self->BindingLibraryCompiler.AddFieldPath(GetterField, true);
diff --git a/Engine/Plugins/Runtime/ModelViewViewModel/Source/ModelViewViewModelBlueprint/Public/Bindings/MVVMCompiledBindingLibraryCompiler.h b/Engine/Plugins/Runtime/ModelViewViewModel/Source/ModelViewViewModelBlueprint/Public/Bindings/MVVMCompiledBindingLibraryCompiler.h
--- a/Engine/Plugins/Runtime/ModelViewViewModel/Source/ModelViewViewModelBlueprint/Public/Bindings/MVVMCompiledBindingLibraryCompiler.h
+++ b/Engine/Plugins/Runtime/ModelViewViewModel/Source/ModelViewViewModelBlueprint/Public/Bindings/MVVMCompiledBindingLibraryCompiler.h
@@ -162,7 +162,7 @@ public:
 	TValueOrError<FFieldPathHandle, FString> AddFieldPath(TSubclassOf<UObject> SourceClass, FStringView FieldPath, bool bRead);
 	
 	/** */
-	TValueOrError<FFieldPathHandle, FString> AddFieldPath(TArrayView<UE::MVVM::FMVVMConstFieldVariant> FieldPath, bool bRead);
+	TValueOrError<FFieldPathHandle, FString> AddFieldPath(TArrayView<const UE::MVVM::FMVVMConstFieldVariant> FieldPath, bool bRead);
 
 	/** */
 	TValueOrError<FFieldPathHandle, FString> AddObjectFieldPath(TSubclassOf<UObject> SourceClass, FStringView FieldPath, UClass* ExpectedType, bool bRead);
@@ -177,7 +177,7 @@ public:
 	TValueOrError<FBindingHandle, FString> AddBinding(FFieldPathHandle Source, FFieldPathHandle Destination, FFieldPathHandle ConversionFunction);
 
 	/** */
-	TValueOrError<FBindingHandle, FString> AddBinding(TArrayView<FFieldPathHandle> Sources, FFieldPathHandle Destination, FFieldPathHandle ConversionFunction);
+	TValueOrError<FBindingHandle, FString> AddBinding(TArrayView<const FFieldPathHandle> Sources, FFieldPathHandle Destination, FFieldPathHandle ConversionFunction);
 
 	struct FCompileResult
 	{
@@ -192,7 +192,7 @@ public:
 
 private:
 	/** */
-	TValueOrError<FFieldPathHandle, FString> AddFieldPathImpl(TArrayView<UE::MVVM::FMVVMConstFieldVariant> FieldPath, bool bRead);
+	TValueOrError<FFieldPathHandle, FString> AddFieldPathImpl(TArrayView<const UE::MVVM::FMVVMConstFieldVariant> FieldPath, bool bRead);
 
 	TPimplPtr<Private::FCompiledBindingLibraryCompilerImpl> Impl;
 };
