MVVM - Child viewmodel properties are no longer duplicated in the Manage ViewModels available bindings list. They were already returned by UMVVMSubsystem::GetAvailableBindings.

#rnx

diff --git a/Engine/Plugins/Runtime/ModelViewViewModel/Source/ModelViewViewModelEditor/Private/Widgets/SMVVMViewModelBindingListWidget.cpp b/Engine/Plugins/Runtime/ModelViewViewModel/Source/ModelViewViewModelEditor/Private/Widgets/SMVVMViewModelBindingListWidget.cpp
--- a/Engine/Plugins/Runtime/ModelViewViewModel/Source/ModelViewViewModelEditor/Private/Widgets/SMVVMViewModelBindingListWidget.cpp
+++ b/Engine/Plugins/Runtime/ModelViewViewModel/Source/ModelViewViewModelEditor/Private/Widgets/SMVVMViewModelBindingListWidget.cpp
@@ -145,12 +145,6 @@
 	static TArray<FMVVMAvailableBinding> GetBindingsList(TSubclassOf<UMVVMViewModelBase> Class)
 	{
 		TArray<FMVVMAvailableBinding> ViewModelAvailableBindingsList = GEngine->GetEngineSubsystem<UMVVMSubsystem>()->GetViewModelAvailableBindings(Class);
-		const TArray<UE::MVVM::FMVVMConstFieldVariant> ChildViewModels = GEditor->GetEditorSubsystem<UMVVMEditorSubsystem>()->GetChildViewModels(Class);
-		for (const UE::MVVM::FMVVMConstFieldVariant& Field : ChildViewModels)
-		{
-			ViewModelAvailableBindingsList.Add(FMVVMAvailableBinding(FMVVMBindingName(Field.GetName()), true, false));
-		}
-
 		return ViewModelAvailableBindingsList;
 	}
 } //namespace UE::MVVM::Private
