MVVM details view extension no longer crashes when given a row without an outer class.

This is fallout from the change to UStruct::IsChildOf() when called with this == nullptr.

#rnx

diff --git a/Engine/Plugins/Runtime/ModelViewViewModel/Source/ModelViewViewModelEditor/Private/Customizations/MVVMBindPropertiesDetailView.cpp b/Engine/Plugins/Runtime/ModelViewViewModel/Source/ModelViewViewModelEditor/Private/Customizations/MVVMBindPropertiesDetailView.cpp
--- a/Engine/Plugins/Runtime/ModelViewViewModel/Source/ModelViewViewModelEditor/Private/Customizations/MVVMBindPropertiesDetailView.cpp
+++ b/Engine/Plugins/Runtime/ModelViewViewModel/Source/ModelViewViewModelEditor/Private/Customizations/MVVMBindPropertiesDetailView.cpp
@@ -23,7 +23,9 @@ bool FMVVMBindPropertiesDetailView::IsSupported(const FOnGenerateGlobalRowExtens
 		return false;
 	}
 
-	if (!InArgs.PropertyHandle->GetOuterBaseClass()->IsChildOf(UWidget::StaticClass()))
+	// only display for UWidgets
+	const UClass* OuterBaseClass = InArgs.PropertyHandle->GetOuterBaseClass();
+	if (OuterBaseClass == nullptr || !OuterBaseClass->IsChildOf(UWidget::StaticClass()))
 	{
 		return false;
 	}
