MVVM: Fix pointer utilized before it was verified against nullptr.
#rnx

diff --git a/Engine/Plugins/Runtime/ModelViewViewModel/Source/ModelViewViewModel/Private/Bindings/MVVMFieldPathHelper.cpp b/Engine/Plugins/Runtime/ModelViewViewModel/Source/ModelViewViewModel/Private/Bindings/MVVMFieldPathHelper.cpp
--- a/Engine/Plugins/Runtime/ModelViewViewModel/Source/ModelViewViewModel/Private/Bindings/MVVMFieldPathHelper.cpp
+++ b/Engine/Plugins/Runtime/ModelViewViewModel/Source/ModelViewViewModel/Private/Bindings/MVVMFieldPathHelper.cpp
@@ -186,11 +186,16 @@
 			return MakeError(FString::Printf(TEXT("The field '%d' does not exist."), Index));
 		}
 
-		bool bIsChild = Field.GetOwner()->IsChildOf(CurrentContainer);
-		bool bIsDownCast = CurrentContainer->IsChildOf(Field.GetOwner());
-		if (CurrentContainer == nullptr || !(bIsChild || bIsDownCast))
+		if (CurrentContainer == nullptr)
 		{
-			return MakeError(FString::Printf(TEXT("The field '%s' does not exist in the struct '%s'."), *Field.GetName().ToString(), CurrentContainer ? *CurrentContainer->GetName() : TEXT("<none>")));
+			return MakeError(FString::Printf(TEXT("The field '%s' does not exist."), *Field.GetName().ToString()));
+		}
+
+		const bool bIsChild = Field.GetOwner()->IsChildOf(CurrentContainer);
+		const bool bIsDownCast = CurrentContainer->IsChildOf(Field.GetOwner());
+		if (!(bIsChild || bIsDownCast))
+		{
+			return MakeError(FString::Printf(TEXT("The field '%s' does not exist in the struct '%s'."), *Field.GetName().ToString(), *CurrentContainer->GetName()));
 		}
 		
 		if (Field.IsProperty())
