MVVM: Fix bad test for down casting when executing binding.
#rnx

diff --git a/Engine/Plugins/Runtime/ModelViewViewModel/Source/ModelViewViewModel/Private/Bindings/MVVMCompiledBindingLibrary.cpp b/Engine/Plugins/Runtime/ModelViewViewModel/Source/ModelViewViewModel/Private/Bindings/MVVMCompiledBindingLibrary.cpp
--- a/Engine/Plugins/Runtime/ModelViewViewModel/Source/ModelViewViewModel/Private/Bindings/MVVMCompiledBindingLibrary.cpp
+++ b/Engine/Plugins/Runtime/ModelViewViewModel/Source/ModelViewViewModel/Private/Bindings/MVVMCompiledBindingLibrary.cpp
@@ -205,11 +205,11 @@
 	check(!Destination.GetObjectVariant().IsNull());
 	check(!Source.GetFieldVariant().IsEmpty());
 	check(!Destination.GetFieldVariant().IsEmpty());
-	if (Source.GetObjectVariant().GetOwner()->IsChildOf(Source.GetFieldVariant().GetOwner()))
+	if (!Source.GetObjectVariant().GetOwner()->IsChildOf(Source.GetFieldVariant().GetOwner()))
 	{
 		return MakeError(EExecutionFailingReason::InvalidCast);
 	}
-	if (Destination.GetObjectVariant().GetOwner()->IsChildOf(Destination.GetFieldVariant().GetOwner()))
+	if (!Destination.GetObjectVariant().GetOwner()->IsChildOf(Destination.GetFieldVariant().GetOwner()))
 	{
 		return MakeError(EExecutionFailingReason::InvalidCast);
 	}
