Change the Blueprint's OnExtensionAdded and OnExtensionRemoved to use raw pointer.

diff --git a/Engine/Plugins/Runtime/ModelViewViewModel/Source/ModelViewViewModelEditor/Private/Widgets/SMVVMViewBindingPanel.cpp b/Engine/Plugins/Runtime/ModelViewViewModel/Source/ModelViewViewModelEditor/Private/Widgets/SMVVMViewBindingPanel.cpp
--- a/Engine/Plugins/Runtime/ModelViewViewModel/Source/ModelViewViewModelEditor/Private/Widgets/SMVVMViewBindingPanel.cpp
+++ b/Engine/Plugins/Runtime/ModelViewViewModel/Source/ModelViewViewModelEditor/Private/Widgets/SMVVMViewBindingPanel.cpp
@@ -271,7 +271,7 @@ FReply SMVVMViewBindingPanel::CreateDrawerDockButtonClicked()
 	return FReply::Handled();
 }
 
-void SMVVMViewBindingPanel::HandleExtensionAdded(TObjectPtr<UBlueprintExtension> NewExtension)
+void SMVVMViewBindingPanel::HandleExtensionAdded(UBlueprintExtension* NewExtension)
 {
 	if (UMVVMWidgetBlueprintExtension_View* MVVMExtensionPtr = Cast<UMVVMWidgetBlueprintExtension_View>(NewExtension))
 	{
diff --git a/Engine/Plugins/Runtime/ModelViewViewModel/Source/ModelViewViewModelEditor/Private/Widgets/SMVVMViewBindingPanel.h b/Engine/Plugins/Runtime/ModelViewViewModel/Source/ModelViewViewModelEditor/Private/Widgets/SMVVMViewBindingPanel.h
--- a/Engine/Plugins/Runtime/ModelViewViewModel/Source/ModelViewViewModelEditor/Private/Widgets/SMVVMViewBindingPanel.h
+++ b/Engine/Plugins/Runtime/ModelViewViewModel/Source/ModelViewViewModelEditor/Private/Widgets/SMVVMViewBindingPanel.h
@@ -51,7 +51,7 @@ private:
 	TSharedRef<SWidget> CreateDrawerDockButton();
 	FReply CreateDrawerDockButtonClicked();
 
-	void HandleExtensionAdded(TObjectPtr<UBlueprintExtension> NewExtension);
+	void HandleExtensionAdded(UBlueprintExtension* NewExtension);
 
 private:
 	TWeakPtr<FWidgetBlueprintEditor> WeakBlueprintEditor;
diff --git a/Engine/Source/Runtime/Engine/Classes/Engine/Blueprint.h b/Engine/Source/Runtime/Engine/Classes/Engine/Blueprint.h
--- a/Engine/Source/Runtime/Engine/Classes/Engine/Blueprint.h
+++ b/Engine/Source/Runtime/Engine/Classes/Engine/Blueprint.h
@@ -621,8 +621,8 @@ PRAGMA_ENABLE_DEPRECATION_WARNINGS
 	TObjectPtr<class UInheritableComponentHandler> InheritableComponentHandler;
 
 #if WITH_EDITORONLY_DATA
-	DECLARE_MULTICAST_DELEGATE_OneParam(FOnExtensionAdded, TObjectPtr<UBlueprintExtension>);
-	DECLARE_MULTICAST_DELEGATE_OneParam(FOnExtensionRemoved, TObjectPtr<UBlueprintExtension>);
+	DECLARE_MULTICAST_DELEGATE_OneParam(FOnExtensionAdded, UBlueprintExtension*);
+	DECLARE_MULTICAST_DELEGATE_OneParam(FOnExtensionRemoved, UBlueprintExtension*);
 
 	/** Array of new variables to be added to generated class */
 	UPROPERTY()
