MVVM: Fix type in error message.

diff --git a/Engine/Plugins/Runtime/ModelViewViewModel/Source/ModelViewViewModel/Private/View/MVVMView.cpp b/Engine/Plugins/Runtime/ModelViewViewModel/Source/ModelViewViewModel/Private/View/MVVMView.cpp
--- a/Engine/Plugins/Runtime/ModelViewViewModel/Source/ModelViewViewModel/Private/View/MVVMView.cpp
+++ b/Engine/Plugins/Runtime/ModelViewViewModel/Source/ModelViewViewModel/Private/View/MVVMView.cpp
@@ -417,7 +417,7 @@ TScriptInterface<INotifyFieldValueChanged> UMVVMView::FindSource(const FMVVMView
 	if (SourceObjectProperty == nullptr)
 	{
 		UE::MVVM::FMessageLog Log(GetUserWidget());
-		Log.Error(FText::Format(LOCTEXT("FindSourceInvalidPropertyName", "We could not evalute the source for binding '{0}'. The property name is invalid."), FText::FromString(Binding.ToString())));
+		Log.Error(FText::Format(LOCTEXT("FindSourceInvalidPropertyName", "We could not evaluate the source for binding '{0}'. The property name is invalid."), FText::FromString(Binding.ToString())));
 		return TScriptInterface<INotifyFieldValueChanged>();
 	}
 
@@ -427,7 +427,7 @@ TScriptInterface<INotifyFieldValueChanged> UMVVMView::FindSource(const FMVVMView
 		if (!bAllowNull)
 		{
 			UE::MVVM::FMessageLog Log(GetUserWidget());
-			Log.Warning(FText::Format(LOCTEXT("FindSourceInvalidUObject", "We could not evalute the source for binding '{0}'. The path point to an invalid object."), FText::FromString(Binding.ToString())));
+			Log.Warning(FText::Format(LOCTEXT("FindSourceInvalidUObject", "We could not evaluate the source for binding '{0}'. The path point to an invalid object."), FText::FromString(Binding.ToString())));
 		}
 		return TScriptInterface<INotifyFieldValueChanged>();
 	}
@@ -435,7 +435,7 @@ TScriptInterface<INotifyFieldValueChanged> UMVVMView::FindSource(const FMVVMView
 	if (!Source->Implements<UNotifyFieldValueChanged>())
 	{
 		UE::MVVM::FMessageLog Log(GetUserWidget());
-		Log.Error(FText::Format(LOCTEXT("FindSourceNotTheCorrectInterface", "We could not evalute the source for binding '{0}'. The object {1} doesn't implements INotifyFieldValueChanged."), FText::FromString(Binding.ToString()), FText::FromName(Source->GetFName())));
+		Log.Error(FText::Format(LOCTEXT("FindSourceNotTheCorrectInterface", "We could not evaluate the source for binding '{0}'. The object {1} doesn't implements INotifyFieldValueChanged."), FText::FromString(Binding.ToString()), FText::FromName(Source->GetFName())));
 		return TScriptInterface<INotifyFieldValueChanged>();
 	}
 
