MVVM: Add function to get the available bindings from class that implements the UNotifyFieldValueChanged interface.

diff --git a/Engine/Plugins/Runtime/ModelViewViewModel/Source/ModelViewViewModel/Private/MVVMSubsystem.cpp b/Engine/Plugins/Runtime/ModelViewViewModel/Source/ModelViewViewModel/Private/MVVMSubsystem.cpp
--- a/Engine/Plugins/Runtime/ModelViewViewModel/Source/ModelViewViewModel/Private/MVVMSubsystem.cpp
+++ b/Engine/Plugins/Runtime/ModelViewViewModel/Source/ModelViewViewModel/Private/MVVMSubsystem.cpp
@@ -191,15 +191,16 @@ namespace UE::MVVM::Private
 		return Result;
 	}
 
-	template<typename TClassType>
-	TArray<FMVVMAvailableBinding> GetAvailableBindings(TSubclassOf<TClassType> InSubClass)
+	TArray<FMVVMAvailableBinding> GetAvailableBindings(TSubclassOf<UObject> InSubClass)
 	{
-		if (InSubClass.Get() == nullptr)
+		if (InSubClass.Get() && InSubClass.Get()->ImplementsInterface(UNotifyFieldValueChanged::StaticClass()))
 		{
-			return TArray<FMVVMAvailableBinding>();
+			TScriptInterface<INotifyFieldValueChanged> DefaultObject = InSubClass.GetDefaultObject();
+			const UE::FieldNotification::IClassDescriptor& ClassDescriptor = DefaultObject->GetFieldNotificationDescriptor();
+			return GetAvailableBindings(InSubClass.Get(), ClassDescriptor);
 		}
-		const UE::FieldNotification::IClassDescriptor& ClassDescriptor = InSubClass.GetDefaultObject()->GetFieldNotificationDescriptor();
-		return GetAvailableBindings(InSubClass.Get(), ClassDescriptor);
+
+		return TArray<FMVVMAvailableBinding>();
 	}
 } //namespace
 
@@ -216,6 +217,12 @@ TArray<FMVVMAvailableBinding> UMVVMSubsystem::GetWidgetAvailableBindings(TSubcla
 }
 
 
+TArray<FMVVMAvailableBinding> UMVVMSubsystem::GetAvailableBindings(TSubclassOf<UObject> Class) const
+{
+	return UE::MVVM::Private::GetAvailableBindings(Class);
+}
+
+
 TValueOrError<bool, FString> UMVVMSubsystem::IsBindingValid(FBindingArgs Args) const
 {
 	if (UE::MVVM::IsForwardBinding(Args.Mode))
diff --git a/Engine/Plugins/Runtime/ModelViewViewModel/Source/ModelViewViewModel/Public/MVVMSubsystem.h b/Engine/Plugins/Runtime/ModelViewViewModel/Source/ModelViewViewModel/Public/MVVMSubsystem.h
--- a/Engine/Plugins/Runtime/ModelViewViewModel/Source/ModelViewViewModel/Public/MVVMSubsystem.h
+++ b/Engine/Plugins/Runtime/ModelViewViewModel/Source/ModelViewViewModel/Public/MVVMSubsystem.h
@@ -49,14 +49,21 @@ public:
 	UFUNCTION(BlueprintCallable, Category = "MVVM")
 	bool DoesWidgetTreeContainedWidget(const UWidgetTree* WidgetTree, const UWidget* ViewWidget) const;
 
-	/** Returns the list of all the binding that are available for the ViewModel. */
+	/** Returns the list of all the bindings that are available for the ViewModel. */
 	UFUNCTION(BlueprintCallable, Category = "MVVM")
 	TArray<FMVVMAvailableBinding> GetViewModelAvailableBindings(TSubclassOf<UMVVMViewModelBase> ViewModelClass) const;
 
-	/** Returns the list of all the binding that are available for the Widget. */
+	/** Returns the list of all the bindings that are available for the Widget. */
 	UFUNCTION(BlueprintCallable, Category = "MVVM")
 	TArray<FMVVMAvailableBinding> GetWidgetAvailableBindings(TSubclassOf<UWidget> WidgetClass) const;
 
+	/**
+	 * Returns the list of all the bindings that are available for the Class.
+	 * The class must implement the UNotifyFieldValueChanged interface.
+	 */
+	UFUNCTION(BlueprintCallable, Category = "MVVM")
+	TArray<FMVVMAvailableBinding> GetAvailableBindings(TSubclassOf<UObject> Class) const;
+
 public:
 	UFUNCTION(BlueprintCallable, Category = "MVVM")
 	UMVVMViewModelCollectionObject* GetGlobalViewModelCollection() const
