MVVM: First draft of the viewmodel panel.

diff --git a/Engine/Plugins/Runtime/ModelViewViewModel/Source/ModelViewViewModelBlueprint/Private/MVVMBlueprintView.cpp b/Engine/Plugins/Runtime/ModelViewViewModel/Source/ModelViewViewModelBlueprint/Private/MVVMBlueprintView.cpp
--- a/Engine/Plugins/Runtime/ModelViewViewModel/Source/ModelViewViewModelBlueprint/Private/MVVMBlueprintView.cpp
+++ b/Engine/Plugins/Runtime/ModelViewViewModel/Source/ModelViewViewModelBlueprint/Private/MVVMBlueprintView.cpp
@@ -21,6 +21,14 @@ const FMVVMBlueprintViewModelContext* UMVVMBlueprintView::FindViewModel(FGuid Vi
 	return const_cast<UMVVMBlueprintView*>(this)->FindViewModel(ViewModelId);
 }
 
+const FMVVMBlueprintViewModelContext* UMVVMBlueprintView::FindViewModel(FName ViewModel) const
+{
+	return AvailableViewModels.FindByPredicate([ViewModel](const FMVVMBlueprintViewModelContext& Other)
+		{
+			return Other.GetViewModelName() == ViewModel;
+		});
+}
+
 void UMVVMBlueprintView::AddViewModel(const FMVVMBlueprintViewModelContext& NewContext)
 {
 	AvailableViewModels.Add(NewContext);
@@ -29,7 +37,7 @@ void UMVVMBlueprintView::AddViewModel(const FMVVMBlueprintViewModelContext& NewC
 }
 
 
-void UMVVMBlueprintView::RemoveViewModel(FGuid ViewModelId)
+bool UMVVMBlueprintView::RemoveViewModel(FGuid ViewModelId)
 {
 	int32 Count = AvailableViewModels.RemoveAll([ViewModelId](const FMVVMBlueprintViewModelContext& VM)
 		{
@@ -41,9 +49,10 @@ void UMVVMBlueprintView::RemoveViewModel(FGuid ViewModelId)
 		FBlueprintEditorUtils::MarkBlueprintAsStructurallyModified(GetOuterUMVVMWidgetBlueprintExtension_View()->GetWidgetBlueprint());
 		OnViewModelsUpdated.Broadcast();
 	}
+	return Count > 0;
 }
 
-void UMVVMBlueprintView::RemoveViewModels(const TArrayView<FGuid> ViewModelIds)
+int32 UMVVMBlueprintView::RemoveViewModels(const TArrayView<FGuid> ViewModelIds)
 {
 	int32 Count = 0;
 	for (const FGuid& ViewModelId : ViewModelIds)
@@ -59,6 +68,26 @@ void UMVVMBlueprintView::RemoveViewModels(const TArrayView<FGuid> ViewModelIds)
 		FBlueprintEditorUtils::MarkBlueprintAsStructurallyModified(GetOuterUMVVMWidgetBlueprintExtension_View()->GetWidgetBlueprint());
 		OnViewModelsUpdated.Broadcast();
 	}
+	return Count;
+}
+
+bool UMVVMBlueprintView::RenameViewModel(FName OldViewModelName, FName NewViewModelName)
+{
+	FMVVMBlueprintViewModelContext* ViewModelContext = AvailableViewModels.FindByPredicate([OldViewModelName](const FMVVMBlueprintViewModelContext& Other)
+			{
+				return Other.GetViewModelName() == OldViewModelName;
+			});
+	if (ViewModelContext)
+	{
+		ViewModelContext->OverrideDisplayName = FText::FromName(NewViewModelName);
+
+		FBlueprintEditorUtils::ReplaceVariableReferences(GetOuterUMVVMWidgetBlueprintExtension_View()->GetWidgetBlueprint(), OldViewModelName, NewViewModelName);
+		FBlueprintEditorUtils::ValidateBlueprintChildVariables(GetOuterUMVVMWidgetBlueprintExtension_View()->GetWidgetBlueprint(), NewViewModelName);
+		FBlueprintEditorUtils::MarkBlueprintAsStructurallyModified(GetOuterUMVVMWidgetBlueprintExtension_View()->GetWidgetBlueprint());
+
+		OnViewModelsUpdated.Broadcast();
+	}
+	return ViewModelContext != nullptr;
 }
 
 void UMVVMBlueprintView::SetViewModels(const TArray<FMVVMBlueprintViewModelContext>& ViewModelContexts)
diff --git a/Engine/Plugins/Runtime/ModelViewViewModel/Source/ModelViewViewModelBlueprint/Public/MVVMBlueprintView.h b/Engine/Plugins/Runtime/ModelViewViewModel/Source/ModelViewViewModelBlueprint/Public/MVVMBlueprintView.h
--- a/Engine/Plugins/Runtime/ModelViewViewModel/Source/ModelViewViewModelBlueprint/Public/MVVMBlueprintView.h
+++ b/Engine/Plugins/Runtime/ModelViewViewModel/Source/ModelViewViewModelBlueprint/Public/MVVMBlueprintView.h
@@ -23,10 +23,12 @@ class MODELVIEWVIEWMODELBLUEPRINT_API UMVVMBlueprintView : public UObject
 public:
 	FMVVMBlueprintViewModelContext* FindViewModel(FGuid ViewModelId);
 	const FMVVMBlueprintViewModelContext* FindViewModel(FGuid ViewModelId) const;
+	const FMVVMBlueprintViewModelContext* FindViewModel(FName ViewModelName) const;
 
 	void AddViewModel(const FMVVMBlueprintViewModelContext& NewContext);
-	void RemoveViewModel(FGuid ViewModelId);
-	void RemoveViewModels(const TArrayView<FGuid> ViewModelIds);
+	bool RemoveViewModel(FGuid ViewModelId);
+	int32 RemoveViewModels(const TArrayView<FGuid> ViewModelIds);
+	bool RenameViewModel(FName OldViewModelName, FName NewViewModelName);
 	void SetViewModels(const TArray<FMVVMBlueprintViewModelContext>& ViewModelContexts);
 
 	const TArrayView<const FMVVMBlueprintViewModelContext> GetViewModels() const
diff --git a/Engine/Plugins/Runtime/ModelViewViewModel/Source/ModelViewViewModelEditor/Private/MVVMEditorSubsystem.cpp b/Engine/Plugins/Runtime/ModelViewViewModel/Source/ModelViewViewModelEditor/Private/MVVMEditorSubsystem.cpp
--- a/Engine/Plugins/Runtime/ModelViewViewModel/Source/ModelViewViewModelEditor/Private/MVVMEditorSubsystem.cpp
+++ b/Engine/Plugins/Runtime/ModelViewViewModel/Source/ModelViewViewModelEditor/Private/MVVMEditorSubsystem.cpp
@@ -8,8 +8,11 @@
 #include "Engine/Engine.h"
 #include "MVVMSubsystem.h"
 #include "MVVMWidgetBlueprintExtension_View.h"
+#include "ScopedTransaction.h"
 #include "WidgetBlueprint.h"
 
+#define LOCTEXT_NAMESPACE "UMVVMEditorSubsystem"
+
 UMVVMBlueprintView* UMVVMEditorSubsystem::RequestView(UWidgetBlueprint* WidgetBlueprint) const
 {
 	return UMVVMWidgetBlueprintExtension_View::RequestExtension<UMVVMWidgetBlueprintExtension_View>(WidgetBlueprint)->GetBlueprintView();
@@ -24,20 +27,43 @@ UMVVMBlueprintView* UMVVMEditorSubsystem::GetView(UWidgetBlueprint* WidgetBluepr
 	return nullptr;
 }
 
-TArray<UE::MVVM::FMVVMConstFieldVariant> UMVVMEditorSubsystem::GetChildViewModels(TSubclassOf<UMVVMViewModelBase> Class)
+TArray<UE::MVVM::FMVVMConstFieldVariant> UMVVMEditorSubsystem::GetChildViewModels(UClass* Class)
 {
-	TArray<UE::MVVM::FMVVMConstFieldVariant> Result;
 
+	auto IsValidObjectProperty = [](const FProperty* Property)
+	{
+		if (const FObjectPropertyBase* ObjectProperty = CastField<FObjectPropertyBase>(Property))
+		{
+			return ObjectProperty->PropertyClass->ImplementsInterface(UNotifyFieldValueChanged::StaticClass());
+		}
+		return false;
+	};
+
+
+	TArray<UE::MVVM::FMVVMConstFieldVariant> Result;
 	for (TFieldIterator<const FProperty> PropertyIt(Class, EFieldIteratorFlags::IncludeSuper); PropertyIt; ++PropertyIt)
 	{
 		const FProperty* Property = *PropertyIt;
-		if (const FObjectPropertyBase* ObjectProperty = CastField<FObjectPropertyBase>(Property))
+		if (IsValidObjectProperty(Property) && UE::MVVM::BindingHelper::IsValidForSourceBinding(Property))
 		{
-			if (ObjectProperty->PropertyClass->IsChildOf(UMVVMViewModelBase::StaticClass()))
+			Result.Add(UE::MVVM::FMVVMConstFieldVariant(Property));
+		}
+	}
+
+	for (TFieldIterator<const UFunction> FunctionItt(Class, EFieldIteratorFlags::IncludeSuper); FunctionItt; ++FunctionItt)
+	{
+		const UFunction* Function = *FunctionItt;
+		check(Function);
+		if (!Function->HasAnyFunctionFlags(FUNC_Private | FUNC_Protected))
+		{
+			if (UE::MVVM::BindingHelper::IsValidForSourceBinding(Function))
 			{
-				if (UE::MVVM::BindingHelper::IsValidForSourceBinding(Property))
+				if (const FProperty* ReturnValue = UE::MVVM::BindingHelper::GetReturnProperty(Function))
 				{
-					Result.Add(UE::MVVM::FMVVMConstFieldVariant(Property));
+					if (IsValidObjectProperty(ReturnValue))
+					{
+						Result.Add(UE::MVVM::FMVVMConstFieldVariant(Function));
+					}
 				}
 			}
 		}
@@ -46,12 +72,52 @@ TArray<UE::MVVM::FMVVMConstFieldVariant> UMVVMEditorSubsystem::GetChildViewModel
 	return Result;
 }
 
-void UMVVMEditorSubsystem::RemoveViewModel(UWidgetBlueprint* WidgetBlueprint, const FMVVMBlueprintViewModelContext& ViewModel)
+void UMVVMEditorSubsystem::RemoveViewModel(UWidgetBlueprint* WidgetBlueprint, FName ViewModel)
 {
 	if (UMVVMBlueprintView* View = GetView(WidgetBlueprint))
 	{
-		View->RemoveViewModel(ViewModel.GetViewModelId());
+		if (const FMVVMBlueprintViewModelContext* ViewModelContext = View->FindViewModel(ViewModel))
+		{
+			View->RemoveViewModel(ViewModelContext->GetViewModelId());
+		}
+	}
+}
+
+bool UMVVMEditorSubsystem::VerifyViewModelRename(UWidgetBlueprint* WidgetBlueprint, FName ViewModel, FName NewViewModel, FText& OutError)
+{
+	FKismetNameValidator Validator(WidgetBlueprint);
+	EValidatorResult ValidatorResult = Validator.IsValid(NewViewModel);
+	if (ValidatorResult != EValidatorResult::Ok)
+	{
+		if (ViewModel == NewViewModel && (ValidatorResult == EValidatorResult::AlreadyInUse || ValidatorResult == EValidatorResult::ExistingName))
+		{
+			// Continue successfully
+		}
+		else
+		{
+			OutError = INameValidatorInterface::GetErrorText(NewViewModel.ToString(), ValidatorResult);
+			return false;
+		}
 	}
+	return true;
+}
+
+bool UMVVMEditorSubsystem::RenameViewModel(UWidgetBlueprint* WidgetBlueprint, FName ViewModel, FName NewViewModel, FText& OutError)
+{
+	if (!VerifyViewModelRename(WidgetBlueprint, ViewModel, NewViewModel, OutError))
+	{
+		return false;
+	}
+
+	UMVVMBlueprintView* View = GetView(WidgetBlueprint);
+	if (View == nullptr)
+	{
+		return false;
+	}
+
+	const FScopedTransaction Transaction(LOCTEXT("RenameViewModel", "Rename viewmodel"));
+	View->Modify();
+	return View->RenameViewModel(ViewModel, NewViewModel);
 }
 
 FMVVMBlueprintViewBinding& UMVVMEditorSubsystem::AddBinding(UWidgetBlueprint* WidgetBlueprint)
@@ -128,3 +194,5 @@ TArray<const UFunction*> UMVVMEditorSubsystem::GetAvailableConversionFunctions(c
 		});
 	return ConversionFunctions;
 }
+
+#undef LOCTEXT_NAMESPACE
diff --git a/Engine/Plugins/Runtime/ModelViewViewModel/Source/ModelViewViewModelEditor/Private/ModelViewViewModelEditorModule.cpp b/Engine/Plugins/Runtime/ModelViewViewModel/Source/ModelViewViewModelEditor/Private/ModelViewViewModelEditorModule.cpp
--- a/Engine/Plugins/Runtime/ModelViewViewModel/Source/ModelViewViewModelEditor/Private/ModelViewViewModelEditorModule.cpp
+++ b/Engine/Plugins/Runtime/ModelViewViewModel/Source/ModelViewViewModelEditor/Private/ModelViewViewModelEditorModule.cpp
@@ -25,6 +25,7 @@
 #include "Styling/SlateColor.h"
 #include "Styling/StyleColors.h"
 #include "Tabs/MVVMBindingSummoner.h"
+#include "Tabs/MVVMViewModelSummoner.h"
 #include "UMGEditorModule.h"
 #include "ViewModel/AssetTypeActions_ViewModelBlueprint.h"
 #include "WidgetBlueprint.h"
@@ -114,6 +115,7 @@ void FModelViewViewModelEditorModule::HandleRegisterBlueprintEditorTab(const FWi
 	if (ApplicationMode.GetModeName() == FWidgetBlueprintApplicationModes::DesignerMode)
 	{
 		TabFactories.RegisterFactory(MakeShared<FMVVMBindingSummoner>(ApplicationMode.GetBlueprintEditor()));
+		TabFactories.RegisterFactory(MakeShared<UE::MVVM::FViewModelSummoner>(ApplicationMode.GetBlueprintEditor()));
 
 		if (ApplicationMode.LayoutExtender)
 		{
diff --git a/Engine/Plugins/Runtime/ModelViewViewModel/Source/ModelViewViewModelEditor/Private/Tabs/MVVMViewModelSummoner.cpp b/Engine/Plugins/Runtime/ModelViewViewModel/Source/ModelViewViewModelEditor/Private/Tabs/MVVMViewModelSummoner.cpp
new file mode 100644
--- /dev/null
+++ b/Engine/Plugins/Runtime/ModelViewViewModel/Source/ModelViewViewModelEditor/Private/Tabs/MVVMViewModelSummoner.cpp
@@ -0,0 +1,39 @@
+// Copyright Epic Games, Inc. All Rights Reserved.
+
+#include "Tabs/MVVMViewModelSummoner.h"
+
+#include "Styling/MVVMEditorStyle.h"
+#include "WidgetBlueprintEditor.h"
+#include "Widgets/SMVVMViewModelPanel.h"
+
+#define LOCTEXT_NAMESPACE "ViewModelPanel"
+
+namespace UE::MVVM
+{
+
+const FName FViewModelSummoner::TabID(TEXT("ViewModelPanel"));
+
+FViewModelSummoner::FViewModelSummoner(TSharedPtr<class FWidgetBlueprintEditor> InBlueprintEditor)
+		: FWorkflowTabFactory(TabID, InBlueprintEditor)
+		, BlueprintEditor(InBlueprintEditor)
+{
+	TabLabel = LOCTEXT("ViewmodelTabLabel", "Viewmodels");
+	TabIcon = FSlateIcon(FMVVMEditorStyle::Get().GetStyleSetName(), "BlueprintView.TabIcon");
+
+	bIsSingleton = true;
+
+	ViewMenuDescription = LOCTEXT("Viewmodel_Desc", "Viewmodels");
+	ViewMenuTooltip = LOCTEXT("Viewmodel_ToolTip", "Show the viewmodels panel");
+}
+
+TSharedRef<SWidget> FViewModelSummoner::CreateTabBody(const FWorkflowTabSpawnInfo& Info) const
+{
+	TSharedPtr<FWidgetBlueprintEditor> BlueprintEditorPtr = StaticCastSharedPtr<FWidgetBlueprintEditor>(BlueprintEditor.Pin());
+
+	return SNew(SMVVMViewModelPanel, BlueprintEditorPtr)
+		.AddMetaData<FTagMetaData>(FTagMetaData(TEXT("ViewmodelPanel")));
+}
+
+} // namespace
+
+#undef LOCTEXT_NAMESPACE 
diff --git a/Engine/Plugins/Runtime/ModelViewViewModel/Source/ModelViewViewModelEditor/Private/Tabs/MVVMViewModelSummoner.h b/Engine/Plugins/Runtime/ModelViewViewModel/Source/ModelViewViewModelEditor/Private/Tabs/MVVMViewModelSummoner.h
new file mode 100644
--- /dev/null
+++ b/Engine/Plugins/Runtime/ModelViewViewModel/Source/ModelViewViewModelEditor/Private/Tabs/MVVMViewModelSummoner.h
@@ -0,0 +1,27 @@
+// Copyright Epic Games, Inc. All Rights Reserved.
+
+#pragma once
+
+#include "CoreMinimal.h"
+#include "Widgets/SWidget.h"
+#include "WorkflowOrientedApp/WorkflowTabFactory.h"
+#include "WidgetBlueprintEditor.h"
+
+namespace UE::MVVM
+{
+
+struct FViewModelSummoner : public FWorkflowTabFactory
+{
+public:
+	static const FName TabID;
+	
+public:
+	FViewModelSummoner(TSharedPtr<class FWidgetBlueprintEditor> InBlueprintEditor);
+	
+	virtual TSharedRef<SWidget> CreateTabBody(const FWorkflowTabSpawnInfo& Info) const override;
+	
+protected:
+	TWeakPtr<class FWidgetBlueprintEditor> BlueprintEditor;
+};
+
+} // namespace
diff --git a/Engine/Plugins/Runtime/ModelViewViewModel/Source/ModelViewViewModelEditor/Private/Widgets/SMVVMManageViewModelsWidget.cpp b/Engine/Plugins/Runtime/ModelViewViewModel/Source/ModelViewViewModelEditor/Private/Widgets/SMVVMManageViewModelsWidget.cpp
--- a/Engine/Plugins/Runtime/ModelViewViewModel/Source/ModelViewViewModelEditor/Private/Widgets/SMVVMManageViewModelsWidget.cpp
+++ b/Engine/Plugins/Runtime/ModelViewViewModel/Source/ModelViewViewModelEditor/Private/Widgets/SMVVMManageViewModelsWidget.cpp
@@ -285,7 +285,7 @@ void SMVVMManageViewModelsWidget::Construct(const FArguments& InArgs)
 						SNew(SBorder)
 						.BorderImage(FStyleDefaults::GetNoBrush())
 						[
-							SAssignNew(ViewModelBindingListWidget, SMVVMViewModelBindingListWidget)
+							SAssignNew(ViewModelBindingListWidget, UE::MVVM::SViewModelBindingListWidget)
 						]
 					]
 				]
@@ -769,7 +769,7 @@ void SMVVMManageViewModelsWidget::HandleTreeViewSelectionChanged(TSharedPtr<FMVV
 	{
 		InClass = GetClassFromNode(Item.ToSharedRef());
 	}
-	ViewModelBindingListWidget->SetViewModel(InClass);
+	ViewModelBindingListWidget->SetViewModel(InClass, FName());
 }
 
 void SMVVMManageViewModelsWidget::HandleTreeViewExpansionChanged(TSharedPtr<FMVVMViewModelTreeNode> Item, bool bExpanded)
diff --git a/Engine/Plugins/Runtime/ModelViewViewModel/Source/ModelViewViewModelEditor/Private/Widgets/SMVVMManageViewModelsWidget.h b/Engine/Plugins/Runtime/ModelViewViewModel/Source/ModelViewViewModelEditor/Private/Widgets/SMVVMManageViewModelsWidget.h
--- a/Engine/Plugins/Runtime/ModelViewViewModel/Source/ModelViewViewModelEditor/Private/Widgets/SMVVMManageViewModelsWidget.h
+++ b/Engine/Plugins/Runtime/ModelViewViewModel/Source/ModelViewViewModelEditor/Private/Widgets/SMVVMManageViewModelsWidget.h
@@ -13,7 +13,6 @@ class ITableRow;
 template<typename T>
 class SListView;
 class SMVVMViewModelContextListWidget;
-class SMVVMViewModelBindingListWidget;
 template<typename T>
 class STreeView;
 class STableViewBase;
@@ -24,6 +23,7 @@ class UWidgetBlueprint;
 namespace UE::MVVM
 {
 	struct FMVVMFieldVariant;
+	class SViewModelBindingListWidget;
 
 	namespace Private
 	{
@@ -108,7 +108,7 @@ private:
 
 	TWeakPtr<SWindow> WeakParentWindow;
 
-	TSharedPtr<SMVVMViewModelBindingListWidget> ViewModelBindingListWidget;
+	TSharedPtr<UE::MVVM::SViewModelBindingListWidget> ViewModelBindingListWidget;
 
 	bool bSearchStringIsEmpty = true;
 	UUserWidget* OuterWidget = nullptr;
diff --git a/Engine/Plugins/Runtime/ModelViewViewModel/Source/ModelViewViewModelEditor/Private/Widgets/SMVVMViewModelBindingListWidget.cpp b/Engine/Plugins/Runtime/ModelViewViewModel/Source/ModelViewViewModelEditor/Private/Widgets/SMVVMViewModelBindingListWidget.cpp
--- a/Engine/Plugins/Runtime/ModelViewViewModel/Source/ModelViewViewModelEditor/Private/Widgets/SMVVMViewModelBindingListWidget.cpp
+++ b/Engine/Plugins/Runtime/ModelViewViewModel/Source/ModelViewViewModelEditor/Private/Widgets/SMVVMViewModelBindingListWidget.cpp
@@ -1,255 +1,562 @@
-// Copyright Epic Games, Inc. All Rights Reserved.
-
-#include "SMVVMViewModelBindingListWidget.h"
-
-#include "Bindings/MVVMBindingHelper.h"
-#include "Editor.h"
-#include "Engine/Engine.h"
-#include "MVVMEditorSubsystem.h"
-#include "MVVMSubsystem.h"
-#include "MVVMViewModelBase.h"
-#include "Styling/SlateIconFinder.h"
-#include "Textures/SlateIcon.h"
-#include "Types/MVVMAvailableBinding.h"
-#include "Types/MVVMFieldVariant.h"
-#include "Widgets/Images/SImage.h"
-#include "Widgets/SMVVMFieldIcon.h"
-#include "Widgets/Views/STreeView.h"
-
-namespace UE::MVVM::Private
-{
-	class FMVVMAvailableBindingNode
-	{
-	public:
-		FMVVMAvailableBindingNode(const FMVVMAvailableBinding& InBinding, TSubclassOf<UMVVMViewModelBase> InClass)
-			:
-			Binding(InBinding),
-			Class(InClass),
-			bIsRoot(false)
-		{}
-
-		FMVVMAvailableBindingNode(TSubclassOf<UMVVMViewModelBase> InClass)
-			:
-			Class(InClass),
-			bIsRoot(true)
-		{}
-
-		const FMVVMAvailableBinding Binding;
-		const TSubclassOf<UMVVMViewModelBase> Class;
-		const bool bIsRoot = false;
-
-		TArray<TSharedPtr<FMVVMAvailableBindingNode>> ChildNodes;
-
-		TSubclassOf<UMVVMViewModelBase> GetViewModelClassProperty()
-		{
-			if (bIsRoot)
-			{
-				return nullptr;
-			}
-
-			UE::MVVM::FMVVMFieldVariant FieldVariant = UE::MVVM::BindingHelper::FindFieldByName(Class, Binding.GetBindingName());
-			if (FieldVariant.IsProperty())
-			{
-				FProperty* Property = FieldVariant.GetProperty();
-
-				if (FObjectProperty* ObjectProperty = CastField<FObjectProperty>(Property))
-				{
-					if (ObjectProperty->PropertyClass->IsChildOf(UMVVMViewModelBase::StaticClass()))
-					{
-						return ObjectProperty->PropertyClass;
-					}
-				}
-			}
-			
-			return nullptr;
-		}
-
-		FText GetDisplayName()
-		{
-			if (bIsRoot)
-			{
-				return Class->GetDisplayNameText();
-			}
-
-			return FText::FromString(Binding.GetBindingName().ToString());
-		}
-
-		FText GetToolTip()
-		{
-			if (bIsRoot)
-			{
-				return FText::GetEmpty();
-			}
-
-			UE::MVVM::FMVVMFieldVariant FieldVariant = UE::MVVM::BindingHelper::FindFieldByName(Class, Binding.GetBindingName());
-			if (FieldVariant.IsProperty())
-			{
-				FProperty* Property = FieldVariant.GetProperty();
-				check(Property);
-				return FText::FormatOrdered(INVTEXT("{0}\n{1}"), FText::FromString(Property->GetCPPType()), Property->GetMetaDataText("ToolTip"));
-			}
-
-			if (FieldVariant.IsFunction())
-			{
-				UFunction* Function = FieldVariant.GetFunction();
-				check(Function);
-				FProperty* Property = Function->GetReturnProperty();
-				if (!Property)
-				{
-					return FText::FormatOrdered(INVTEXT("void\n{0}"), Function->GetMetaDataText("ToolTip"));
-				}
-				return FText::FormatOrdered(INVTEXT("{0}\n{1}"), FText::FromString(Property->GetCPPType()), Function->GetMetaDataText("ToolTip"));
-			}
-
-			return FText::GetEmpty();
-		}
-
-		TSharedRef<SWidget> GetIcon()
-		{
-			if (bIsRoot)
-			{
-				return SNew(SImage).Image(FSlateIconFinder::FindIconBrushForClass(Class));
-			}
-
-			UE::MVVM::FMVVMFieldVariant FieldVariant = UE::MVVM::BindingHelper::FindFieldByName(Class, Binding.GetBindingName());
-			
-			return SNew(SMVVMFieldIcon)
-				.Field(FieldVariant);
-		}
-
-		static bool Sort(const TSharedPtr<FMVVMAvailableBindingNode>& Node, const TSharedPtr<FMVVMAvailableBindingNode>& OtherNode)
-		{
-			TSubclassOf<UMVVMViewModelBase> NodeClass = Node->GetViewModelClassProperty();
-			TSubclassOf<UMVVMViewModelBase> OtherNodeClass = OtherNode->GetViewModelClassProperty();
-			const FString& NodeStr = Node->Binding.GetBindingName().ToString();
-			const FString& OtherNodeStr = OtherNode->Binding.GetBindingName().ToString();
-			if (NodeClass && OtherNodeClass)
-			{
-				return NodeStr < OtherNodeStr;
-			}
-
-			if (NodeClass)
-			{
-				return true;
-			}
-
-			if (OtherNodeClass)
-			{
-				return false;
-			}
-
-			return NodeStr < OtherNodeStr;
-		}
-	};
-
-	static TArray<FMVVMAvailableBinding> GetBindingsList(TSubclassOf<UMVVMViewModelBase> Class)
-	{
-		TArray<FMVVMAvailableBinding> ViewModelAvailableBindingsList = GEngine->GetEngineSubsystem<UMVVMSubsystem>()->GetViewModelAvailableBindings(Class);
-		return ViewModelAvailableBindingsList;
-	}
-} //namespace UE::MVVM::Private
-
-void SMVVMViewModelBindingListWidget::Construct(const FArguments& InArgs)
-{
-	using UE::MVVM::Private::FMVVMAvailableBindingNode;
-	ChildSlot
-	[
-		SAssignNew(TreeWidget, STreeView<TSharedPtr<FMVVMAvailableBindingNode>>)
-		.TreeItemsSource(&TreeSource)
-		.OnGetChildren(this, &SMVVMViewModelBindingListWidget::HandleGetChildren)
-		.OnGenerateRow(this, &SMVVMViewModelBindingListWidget::HandleGenerateRow)
-		.ItemHeight(20.0f)
-	];
-}
-
-void SMVVMViewModelBindingListWidget::SetViewModel(TSubclassOf<UMVVMViewModelBase> ViewModelClass)
-{
-	using UE::MVVM::Private::FMVVMAvailableBindingNode;
-	TreeSource.Empty();
-	if (ViewModelClass)
-	{
-		TSharedRef<FMVVMAvailableBindingNode> RootNode = MakeShared<FMVVMAvailableBindingNode>(ViewModelClass);
-		TArray<FMVVMAvailableBinding> ViewModelAvailableBindingsList = UE::MVVM::Private::GetBindingsList(ViewModelClass);
-		for (const FMVVMAvailableBinding& Binding : ViewModelAvailableBindingsList)
-		{
-			RootNode->ChildNodes.Add(MakeShared<FMVVMAvailableBindingNode>(Binding, ViewModelClass));
-		}
-		RootNode->ChildNodes.Sort(FMVVMAvailableBindingNode::Sort);
-		TreeSource.Add(RootNode);
-		TreeWidget->SetItemExpansion(RootNode, true);
-	}
-	TreeWidget->RequestListRefresh();
-}
-
-TSharedRef<ITableRow> SMVVMViewModelBindingListWidget::HandleGenerateRow(TSharedPtr<FMVVMAvailableBindingNode> Item, const TSharedRef<STableViewBase>& OwnerTable)
-{
-	using UE::MVVM::Private::FMVVMAvailableBindingNode;
-	typedef STableRow<TSharedPtr<FMVVMAvailableBindingNode>> RowType;
-
-	TSharedRef<RowType> NewRow = SNew(RowType, OwnerTable);
-
-
-	NewRow->SetContent(
-		SNew(SHorizontalBox)
-			.ToolTipText(Item->GetToolTip())
-			+ SHorizontalBox::Slot()
-			.AutoWidth()
-			.HAlign(HAlign_Right)
-			.VAlign(VAlign_Center)
-			[
-				Item->GetIcon()
-			]
-			+ SHorizontalBox::Slot()
-			.Padding(4.0f)
-			[
-				SNew(STextBlock)
-				.Visibility(EVisibility::Visible)
-				.Text(Item->GetDisplayName())
-			]
-	);
-
-	return NewRow;
-}
-
-void SMVVMViewModelBindingListWidget::HandleGetChildren(TSharedPtr<FMVVMAvailableBindingNode> InParent, TArray<TSharedPtr<FMVVMAvailableBindingNode>>& OutChildren)
-{
-	using UE::MVVM::Private::FMVVMAvailableBindingNode;
-
-	if (!InParent->ChildNodes.IsEmpty())
-	{
-		OutChildren = InParent->ChildNodes;
-		return;
-	}
-
-	UE::MVVM::FMVVMFieldVariant FieldVariant = UE::MVVM::BindingHelper::FindFieldByName(InParent->Class, InParent->Binding.GetBindingName());
-
-	if (FieldVariant.IsProperty())
-	{
-		FProperty* Property = FieldVariant.GetProperty();
-
-		if (FObjectProperty* ObjectProperty = CastField<FObjectProperty>(Property))
-		{
-			if (ObjectProperty->PropertyClass->IsChildOf(UMVVMViewModelBase::StaticClass()))
-			{
-				if (InParent->ChildNodes.IsEmpty())
-				{
-					TSubclassOf<UMVVMViewModelBase> ChildClass = ObjectProperty->PropertyClass;
-					TArray<FMVVMAvailableBinding> ViewModelAvailableBindingsList = UE::MVVM::Private::GetBindingsList(ChildClass);
-					for (const FMVVMAvailableBinding& Binding : ViewModelAvailableBindingsList)
-					{
-						TSharedPtr<FMVVMAvailableBindingNode> ChildNode = MakeShared<FMVVMAvailableBindingNode>(Binding, ChildClass);
-						InParent->ChildNodes.Add(ChildNode);
-					}
-					InParent->ChildNodes.Sort(FMVVMAvailableBindingNode::Sort);
-				}
-				OutChildren.Append(InParent->ChildNodes);
-			}
-		}
-	}
-}
-
-TSharedRef<SWidget> SMVVMViewModelBindingListWidget::GetFieldIcon(const UE::MVVM::FMVVMFieldVariant& FieldVariant)
-{
-	return SNew(SMVVMFieldIcon)
-		.Field(FieldVariant);
-}
+// Copyright Epic Games, Inc. All Rights Reserved.
+
+#include "SMVVMViewModelBindingListWidget.h"
+
+#include "Bindings/MVVMBindingHelper.h"
+#include "Editor.h"
+#include "Engine/Engine.h"
+#include "MVVMEditorSubsystem.h"
+#include "MVVMSubsystem.h"
+#include "MVVMViewModelBase.h"
+#include "Textures/SlateIcon.h"
+#include "Types/MVVMAvailableBinding.h"
+#include "Types/MVVMFieldVariant.h"
+
+#include "Framework/Commands/GenericCommands.h"
+#include "Framework/Commands/UICommandList.h"
+#include "Framework/MultiBox/MultiBoxBuilder.h"
+#include "Styling/SlateIconFinder.h"
+
+#include "Widgets/Images/SImage.h"
+#include "Widgets/SMVVMFieldIcon.h"
+#include "Widgets/Text/SInlineEditableTextBlock.h"
+#include "Widgets/Views/STreeView.h"
+
+
+#define LOCTEXT_NAMESPACE "SViewModelBindingListWidget"
+
+namespace UE::MVVM
+{
+
+namespace Private
+{
+
+static TArray<FMVVMAvailableBinding> GetBindingsList(TSubclassOf<UObject> Class)
+{
+	TArray<FMVVMAvailableBinding> ViewModelAvailableBindingsList = GEngine->GetEngineSubsystem<UMVVMSubsystem>()->GetAvailableBindings(Class);
+	const TArray<FMVVMConstFieldVariant> ChildViewModels = GEditor->GetEditorSubsystem<UMVVMEditorSubsystem>()->GetChildViewModels(Class);
+	for (const FMVVMConstFieldVariant& Field : ChildViewModels)
+	{
+		if (!ViewModelAvailableBindingsList.ContainsByPredicate([&Field](const FMVVMAvailableBinding& Other){ return Other.GetBindingName().ToName() == Field.GetName(); }))
+		{
+			ViewModelAvailableBindingsList.Add(FMVVMAvailableBinding(FMVVMBindingName(Field.GetName()), true, false));
+		}
+	}
+
+	return ViewModelAvailableBindingsList;
+}
+
+
+/** 
+ * 
+ */
+class FAvailableBindingNode : public TSharedFromThis<FAvailableBindingNode>
+{
+public:
+	FAvailableBindingNode(SViewModelBindingListWidget* InOwner, const FMVVMAvailableBinding& InBinding, TSubclassOf<UObject> InViewModelClass)
+		: Owner(InOwner)
+		, bIsRoot(false)
+		, Binding(InBinding)
+		, ViewModelOwnerClass(InViewModelClass)
+		, ViewModelClassProperty(GetCacheViewModelClassProperty())
+	{
+	}
+
+
+	FAvailableBindingNode(SViewModelBindingListWidget* InOwner, const SViewModelBindingListWidget::FViewModel& InViewModel)
+		: Owner(InOwner)
+		, bIsRoot(true)
+		, ViewModel(InViewModel)
+	{
+		check(InViewModel.Class && InViewModel.Class->ImplementsInterface(UNotifyFieldValueChanged::StaticClass()));
+	}
+
+private:
+	SViewModelBindingListWidget* Owner = nullptr;
+
+	const bool bIsRoot = false;
+	const SViewModelBindingListWidget::FViewModel ViewModel;
+
+	const FMVVMAvailableBinding Binding;
+	const TSubclassOf<UObject> ViewModelOwnerClass;
+	const TSubclassOf<UObject> ViewModelClassProperty; // if the property is also a viewmodel
+
+	bool bChildGenerated = false;
+	TArray<TSharedPtr<FAvailableBindingNode>, TInlineAllocator<1>> ChildNodes;
+	TWeakPtr<SInlineEditableTextBlock> ViewModelEditTextBlock;
+
+public:
+	bool ShowContextMenu() const
+	{
+		return bIsRoot;
+	}
+
+	bool CanRename() const
+	{
+		return bIsRoot && ViewModelEditTextBlock.IsValid();
+	}
+
+
+	void EnterRename() const
+	{
+		if (bIsRoot)
+		{
+			if (TSharedPtr<SInlineEditableTextBlock> SafeEditBox = ViewModelEditTextBlock.Pin())
+			{
+				SafeEditBox->EnterEditingMode();
+			}
+		}
+	}
+
+
+	TArray<TSharedPtr<FAvailableBindingNode>> GetChildNodes()
+	{
+		if (!bChildGenerated && ViewModelClassProperty.Get() != nullptr)
+		{
+			// Do not build when filtering (only search in what it's already been built)
+			if (!Owner->FilterHandler->GetIsEnabled())
+			{
+				BuildChildNodesRecursive(ViewModelClassProperty, 2);
+			}
+		}
+		return TArray<TSharedPtr<FAvailableBindingNode>>(ChildNodes);
+	}
+
+
+	void BuildChildNodes()
+	{
+		BuildChildNodesRecursive(ViewModel.Class, 3);
+	}
+
+
+	FText GetDisplayName()
+	{
+		if (bIsRoot)
+		{
+			return !ViewModel.Name.IsNone() ? FText::FromName(ViewModel.Name) : ViewModel.Class->GetDisplayNameText();
+		}
+
+		return FText::FromName(Binding.GetBindingName().ToName());
+	}
+
+
+	void GetFilterStrings(TArray<FString>& OutStrings) const
+	{
+		auto AddClass = [&OutStrings](UClass* Class)
+		{
+			OutStrings.Add(Class->GetName());
+			OutStrings.Add(Class->GetDisplayNameText().ToString());
+		};
+
+		if (bIsRoot)
+		{
+			if (!ViewModel.Name.IsNone())
+			{
+				OutStrings.Add(ViewModel.Name.ToString());
+			}
+			AddClass(ViewModel.Class);
+		}
+		else
+		{
+			OutStrings.Add(Binding.GetBindingName().ToString());
+			if (ViewModelOwnerClass.Get())
+			{
+				AddClass(ViewModelOwnerClass.Get());
+			}
+			if (ViewModelClassProperty.Get())
+			{
+				AddClass(ViewModelClassProperty.Get());
+			}
+		}
+	}
+
+
+	FText GetToolTip()
+	{
+		if (bIsRoot)
+		{
+			return ViewModel.Class->GetDisplayNameText();
+		}
+
+		FMVVMFieldVariant FieldVariant = BindingHelper::FindFieldByName(ViewModelOwnerClass, Binding.GetBindingName());
+		if (FieldVariant.IsProperty())
+		{
+			FProperty* Property = FieldVariant.GetProperty();
+			check(Property);
+			return FText::FormatOrdered(INVTEXT("{0}\n{1}"), FText::FromString(Property->GetCPPType()), Property->GetMetaDataText("ToolTip"));
+		}
+
+		if (FieldVariant.IsFunction())
+		{
+			UFunction* Function = FieldVariant.GetFunction();
+			check(Function);
+			FProperty* Property = Function->GetReturnProperty();
+			if (!Property)
+			{
+				return FText::FormatOrdered(INVTEXT("void\n{0}"), Function->GetMetaDataText("ToolTip"));
+			}
+			return FText::FormatOrdered(INVTEXT("{0}\n{1}"), FText::FromString(Property->GetCPPType()), Function->GetMetaDataText("ToolTip"));
+		}
+
+		return FText::GetEmpty();
+	}
+
+
+	TSharedRef<SWidget> CreateIconWidget()
+	{
+		if (bIsRoot)
+		{
+			return SNew(SImage).Image(FSlateIconFinder::FindIconBrushForClass(ViewModel.Class));
+		}
+
+		FMVVMFieldVariant FieldVariant = BindingHelper::FindFieldByName(ViewModelOwnerClass, Binding.GetBindingName());
+			
+		return SNew(SMVVMFieldIcon)
+			.Field(FieldVariant);
+	}
+
+
+	TSharedRef<SWidget> CreateDetailWidget(TAttribute<FText> GetFilterText, const SViewModelBindingListWidget::FOnViewModelRenamed& OnRenamed)
+	{
+		ViewModelEditTextBlock.Reset();
+		if (OnRenamed.IsBound() && bIsRoot)
+		{
+			return SAssignNew(ViewModelEditTextBlock, SInlineEditableTextBlock)
+				.Text(GetDisplayName())
+				.HighlightText(GetFilterText)
+				.OnVerifyTextChanged(this, &FAvailableBindingNode::HandleVerifyNameTextChanged, OnRenamed)
+				.OnTextCommitted(this, &FAvailableBindingNode::HandleNameTextCommited, OnRenamed);
+		}
+
+		return SNew(STextBlock)
+			.Text(GetDisplayName())
+			.HighlightText(GetFilterText);
+	}
+	
+private:
+	bool HandleVerifyNameTextChanged(const FText& InText, FText& OutErrorMessage, SViewModelBindingListWidget::FOnViewModelRenamed OnRenamed)
+	{
+		return OnRenamed.Execute(ViewModel, InText, false, OutErrorMessage);
+	}
+	
+	
+	void HandleNameTextCommited(const FText& InText, ETextCommit::Type CommitInfo, SViewModelBindingListWidget::FOnViewModelRenamed OnRenamed)
+	{
+		if (CommitInfo == ETextCommit::OnEnter)
+		{
+			FText OutErrorMessage;
+			OnRenamed.Execute(ViewModel, InText, true, OutErrorMessage);
+		}
+	}
+
+
+	static bool Sort(const TSharedPtr<FAvailableBindingNode>& Node, const TSharedPtr<FAvailableBindingNode>& OtherNode)
+	{
+		TSubclassOf<UObject> NodeClass = Node->ViewModelClassProperty;
+		TSubclassOf<UObject> OtherNodeClass = OtherNode->ViewModelClassProperty;
+		const FName NodeStr = Node->Binding.GetBindingName().ToName();
+		const FName OtherNodeStr = OtherNode->Binding.GetBindingName().ToName();
+
+		if (NodeClass && OtherNodeClass)
+		{
+			return NodeStr.LexicalLess(OtherNodeStr);
+		}
+
+		if (NodeClass)
+		{
+			return true;
+		}
+
+		if (OtherNodeClass)
+		{
+			return false;
+		}
+
+		return NodeStr.LexicalLess(OtherNodeStr);
+	}
+
+private:
+	TSubclassOf<UObject> GetCacheViewModelClassProperty() const
+	{
+		if (bIsRoot)
+		{
+			return nullptr;
+		}
+
+		FMVVMFieldVariant FieldVariant = BindingHelper::FindFieldByName(ViewModelOwnerClass, Binding.GetBindingName());
+		if (FieldVariant.IsProperty())
+		{
+			FProperty* Property = FieldVariant.GetProperty();
+
+			if (FObjectProperty* ObjectProperty = CastField<FObjectProperty>(Property))
+			{
+				if (ObjectProperty->PropertyClass->ImplementsInterface(UNotifyFieldValueChanged::StaticClass()))
+				{
+					return ObjectProperty->PropertyClass;
+				}
+			}
+		}
+
+		return nullptr;
+	}
+
+
+	void BuildChildNodesRecursive(UClass* ViewModelClass, int32 RecursiveCount)
+	{
+		if (RecursiveCount <= 0)
+		{
+			return;
+		}
+		--RecursiveCount;
+
+		ChildNodes.Reset();
+		for (const FMVVMAvailableBinding& ChildBinding : Private::GetBindingsList(ViewModelClass))
+		{
+			TSharedRef<Private::FAvailableBindingNode> ChildNode = MakeShared<Private::FAvailableBindingNode>(Owner, ChildBinding, ViewModelClass);
+			ChildNodes.Add(ChildNode);
+			if (ChildNode->ViewModelClassProperty.Get() != nullptr)
+			{
+				ChildNode->BuildChildNodesRecursive(ChildNode->ViewModelClassProperty, RecursiveCount);
+			}
+		}
+		ChildNodes.Sort(Private::FAvailableBindingNode::Sort);
+		bChildGenerated = true;
+	}
+};
+
+} //namespace Private
+
+
+void SViewModelBindingListWidget::Construct(const FArguments& InArgs)
+{
+	OnRenamed = InArgs._OnRenamed;
+	OnDeleted = InArgs._OnDeleted;
+	OnDuplicated = InArgs._OnDuplicated;
+
+	SearchFilter = MakeShared<ViewModelTextFilter>(ViewModelTextFilter::FItemToStringArray::CreateSP(this, &SViewModelBindingListWidget::HandleGetFilterStrings));
+
+	FilterHandler = MakeShared<ViewModelTreeFilter>();
+	FilterHandler->SetFilter(SearchFilter.Get());
+	FilterHandler->SetRootItems(&TreeSource, &FilteredTreeSource);
+	FilterHandler->SetGetChildrenDelegate(ViewModelTreeFilter::FOnGetChildren::CreateRaw(this, &SViewModelBindingListWidget::HandleGetChildren));
+
+	SetViewModels(InArgs._ViewModels);
+	CreateCommandList();
+
+	SAssignNew(TreeWidget, STreeView<TSharedPtr<Private::FAvailableBindingNode>>)
+		.ItemHeight(1.0f)
+		.TreeItemsSource(&FilteredTreeSource)
+		.SelectionMode(ESelectionMode::Single)
+		.OnGetChildren(FilterHandler.ToSharedRef(), &ViewModelTreeFilter::OnGetFilteredChildren)
+		.OnGenerateRow(this, &SViewModelBindingListWidget::HandleGenerateRow)
+		//	.OnSelectionChanged(this, &SMVVMViewModelPanel::HandleSelectionChanged)
+		.OnContextMenuOpening(this, &SViewModelBindingListWidget::HandleContextMenuOpening);
+
+	FilterHandler->SetTreeView(TreeWidget.Get());
+
+	ChildSlot
+	[
+		TreeWidget.ToSharedRef()
+	];
+}
+
+
+void SViewModelBindingListWidget::SetViewModel(TSubclassOf<UObject> ViewModelClass, FName InViewModelName)
+{
+	FViewModel ViewModel;
+	ViewModel.Class = ViewModelClass;
+	ViewModel.Name = InViewModelName;
+
+	SetViewModels(MakeArrayView(&ViewModel, 1));
+}
+
+
+void SViewModelBindingListWidget::SetViewModels(TArrayView<const FViewModel> ViewModels)
+{
+	TreeSource.Reset();
+	for (const FViewModel& ViewModel : ViewModels)
+	{
+		if (ViewModel.Class && ViewModel.Class->ImplementsInterface(UNotifyFieldValueChanged::StaticClass()))
+		{
+			TSharedRef<Private::FAvailableBindingNode> RootNode = MakeShared<Private::FAvailableBindingNode>(this, ViewModel);
+			RootNode->BuildChildNodes();
+			TreeSource.Add(RootNode);
+
+			if (TreeWidget)
+			{
+				TreeWidget->SetItemExpansion(RootNode, true);
+			}
+		}
+	}
+
+	if (FilterHandler)
+	{
+		FilterHandler->RefreshAndFilterTree();
+	}
+}
+
+
+FText SViewModelBindingListWidget::SetRawFilterText(const FText& InFilterText)
+{
+	const bool bNewFilteringEnabled = !InFilterText.IsEmpty();
+	FilterHandler->SetIsEnabled(bNewFilteringEnabled);
+	SearchFilter->SetRawFilterText(InFilterText);
+	FilterHandler->RefreshAndFilterTree();
+	return SearchFilter->GetFilterErrorText();
+}
+
+
+FText SViewModelBindingListWidget::GetRawFilterText() const
+{
+	return SearchFilter->GetRawFilterText();
+}
+
+FReply SViewModelBindingListWidget::OnKeyDown(const FGeometry& MyGeometry, const FKeyEvent& InKeyEvent)
+{
+	if (CommandList->ProcessCommandBindings(InKeyEvent))
+	{
+		return FReply::Handled();
+	}
+
+	return FReply::Unhandled();
+}
+
+
+void SViewModelBindingListWidget::HandleGetFilterStrings(TSharedPtr<Private::FAvailableBindingNode> Item, TArray<FString>& OutStrings)
+{
+	Item->GetFilterStrings(OutStrings);
+}
+
+
+TSharedRef<ITableRow> SViewModelBindingListWidget::HandleGenerateRow(TSharedPtr<Private::FAvailableBindingNode> Item, const TSharedRef<STableViewBase>& OwnerTable)
+{
+	using Private::FAvailableBindingNode;
+	typedef STableRow<TSharedPtr<FAvailableBindingNode>> RowType;
+
+
+	return SNew(RowType, OwnerTable)
+		//.OnDragDetected(this, &SHierarchyViewItem::HandleDragDetected)
+		//.OnDragEnter(this, &SHierarchyViewItem::HandleDragEnter)
+		//.OnDragLeave(this, &SHierarchyViewItem::HandleDragLeave)
+		.Padding(0.0f)
+		.Content()
+		[
+			SNew(SHorizontalBox)
+			.ToolTipText(Item->GetToolTip())
+			+ SHorizontalBox::Slot()
+			.AutoWidth()
+			.HAlign(HAlign_Right)
+			.VAlign(VAlign_Center)
+			[
+				Item->CreateIconWidget()
+			]
+			+ SHorizontalBox::Slot()
+			.Padding(4.0f)
+			[
+				Item->CreateDetailWidget(MakeAttributeSP(this, &SViewModelBindingListWidget::GetRawFilterText), OnRenamed)
+			]
+		];
+}
+
+
+void SViewModelBindingListWidget::HandleGetChildren(TSharedPtr<Private::FAvailableBindingNode> InParent, TArray<TSharedPtr<Private::FAvailableBindingNode>>& OutChildren)
+{
+	OutChildren = InParent->GetChildNodes();
+}
+
+
+void SViewModelBindingListWidget::CreateCommandList()
+{
+	CommandList = MakeShared<FUICommandList>();
+
+	CommandList->MapAction(
+		FGenericCommands::Get().Duplicate,
+		FExecuteAction::CreateSP(this, &SViewModelBindingListWidget::HandleDuplicateViewModel),
+		FCanExecuteAction::CreateSP(this, &SViewModelBindingListWidget::HandleCanDuplicateViewModel)
+	);
+	CommandList->MapAction(
+		FGenericCommands::Get().Delete,
+		FExecuteAction::CreateSP(this, &SViewModelBindingListWidget::HandleDeleteViewModel),
+		FCanExecuteAction::CreateSP(this, &SViewModelBindingListWidget::HandleCanDeleteViewModel)
+	);
+	CommandList->MapAction(
+		FGenericCommands::Get().Rename,
+		FExecuteAction::CreateSP(this, &SViewModelBindingListWidget::HandleRenameViewModel),
+		FCanExecuteAction::CreateSP(this, &SViewModelBindingListWidget::HandleCanRenameViewModel)
+	);
+}
+
+
+TSharedPtr<SWidget> SViewModelBindingListWidget::HandleContextMenuOpening()
+{
+	TArray<TSharedPtr<Private::FAvailableBindingNode>> Items = TreeWidget->GetSelectedItems();
+	if (Items.Num() == 1 && Items[0] && Items[0]->ShowContextMenu())
+	{
+		if (OnRenamed.IsBound() || OnDeleted.IsBound() || OnDuplicated.IsBound())
+		{
+			FMenuBuilder MenuBuilder(true, CommandList);
+
+			MenuBuilder.BeginSection("Edit", LOCTEXT("Edit", "Edit"));
+			{
+				if (OnDuplicated.IsBound())
+				{
+					MenuBuilder.AddMenuEntry(FGenericCommands::Get().Duplicate);
+				}
+				if (OnDeleted.IsBound())
+				{
+					MenuBuilder.AddMenuEntry(FGenericCommands::Get().Delete);
+				}
+				if (OnRenamed.IsBound())
+				{
+					MenuBuilder.AddMenuEntry(FGenericCommands::Get().Rename);
+				}
+			}
+			MenuBuilder.EndSection();
+
+			return MenuBuilder.MakeWidget();
+		}
+	}
+	return TSharedPtr<SWidget>();
+}
+
+
+void SViewModelBindingListWidget::HandleDuplicateViewModel()
+{}
+
+
+bool SViewModelBindingListWidget::HandleCanDuplicateViewModel()
+{
+	return true; //todo unless it's an autogen viewmodel
+}
+
+
+void SViewModelBindingListWidget::HandleDeleteViewModel()
+{}
+
+
+bool SViewModelBindingListWidget::HandleCanDeleteViewModel()
+{
+	return true; //todo unless it's an autogen viewmodel
+}
+
+
+void SViewModelBindingListWidget::HandleRenameViewModel()
+{
+	TArray<TSharedPtr<Private::FAvailableBindingNode>> Items = TreeWidget->GetSelectedItems();
+	if (Items.Num() == 1 && Items[0])
+	{
+		Items[0]->EnterRename();
+	}
+}
+
+
+bool SViewModelBindingListWidget::HandleCanRenameViewModel()
+{
+	TArray<TSharedPtr<Private::FAvailableBindingNode>> Items = TreeWidget->GetSelectedItems();
+	return Items.Num() == 1 && Items[0] && Items[0]->CanRename();
+}
+
+
+} // namespace UE::MVVM
+
+#undef LOCTEXT_NAMESPACE
\ No newline at end of file
diff --git a/Engine/Plugins/Runtime/ModelViewViewModel/Source/ModelViewViewModelEditor/Private/Widgets/SMVVMViewModelBindingListWidget.h b/Engine/Plugins/Runtime/ModelViewViewModel/Source/ModelViewViewModelEditor/Private/Widgets/SMVVMViewModelBindingListWidget.h
--- a/Engine/Plugins/Runtime/ModelViewViewModel/Source/ModelViewViewModelEditor/Private/Widgets/SMVVMViewModelBindingListWidget.h
+++ b/Engine/Plugins/Runtime/ModelViewViewModel/Source/ModelViewViewModelEditor/Private/Widgets/SMVVMViewModelBindingListWidget.h
@@ -3,45 +3,91 @@
 #pragma once
 
 #include "CoreMinimal.h"
+#include "Framework/SlateDelegates.h"
+#include "Misc/IFilter.h"
+#include "Misc/TextFilter.h"
+#include "Framework/Views/TreeFilterHandler.h"
 #include "Templates/SubclassOf.h"
+#include "Types/MVVMFieldVariant.h"
 #include "Widgets/SCompoundWidget.h"
 
 struct FMVVMAvailableBinding;
-
+class FUICommandList;
+class ITableRow;
 template<typename T>
 class STreeView;
-class ITableRow;
 class STableViewBase;
+class SInlineEditableTextBlock;
 class UMVVMViewModelBase;
+
+
 namespace UE::MVVM
 {
-struct FMVVMFieldVariant;
 namespace Private
 {
-	class FMVVMAvailableBindingNode;
-}
-}
+class FAvailableBindingNode;
+}//namespace Private
 
 
-class SMVVMViewModelBindingListWidget : public SCompoundWidget
+class SViewModelBindingListWidget : public SCompoundWidget
 {
+	friend Private::FAvailableBindingNode;
 public:
-	SLATE_BEGIN_ARGS(SMVVMViewModelBindingListWidget) {}
+	struct FViewModel
+	{
+		TSubclassOf<UObject> Class;
+		FGuid ViewModelId;
+		FName Name;
+	};
+
+	DECLARE_DELEGATE_RetVal_FourParams(bool, FOnViewModelRenamed, const FViewModel&, const FText&, bool, FText&);
+	DECLARE_DELEGATE_OneParam(FOnViewModelDeleted, const FViewModel&);
+	DECLARE_DELEGATE_RetVal_OneParam(FViewModel, FOnViewModelDuplicated, const FViewModel&);
+
+	SLATE_BEGIN_ARGS(SViewModelBindingListWidget) {}
+		SLATE_ARGUMENT(TArrayView<const FViewModel>, ViewModels)
+		SLATE_EVENT(FOnViewModelRenamed, OnRenamed)
+		SLATE_EVENT(FOnViewModelDeleted, OnDeleted)
+		SLATE_EVENT(FOnViewModelDuplicated, OnDuplicated)
 	SLATE_END_ARGS()
 
 	void Construct(const FArguments& InArgs);
-	void SetViewModel(TSubclassOf<UMVVMViewModelBase> ViewModelClass);
+	void SetViewModel(TSubclassOf<UObject> ViewModelClass, FName ViewModelName);
+	void SetViewModels(TArrayView<const FViewModel> ViewModels);
+	FText SetRawFilterText(const FText& InFilterText);
 
 private:
-	typedef UE::MVVM::Private::FMVVMAvailableBindingNode FMVVMAvailableBindingNode;
+	virtual FReply OnKeyDown(const FGeometry& MyGeometry, const FKeyEvent& InKeyEvent) override;
+
+	FText GetRawFilterText() const;
+
+	void HandleGetFilterStrings(TSharedPtr<Private::FAvailableBindingNode> Item, TArray<FString>& OutStrings);
+	TSharedRef<ITableRow> HandleGenerateRow(TSharedPtr<Private::FAvailableBindingNode> Item, const TSharedRef<STableViewBase>& OwnerTable);
+	void HandleGetChildren(TSharedPtr<Private::FAvailableBindingNode> InParent, TArray<TSharedPtr<Private::FAvailableBindingNode>>& OutChildren);
+
+	TSharedPtr<SWidget> HandleContextMenuOpening();
+	void HandleDuplicateViewModel();
+	bool HandleCanDuplicateViewModel();
+	void HandleDeleteViewModel();
+	bool HandleCanDeleteViewModel();
+	void HandleRenameViewModel();
+	bool HandleCanRenameViewModel();
+
+	void CreateCommandList();
 
-	TSharedRef<ITableRow> HandleGenerateRow(TSharedPtr<FMVVMAvailableBindingNode> Item, const TSharedRef<STableViewBase>& OwnerTable);
-	void HandleGetChildren(TSharedPtr<FMVVMAvailableBindingNode> InParent, TArray<TSharedPtr<FMVVMAvailableBindingNode>>& OutChildren);
+	TSharedPtr<STreeView<TSharedPtr<Private::FAvailableBindingNode>>> TreeWidget;
+	TArray<TSharedPtr<Private::FAvailableBindingNode>> TreeSource;
+	TArray<TSharedPtr<Private::FAvailableBindingNode>> FilteredTreeSource;
 
-	TSharedRef<SWidget> GetFieldIcon(const UE::MVVM::FMVVMFieldVariant& FieldVariant);
+	TSharedPtr<FUICommandList> CommandList;
+	typedef TTextFilter<TSharedPtr<Private::FAvailableBindingNode>> ViewModelTextFilter;
+	TSharedPtr<ViewModelTextFilter> SearchFilter;
+	typedef TreeFilterHandler<TSharedPtr<Private::FAvailableBindingNode>> ViewModelTreeFilter;
+	TSharedPtr<ViewModelTreeFilter> FilterHandler;
 
-	TSharedPtr<STreeView<TSharedPtr<FMVVMAvailableBindingNode>>> TreeWidget;
-	TArray<TSharedPtr<FMVVMAvailableBindingNode>> TreeSource;
+	FOnViewModelRenamed OnRenamed;
+	FOnViewModelDeleted OnDeleted;
+	FOnViewModelDuplicated OnDuplicated;
+};
 
-	TSubclassOf<UMVVMViewModelBase> CurrentViewModel;
-};
\ No newline at end of file
+} //namespace UE::MVVM
\ No newline at end of file
diff --git a/Engine/Plugins/Runtime/ModelViewViewModel/Source/ModelViewViewModelEditor/Private/Widgets/SMVVMViewModelPanel.cpp b/Engine/Plugins/Runtime/ModelViewViewModel/Source/ModelViewViewModelEditor/Private/Widgets/SMVVMViewModelPanel.cpp
new file mode 100644
--- /dev/null
+++ b/Engine/Plugins/Runtime/ModelViewViewModel/Source/ModelViewViewModelEditor/Private/Widgets/SMVVMViewModelPanel.cpp
@@ -0,0 +1,208 @@
+// Copyright Epic Games, Inc. All Rights Reserved.
+
+#include "Widgets/SMVVMViewModelPanel.h"
+
+#include "Editor.h"
+#include "Editor/EditorEngine.h"
+#include "Kismet2/Kismet2NameValidators.h"
+#include "MVVMBlueprintView.h"
+#include "MVVMBlueprintViewModelContext.h"
+#include "MVVMEditorSubsystem.h"
+#include "MVVMViewModelBase.h"
+#include "WidgetBlueprint.h"
+#include "WidgetBlueprintEditor.h"
+
+#include "Styling/AppStyle.h"
+#include "Styling/SlateIconFinder.h"
+#include "Styling/MVVMEditorStyle.h"
+
+#include "Widgets/SBoxPanel.h"
+#include "Widgets/Input/SSearchBox.h"
+#include "Widgets/Layout/SBorder.h"
+#include "Widgets/SMVVMViewModelBindingListWidget.h"
+
+#define LOCTEXT_NAMESPACE "ViewModelPanel"
+
+namespace UE::MVVM
+{
+
+void SMVVMViewModelPanel::Construct(const FArguments& InArgs, TSharedPtr<FWidgetBlueprintEditor> WidgetBlueprintEditor)
+{
+	UWidgetBlueprint* WidgetBlueprint = WidgetBlueprintEditor->GetWidgetBlueprintObj();
+	check(WidgetBlueprint);
+	UMVVMBlueprintView* CurrentBlueprintView = GEditor->GetEditorSubsystem<UMVVMEditorSubsystem>()->GetView(WidgetBlueprint);
+
+	WeakBlueprintEditor = WidgetBlueprintEditor;
+	WeakBlueprintView = CurrentBlueprintView;
+
+	WidgetBlueprint->OnExtensionAdded.AddSP(this, &SMVVMViewModelPanel::HandleViewUpdated);
+	WidgetBlueprint->OnExtensionRemoved.AddSP(this, &SMVVMViewModelPanel::HandleViewUpdated);
+
+	if (CurrentBlueprintView)
+	{
+		// Listen to when the viewmodel are modified
+		CurrentBlueprintView->OnViewModelsUpdated.AddSP(this, &SMVVMViewModelPanel::HandleViewModelsUpdated);
+	}
+
+	GEditor->OnBlueprintCompiled().AddSP(this, &SMVVMViewModelPanel::HandleBlueprintCompiled);
+
+	ViewModelTreeView = SNew(SViewModelBindingListWidget)
+		.OnRenamed(this, &SMVVMViewModelPanel::HandleViewModelRenamed);
+
+	GenerateViewModelTreeView();
+
+	ChildSlot
+	[
+		SNew(SVerticalBox)
+
+		+ SVerticalBox::Slot()
+		.Padding(4)
+		.AutoHeight()
+		[
+			SAssignNew(SearchBoxPtr, SSearchBox)
+			.HintText(LOCTEXT("SearchViewmodel", "Search"))
+			.OnTextChanged(this, &SMVVMViewModelPanel::HandleSearchChanged)
+		]
+		+ SVerticalBox::Slot()
+		.FillHeight(1.0f)
+		[
+			SNew(SBorder)
+			.BorderImage(FAppStyle::Get().GetBrush("Brushes.Recessed"))
+			.Padding(0)
+			[
+				ViewModelTreeView.ToSharedRef()
+			]
+		]
+	];
+}
+
+
+SMVVMViewModelPanel::~SMVVMViewModelPanel()
+{
+	if (GEditor)
+	{
+		GEditor->OnBlueprintCompiled().RemoveAll(this);
+	}
+
+	if (TSharedPtr<FWidgetBlueprintEditor> WidgetBlueprintEditor = WeakBlueprintEditor.Pin())
+	{
+		if (UWidgetBlueprint* WidgetBlueprint = WidgetBlueprintEditor->GetWidgetBlueprintObj())
+		{
+			WidgetBlueprint->OnExtensionAdded.RemoveAll(this);
+			WidgetBlueprint->OnExtensionRemoved.RemoveAll(this);
+		}
+	}
+
+	if (UMVVMBlueprintView* CurrentBlueprintView = WeakBlueprintView.Get())
+	{
+		// bind to check if the view is enabled
+		CurrentBlueprintView->OnViewModelsUpdated.RemoveAll(this);
+	}
+}
+
+
+void SMVVMViewModelPanel::GenerateViewModelTreeView()
+{
+	if (ViewModelTreeView)
+	{
+		TArray<SViewModelBindingListWidget::FViewModel, TInlineAllocator<16>> ViewModelList;
+		if (UMVVMBlueprintView* View = WeakBlueprintView.Get())
+		{
+			for (const FMVVMBlueprintViewModelContext& ViewModelContext : View->GetViewModels())
+			{
+				SViewModelBindingListWidget::FViewModel ViewModel;
+				ViewModel.Class = ViewModelContext.GetViewModelClass();
+				ViewModel.ViewModelId = ViewModelContext.GetViewModelId();
+				ViewModel.Name = ViewModelContext.GetViewModelName();
+				ViewModelList.Add(ViewModel);
+			}
+		}
+
+		ViewModelTreeView->SetViewModels(ViewModelList);
+	}
+}
+
+
+void SMVVMViewModelPanel::HandleViewUpdated(UBlueprintExtension*)
+{
+	bool bViewUpdated = false;
+	if (TSharedPtr<FWidgetBlueprintEditor> WidgetBlueprintEditor = WeakBlueprintEditor.Pin())
+	{
+		if (UWidgetBlueprint* WidgetBlueprint = WidgetBlueprintEditor->GetWidgetBlueprintObj())
+		{
+			UMVVMBlueprintView* CurrentBlueprintView = GEditor->GetEditorSubsystem<UMVVMEditorSubsystem>()->GetView(WidgetBlueprint);
+			WeakBlueprintView = CurrentBlueprintView;
+
+			if (CurrentBlueprintView)
+			{
+				CurrentBlueprintView->OnViewModelsUpdated.AddSP(this, &SMVVMViewModelPanel::HandleViewModelsUpdated);
+			}
+		}
+	}
+	GenerateViewModelTreeView();
+}
+
+
+void SMVVMViewModelPanel::HandleViewModelsUpdated()
+{
+	GenerateViewModelTreeView();
+}
+
+
+void SMVVMViewModelPanel::HandleSearchChanged(const FText& InFilterText)
+{
+	SearchBoxPtr->SetError(ViewModelTreeView->SetRawFilterText(InFilterText));
+}
+
+
+bool SMVVMViewModelPanel::HandleViewModelRenamed(const SViewModelBindingListWidget::FViewModel& ViewModel, const FText& RenameTo, bool bCommit, FText& OutErrorMessage)
+{
+	if (RenameTo.IsEmptyOrWhitespace())
+	{
+		OutErrorMessage = LOCTEXT("EmptyWidgetName", "Empty Widget Name");
+		return false;
+	}
+
+	const FString& NewNameString = RenameTo.ToString();
+	if (NewNameString.Len() >= NAME_SIZE)
+	{
+		OutErrorMessage = LOCTEXT("WidgetNameTooLong", "Widget Name is Too Long");
+		return false;
+	}
+
+	FString GeneratedName = SlugStringForValidName(NewNameString);
+	if (GeneratedName.IsEmpty())
+	{
+		OutErrorMessage = LOCTEXT("EmptyWidgetName", "Empty Widget Name");
+		return false;
+	}
+
+	const FName GeneratedFName(*GeneratedName);
+	check(GeneratedFName.IsValidXName(INVALID_OBJECTNAME_CHARACTERS));
+
+	if (TSharedPtr<FWidgetBlueprintEditor> WidgetBlueprintEditor = WeakBlueprintEditor.Pin())
+	{
+		if (UWidgetBlueprint* WidgetBP = WidgetBlueprintEditor->GetWidgetBlueprintObj())
+		{
+			if (bCommit)
+			{
+				return GEditor->GetEditorSubsystem<UMVVMEditorSubsystem>()->RenameViewModel(WidgetBP, ViewModel.Name, *NewNameString, OutErrorMessage);
+			}
+			else
+			{
+				return GEditor->GetEditorSubsystem<UMVVMEditorSubsystem>()->VerifyViewModelRename(WidgetBP, ViewModel.Name, *NewNameString, OutErrorMessage);
+			}
+		}
+	}
+	return false;
+}
+
+
+void SMVVMViewModelPanel::HandleBlueprintCompiled()
+{
+	GenerateViewModelTreeView();
+}
+
+} // namespace
+
+#undef LOCTEXT_NAMESPACE
\ No newline at end of file
diff --git a/Engine/Plugins/Runtime/ModelViewViewModel/Source/ModelViewViewModelEditor/Private/Widgets/SMVVMViewModelPanel.h b/Engine/Plugins/Runtime/ModelViewViewModel/Source/ModelViewViewModelEditor/Private/Widgets/SMVVMViewModelPanel.h
new file mode 100644
--- /dev/null
+++ b/Engine/Plugins/Runtime/ModelViewViewModel/Source/ModelViewViewModelEditor/Private/Widgets/SMVVMViewModelPanel.h
@@ -0,0 +1,50 @@
+// Copyright Epic Games, Inc. All Rights Reserved.
+
+#pragma once
+
+#include "CoreMinimal.h"
+
+#include "Widgets/SMVVMViewModelBindingListWidget.h"
+#include "Widgets/SCompoundWidget.h"
+
+class FWidgetBlueprintEditor;
+class SSearchBox;
+class STableViewBase;
+template<typename ItemType>
+class SListView;
+class UBlueprintExtension;
+class UMVVMBlueprintView;
+
+namespace UE::MVVM
+{
+
+/**
+ * 
+ */
+class SMVVMViewModelPanel : public SCompoundWidget
+{
+public:
+	SLATE_BEGIN_ARGS(SMVVMViewModelPanel) {}
+	SLATE_END_ARGS()
+
+	void Construct(const FArguments& InArgs, TSharedPtr<FWidgetBlueprintEditor> Editor);
+	virtual ~SMVVMViewModelPanel();
+
+private:
+	void HandleViewUpdated(UBlueprintExtension*);
+	void HandleViewModelsUpdated();
+	void HandleSearchChanged(const FText& InFilterText);
+	bool HandleViewModelRenamed(const SViewModelBindingListWidget::FViewModel& ViewModel, const FText& RenameTo, bool bCommit, FText& OutErrorMessage);
+	void HandleBlueprintCompiled();
+
+	void GenerateViewModelTreeView();
+
+private:
+	TSharedPtr<SSearchBox> SearchBoxPtr;
+	TSharedPtr<SViewModelBindingListWidget> ViewModelTreeView;
+
+	TWeakPtr<FWidgetBlueprintEditor> WeakBlueprintEditor;
+	TWeakObjectPtr<UMVVMBlueprintView> WeakBlueprintView;
+};
+
+} // namespace
diff --git a/Engine/Plugins/Runtime/ModelViewViewModel/Source/ModelViewViewModelEditor/Public/MVVMEditorSubsystem.h b/Engine/Plugins/Runtime/ModelViewViewModel/Source/ModelViewViewModelEditor/Public/MVVMEditorSubsystem.h
--- a/Engine/Plugins/Runtime/ModelViewViewModel/Source/ModelViewViewModelEditor/Public/MVVMEditorSubsystem.h
+++ b/Engine/Plugins/Runtime/ModelViewViewModel/Source/ModelViewViewModelEditor/Public/MVVMEditorSubsystem.h
@@ -27,7 +27,13 @@ public:
 	UMVVMBlueprintView* GetView(UWidgetBlueprint* WidgetBlueprint) const;
 	
 	UFUNCTION(BlueprintCallable, Category = "MVVM")
-	void RemoveViewModel(UWidgetBlueprint* WidgetBlueprint, const FMVVMBlueprintViewModelContext& ViewModel);
+	void RemoveViewModel(UWidgetBlueprint* WidgetBlueprint, FName ViewModel);
+
+	UFUNCTION(BlueprintCallable, Category = "MVVM")
+	bool VerifyViewModelRename(UWidgetBlueprint* WidgetBlueprint, FName ViewModel, FName NewViewModel, FText& OutError);
+
+	UFUNCTION(BlueprintCallable, Category = "MVVM")
+	bool RenameViewModel(UWidgetBlueprint* WidgetBlueprint, FName ViewModel, FName NewViewModel, FText& OutError);
 
 	UFUNCTION(BlueprintCallable, Category = "MVVM")
 	FMVVMBlueprintViewBinding& AddBinding(UWidgetBlueprint* WidgetBlueprint);
@@ -35,6 +41,7 @@ public:
 	UFUNCTION(BlueprintCallable, Category = "MVVM")
 	void RemoveBinding(UWidgetBlueprint* WidgetBlueprint, const FMVVMBlueprintViewBinding& Binding);
 
-	TArray<UE::MVVM::FMVVMConstFieldVariant> GetChildViewModels(TSubclassOf<UMVVMViewModelBase> Class);
+	TArray<UE::MVVM::FMVVMConstFieldVariant> GetChildViewModels(UClass* Class);
+
 	TArray<const UFunction*> GetAvailableConversionFunctions(const UE::MVVM::FMVVMConstFieldVariant& Source, const UE::MVVM::FMVVMConstFieldVariant& Dest) const;
 };
