MVVM - Removed unnecessary API tags from FunctionGraphHelper.

#rnx

diff --git a/Engine/Plugins/Runtime/ModelViewViewModel/Source/ModelViewViewModelBlueprint/Private/MVVMFunctionGraphHelper.h b/Engine/Plugins/Runtime/ModelViewViewModel/Source/ModelViewViewModelBlueprint/Private/MVVMFunctionGraphHelper.h
--- a/Engine/Plugins/Runtime/ModelViewViewModel/Source/ModelViewViewModelBlueprint/Private/MVVMFunctionGraphHelper.h
+++ b/Engine/Plugins/Runtime/ModelViewViewModel/Source/ModelViewViewModelBlueprint/Private/MVVMFunctionGraphHelper.h
@@ -7,6 +7,7 @@
 #include "UObject/Object.h"
 
 class FKismetCompilerContext;
+class UBlueprint;
 class UEdGraph;
 
 namespace UE::MVVM::FunctionGraphHelper
@@ -16,21 +17,21 @@ namespace UE::MVVM::FunctionGraphHelper
 	MODELVIEWVIEWMODELBLUEPRINT_API UEdGraph* CreateFunctionGraph(FKismetCompilerContext& Context, const FStringView FunctionName, EFunctionFlags ExtraFunctionFlag, const FStringView Category, bool bIsEditable);
 
 	/** Generate a function graph that cannot be viewed in the BP editor. */
-	MODELVIEWVIEWMODELBLUEPRINT_API UEdGraph* CreateIntermediateFunctionGraph(FKismetCompilerContext& Context, const FStringView FunctionName, EFunctionFlags ExtraFunctionFlag, const FStringView Category, bool bIsEditable);
+	UEdGraph* CreateIntermediateFunctionGraph(FKismetCompilerContext& Context, const FStringView FunctionName, EFunctionFlags ExtraFunctionFlag, const FStringView Category, bool bIsEditable);
 
 	/** Add an input argument to an existing function graph. */
-	MODELVIEWVIEWMODELBLUEPRINT_API bool AddFunctionArgument(UEdGraph* FunctionGraph, TSubclassOf<UObject> Argument, FName ArgumentName);
+	bool AddFunctionArgument(UEdGraph* FunctionGraph, TSubclassOf<UObject> Argument, FName ArgumentName);
 
 	/** Add an input argument to an existing function graph. */
-	MODELVIEWVIEWMODELBLUEPRINT_API bool AddFunctionArgument(UEdGraph* FunctionGraph, FProperty* Argument, FName ArgumentName);
+	bool AddFunctionArgument(UEdGraph* FunctionGraph, FProperty* Argument, FName ArgumentName);
 
 	/** Generate the nodes needed for to call the MVVMView::SetViewModel function. */
-	MODELVIEWVIEWMODELBLUEPRINT_API bool GenerateViewModelSetter(FKismetCompilerContext& Context, UEdGraph* FunctionGraph, FName IiewModelName);
+	bool GenerateViewModelSetter(FKismetCompilerContext& Context, UEdGraph* FunctionGraph, FName ViewModelName);
 
 	/** Generate the nodes needed for to set the property and to call the broadcast field notify function. */
-	MODELVIEWVIEWMODELBLUEPRINT_API bool GenerateViewModelFieldNotifySetter(FKismetCompilerContext& Context, UEdGraph* FunctionGraph, FProperty* Property, FName InputPinName);
+	bool GenerateViewModelFieldNotifySetter(FKismetCompilerContext& Context, UEdGraph* FunctionGraph, FProperty* Property, FName InputPinName);
 
 	/** Generate the nodes needed to broadcast field notify function. */
-	MODELVIEWVIEWMODELBLUEPRINT_API bool GenerateViewModelFielNotifyBroadcast(FKismetCompilerContext& Context, UEdGraph* FunctionGraph, FProperty* Property);
+	bool GenerateViewModelFielNotifyBroadcast(FKismetCompilerContext& Context, UEdGraph* FunctionGraph, FProperty* Property);
 
 } //namespace
\ No newline at end of file
