MVVM - FMVVMBlueprintPropertyPath can now be stored in containers due to added GetTypeHash() and list type traits override.

#rnx

diff --git a/Engine/Plugins/Runtime/ModelViewViewModel/Source/ModelViewViewModelBlueprint/Public/MVVMPropertyPath.h b/Engine/Plugins/Runtime/ModelViewViewModel/Source/ModelViewViewModelBlueprint/Public/MVVMPropertyPath.h
--- a/Engine/Plugins/Runtime/ModelViewViewModel/Source/ModelViewViewModelBlueprint/Public/MVVMPropertyPath.h
+++ b/Engine/Plugins/Runtime/ModelViewViewModel/Source/ModelViewViewModelBlueprint/Public/MVVMPropertyPath.h
@@ -179,12 +179,17 @@
 	}
 };
 
+inline uint32 GetTypeHash(const FMVVMBlueprintFieldPath& FieldPath)
+{
+	return GetTypeHash(FieldPath.GetFieldName());
+}
+
 /**
  * Base path to properties for MVVM view models and widgets.
  * 
  * Used to associate properties within MVVM bindings in editor & during MVVM compilation
  */
-USTRUCT()
+USTRUCT(BlueprintType)
 struct FMVVMBlueprintPropertyPath
 {
 	GENERATED_BODY()
@@ -340,6 +345,18 @@
 	}
 };
 
+inline uint32 GetTypeHash(const FMVVMBlueprintPropertyPath& Path)
+{
+	uint32 Hash = 0;
+	TArray<FName> Paths = Path.GetPaths();
+
+	for (const FName& SubPath : Paths)
+	{
+		Hash = HashCombine(Hash, GetTypeHash(SubPath));
+	}
+	return Hash;
+}
+
 template<>
 struct TStructOpsTypeTraits<FMVVMBlueprintPropertyPath> : public TStructOpsTypeTraitsBase2<FMVVMBlueprintPropertyPath>
 {
diff --git a/Engine/Plugins/Runtime/ModelViewViewModel/Source/ModelViewViewModelEditor/Private/MVVMFieldVariantListTraits.h b/Engine/Plugins/Runtime/ModelViewViewModel/Source/ModelViewViewModelEditor/Private/MVVMFieldVariantListTraits.h
--- a/Engine/Plugins/Runtime/ModelViewViewModel/Source/ModelViewViewModelEditor/Private/MVVMFieldVariantListTraits.h
+++ b/Engine/Plugins/Runtime/ModelViewViewModel/Source/ModelViewViewModelEditor/Private/MVVMFieldVariantListTraits.h
@@ -3,6 +3,7 @@
 #pragma once
 
 #include "Framework/Views/TableViewTypeTraits.h"
+#include "MVVMPropertyPath.h"
 #include "Types/MVVMFieldVariant.h"
 
 template <>
@@ -110,5 +111,59 @@ struct TListTypeTraits<UE::MVVM::FMVVMConstFieldVariant>
 		return InPtr.GetName().ToString();
 	}
 
+	class SerializerType {};
+};
+
+template <>
+struct TIsValidListItem<FMVVMBlueprintPropertyPath>
+{
+	enum
+	{
+		Value = true
+	};
+};
+
+template <>
+struct TListTypeTraits<FMVVMBlueprintPropertyPath>
+{
+	typedef FMVVMBlueprintPropertyPath NullableType;
+
+	using MapKeyFuncs = TDefaultMapHashableKeyFuncs<FMVVMBlueprintPropertyPath, TSharedRef<ITableRow>, false>;
+	using MapKeyFuncsSparse = TDefaultMapHashableKeyFuncs<FMVVMBlueprintPropertyPath, FSparseItemInfo, false>;
+	using SetKeyFuncs = DefaultKeyFuncs<FMVVMBlueprintPropertyPath>;
+
+	template<typename U>
+	static void AddReferencedObjects(FReferenceCollector&,
+		TArray<FMVVMBlueprintPropertyPath>&,
+		TSet<FMVVMBlueprintPropertyPath>&,
+		TMap<const U*, FMVVMBlueprintPropertyPath>&)
+	{
+	}
+
+	static bool IsPtrValid(const FMVVMBlueprintPropertyPath& InPtr)
+	{
+		return !InPtr.IsEmpty();
+	}
+
+	static void ResetPtr(FMVVMBlueprintPropertyPath& InPtr)
+	{
+		InPtr = FMVVMBlueprintPropertyPath();
+	}
+
+	static FMVVMBlueprintPropertyPath MakeNullPtr()
+	{
+		return FMVVMBlueprintPropertyPath();
+	}
+
+	static FMVVMBlueprintPropertyPath NullableItemTypeConvertToItemType(const FMVVMBlueprintPropertyPath& InPtr)
+	{
+		return InPtr;
+	}
+
+	static FString DebugDump(FMVVMBlueprintPropertyPath InPtr)
+	{
+		return InPtr.GetBasePropertyPath();
+	}
+
 	class SerializerType {};
 };
\ No newline at end of file
