MVVM: Fix crash when a field is renamed after the path was created.
#prefligh 62a0a61ae820a589df465387

diff --git a/Engine/Plugins/Runtime/ModelViewViewModel/Source/ModelViewViewModelBlueprint/Private/MVVMViewBlueprintCompiler.cpp b/Engine/Plugins/Runtime/ModelViewViewModel/Source/ModelViewViewModelBlueprint/Private/MVVMViewBlueprintCompiler.cpp
--- a/Engine/Plugins/Runtime/ModelViewViewModel/Source/ModelViewViewModelBlueprint/Private/MVVMViewBlueprintCompiler.cpp
+++ b/Engine/Plugins/Runtime/ModelViewViewModel/Source/ModelViewViewModelBlueprint/Private/MVVMViewBlueprintCompiler.cpp
@@ -712,6 +712,12 @@
 		if (IsForwardBinding(Binding.BindingType))
 		{
 			const FBindingSourceContext BindingSourceContext = CreateBindingSourceContext(BlueprintView, Class, Binding.ViewModelPath);
+			if (!IsPropertyPathValid(BindingSourceContext.PropertyPath))
+			{
+				AddErrorForBinding(Binding, BlueprintView, FString::Printf(TEXT("The property path '%s' is invalid."), *PropertyPathToString(BlueprintView, Binding.ViewModelPath)));
+				bIsBindingsValid = false;
+				continue;
+			}
 			if (BindingSourceContext.bIsRootWidget && BindingSourceContext.SourceClass == nullptr)
 			{
 				AddErrorForBinding(Binding, BlueprintView, TEXT("Internal error. The binding could not find its source."));
@@ -746,6 +752,12 @@
 			TArray<UE::MVVM::FMVVMConstFieldVariant> SetterPath;
 			{
 				SetterPath = CreateBindingDestinationPath(BlueprintView, Class, Binding.WidgetPath);
+				if (!IsPropertyPathValid(SetterPath))
+				{
+					AddErrorForBinding(Binding, BlueprintView, FString::Printf(TEXT("The property path '%s' is invalid."), *PropertyPathToString(BlueprintView, Binding.WidgetPath)));
+					bIsBindingsValid = false;
+					continue;
+				}
 			}
 
 			TValueOrError<FCompilerBinding, FString> AddBindingResult = AddBinding(Class, GetterFields, SetterPath, Binding.Conversion.SourceToDestinationFunction);
@@ -768,6 +780,12 @@
 		if (IsBackwardBinding(Binding.BindingType))
 		{
 			const FBindingSourceContext BindingSourceContext = CreateBindingSourceContext(BlueprintView, Class, Binding.WidgetPath);
+			if (!IsPropertyPathValid(BindingSourceContext.PropertyPath))
+			{
+				AddErrorForBinding(Binding, BlueprintView, FString::Printf(TEXT("The property path '%s' is invalid."), *PropertyPathToString(BlueprintView, Binding.WidgetPath)));
+				bIsBindingsValid = false;
+				continue;
+			}
 			if (BindingSourceContext.bIsRootWidget && BindingSourceContext.SourceClass == nullptr)
 			{
 				AddErrorForBinding(Binding, BlueprintView, TEXT("Internal error. The binding could not find its source."));
@@ -802,6 +820,12 @@
 			TArray<UE::MVVM::FMVVMConstFieldVariant> SetterPath;
 			{
 				SetterPath = CreateBindingDestinationPath(BlueprintView, Class, Binding.ViewModelPath);
+				if (!IsPropertyPathValid(SetterPath))
+				{
+					AddErrorForBinding(Binding, BlueprintView, FString::Printf(TEXT("The property path '%s' is invalid."), *PropertyPathToString(BlueprintView, Binding.ViewModelPath)));
+					bIsBindingsValid = false;
+					continue;
+				}
 			}
 
 			TValueOrError<FCompilerBinding, FString> AddBindingResult = AddBinding(Class, GetterFields, SetterPath, Binding.Conversion.DestinationToSourceFunction);
@@ -983,26 +1007,43 @@
 }
 
 
-FString FMVVMViewBlueprintCompiler::CreatePropertyPath(FName PropertyName, FString PropertyPath) const
+FString FMVVMViewBlueprintCompiler::PropertyPathToString(const UMVVMBlueprintView* BlueprintView, const FMVVMBlueprintPropertyPath& PropertyPath) const
 {
-	TStringBuilder<512> Builder;
-	if (PropertyName.IsNone())
+	if (PropertyPath.IsEmpty())
 	{
-		Builder << PropertyPath;
+		return FString();
+	}
+
+	TStringBuilder<512> Result;
+	if (PropertyPath.IsFromViewModel())
+	{
+		if (const FMVVMBlueprintViewModelContext* SourceViewModelContext = BlueprintView->FindViewModel(PropertyPath.GetViewModelId()))
+		{
+			Result << SourceViewModelContext->GetViewModelName();
+		}
+		else
+		{
+			Result << TEXT("<none>");
+		}
+	}
+	else if (PropertyPath.IsFromWidget())
+	{
+		Result << PropertyPath.GetWidgetName();
 	}
 	else
 	{
-		Builder << PropertyName;
-		if (!PropertyPath.IsEmpty())
-		{
-			Builder << TEXT('.');
-			Builder << PropertyPath;
-		}
+		Result << TEXT("<none>");
 	}
-	return Builder.ToString();
+
+	FString BasePropertyPath = PropertyPath.GetBasePropertyPath();
+	if (BasePropertyPath.Len())
+	{
+		Result << TEXT('.');
+		Result << MoveTemp(BasePropertyPath);
+	}
+	return Result.ToString();
 }
 
-
 TArray<FMVVMConstFieldVariant> FMVVMViewBlueprintCompiler::CreatePropertyPath(const UClass* Class, FName PropertyName, TArray<FMVVMConstFieldVariant> Properties) const
 {
 	if (PropertyName.IsNone())
@@ -1016,4 +1057,25 @@
 	return Properties;
 }
 
+
+bool FMVVMViewBlueprintCompiler::IsPropertyPathValid(TArrayView<const FMVVMConstFieldVariant> PropertyPath) const
+{
+	for (const FMVVMConstFieldVariant& Field : PropertyPath)
+	{
+		if (Field.IsEmpty())
+		{
+			return false;
+		}
+		if (Field.IsProperty() && Field.GetProperty() == nullptr)
+		{
+			return false;
+		}
+		if (Field.IsFunction() && Field.GetFunction() == nullptr)
+		{
+			return false;
+		}
+	}
+	return true;
+}
+
 } //namespace
diff --git a/Engine/Plugins/Runtime/ModelViewViewModel/Source/ModelViewViewModelBlueprint/Private/MVVMViewBlueprintCompiler.h b/Engine/Plugins/Runtime/ModelViewViewModel/Source/ModelViewViewModelBlueprint/Private/MVVMViewBlueprintCompiler.h
--- a/Engine/Plugins/Runtime/ModelViewViewModel/Source/ModelViewViewModelBlueprint/Private/MVVMViewBlueprintCompiler.h
+++ b/Engine/Plugins/Runtime/ModelViewViewModel/Source/ModelViewViewModelBlueprint/Private/MVVMViewBlueprintCompiler.h
@@ -73,8 +73,9 @@
 	FBindingSourceContext CreateBindingSourceContext(const UMVVMBlueprintView* BlueprintView, const UWidgetBlueprintGeneratedClass* Class, const FMVVMBlueprintPropertyPath& PropertyPath) const;
 	TArray<FMVVMConstFieldVariant> CreateBindingDestinationPath(const UMVVMBlueprintView* BlueprintView, const UWidgetBlueprintGeneratedClass* Class, const FMVVMBlueprintPropertyPath& PropertyPath) const;
 
-	FString CreatePropertyPath(FName PropertyName, FString PropertyPath) const;
+	FString PropertyPathToString(const UMVVMBlueprintView* BlueprintView, const FMVVMBlueprintPropertyPath& PropertyPath) const;
 	TArray<FMVVMConstFieldVariant> CreatePropertyPath(const UClass* Class, FName PropertyName, TArray<FMVVMConstFieldVariant> Properties) const;
+	bool IsPropertyPathValid(TArrayView<const FMVVMConstFieldVariant> PropertyPath) const;
 
 private:
 	struct FCompilerSourceContext
diff --git a/Engine/Plugins/Runtime/ModelViewViewModel/Source/ModelViewViewModelBlueprint/Public/MVVMPropertyPath.h b/Engine/Plugins/Runtime/ModelViewViewModel/Source/ModelViewViewModelBlueprint/Public/MVVMPropertyPath.h
--- a/Engine/Plugins/Runtime/ModelViewViewModel/Source/ModelViewViewModelBlueprint/Public/MVVMPropertyPath.h
+++ b/Engine/Plugins/Runtime/ModelViewViewModel/Source/ModelViewViewModelBlueprint/Public/MVVMPropertyPath.h
@@ -317,7 +317,7 @@ public:
 			{
 				Result << TEXT('.');
 			}
-			Result << Path.GetFieldName().ToString();
+			Result << Path.GetFieldName();
 		}
 		return Result.ToString();
 	}
