MVVM: Parameters that are const ref in BP have a an extra flag to mark them as param.

diff --git a/Engine/Plugins/Runtime/ModelViewViewModel/Source/ModelViewViewModel/Private/Bindings/MVVMBindingHelper.cpp b/Engine/Plugins/Runtime/ModelViewViewModel/Source/ModelViewViewModel/Private/Bindings/MVVMBindingHelper.cpp
--- a/Engine/Plugins/Runtime/ModelViewViewModel/Source/ModelViewViewModel/Private/Bindings/MVVMBindingHelper.cpp
+++ b/Engine/Plugins/Runtime/ModelViewViewModel/Source/ModelViewViewModel/Private/Bindings/MVVMBindingHelper.cpp
@@ -492,7 +492,7 @@ namespace UE::MVVM::BindingHelper
 		{
 			for (TFieldIterator<FProperty> It(InFunction); It && (It->PropertyFlags & CPF_Parm); ++It)
 			{
-				if (It->HasAllPropertyFlags(CPF_OutParm) && !It->HasAnyPropertyFlags(CPF_Deprecated | CPF_EditorOnly | CPF_ConstParm))
+				if (It->HasAllPropertyFlags(CPF_OutParm) && !It->HasAnyPropertyFlags(CPF_Deprecated | CPF_EditorOnly | CPF_ConstParm | CPF_ReferenceParm))
 				{
 					Result = *It;
 					break;
@@ -512,7 +512,7 @@ namespace UE::MVVM::BindingHelper
 		{
 			if (It->HasAllPropertyFlags(CPF_Parm) && !It->HasAnyPropertyFlags(CPF_Deprecated | CPF_EditorOnly | CPF_ReturnParm))
 			{
-				if (It->HasAllPropertyFlags(CPF_OutParm) && !It->HasAllPropertyFlags(CPF_ConstParm))
+				if (It->HasAllPropertyFlags(CPF_OutParm) && !It->HasAnyPropertyFlags(CPF_ConstParm | CPF_ReferenceParm))
 				{
 					continue;
 				}
