MVVM - Removed references to unused styles that were causing warnings on load.

#rnx

diff --git a/Engine/Plugins/Runtime/ModelViewViewModel/Source/ModelViewViewModelEditor/Private/Styling/MVVMEditorStyle.cpp b/Engine/Plugins/Runtime/ModelViewViewModel/Source/ModelViewViewModelEditor/Private/Styling/MVVMEditorStyle.cpp
--- a/Engine/Plugins/Runtime/ModelViewViewModel/Source/ModelViewViewModelEditor/Private/Styling/MVVMEditorStyle.cpp
+++ b/Engine/Plugins/Runtime/ModelViewViewModel/Source/ModelViewViewModelEditor/Private/Styling/MVVMEditorStyle.cpp
@@ -32,8 +32,8 @@ FMVVMEditorStyle::FMVVMEditorStyle()
 
 	// Class Icons
 	{
-		Set("ClassIcon.MVVMBlueprintView", new IMAGE_BRUSH("Slate/PH_MVVMBlueprintView_16x", Icon16x16));
-		Set("ClassIcon.MVVMBlueprintViewModel", new IMAGE_BRUSH("Slate/PH_MVVMBlueprintViewModel_16x", Icon16x16));
+		Set("ClassIcon.MVVMBlueprintView", new IMAGE_BRUSH_SVG("Slate/ViewModel", Icon16x16));
+		Set("ClassIcon.MVVMBlueprintViewModel", new IMAGE_BRUSH_SVG("Slate/ViewModel", Icon16x16));
 	}
 
 	Set("BlueprintView.TabIcon", new IMAGE_BRUSH_SVG("Slate/ViewModel", Icon16x16));
@@ -50,8 +50,6 @@ FMVVMEditorStyle::FMVVMEditorStyle()
 	}
 
 	{
-		Set("PropertyPath.ContextBorder", new BOX_BRUSH("Common/LightGroupBorder", FMargin(4.0f / 16.0f)));
-
 		Set("PropertyPath.ContextText", FTextBlockStyle(FAppStyle::Get().GetWidgetStyle<FTextBlockStyle>("NormalText"))
 		.SetFont(DEFAULT_FONT("Bold", 8))
 		.SetColorAndOpacity(FLinearColor(1.0f, 1.0f, 1.0f, 0.9f))
