MVVM: Prevent static array binding. Add option to auto convert integral property.

diff --git a/Engine/Plugins/Runtime/ModelViewViewModel/Source/ModelViewViewModel/Private/Bindings/MVVMBindingHelper.cpp b/Engine/Plugins/Runtime/ModelViewViewModel/Source/ModelViewViewModel/Private/Bindings/MVVMBindingHelper.cpp
--- a/Engine/Plugins/Runtime/ModelViewViewModel/Source/ModelViewViewModel/Private/Bindings/MVVMBindingHelper.cpp
+++ b/Engine/Plugins/Runtime/ModelViewViewModel/Source/ModelViewViewModel/Private/Bindings/MVVMBindingHelper.cpp
@@ -6,6 +6,10 @@
 #include "UObject/PropertyAccessUtil.h"
 #include "UObject/UnrealType.h"
 
+#ifndef UE_MVVM_ALLOW_AUTO_INTEGRAL_CONVERSION
+#define UE_MVVM_ALLOW_AUTO_INTEGRAL_CONVERSION 1
+#endif
+
 namespace UE::MVVM::BindingHelper
 {
 	namespace Private
@@ -40,6 +44,7 @@ namespace UE::MVVM::BindingHelper
 	bool IsValidForSourceBinding(const FProperty* InProperty)
 	{
 		return InProperty != nullptr
+			&& InProperty->ArrayDim == 1
 			&& !InProperty->HasAnyPropertyFlags(CPF_Deprecated | CPF_EditorOnly)
 			&& InProperty->HasAnyPropertyFlags(CPF_BlueprintVisible | CPF_BlueprintAssignable)
 #if WITH_EDITOR
@@ -52,6 +57,7 @@ namespace UE::MVVM::BindingHelper
 	bool IsValidForDestinationBinding(const FProperty* InProperty)
 	{
 		return InProperty != nullptr
+			&& InProperty->ArrayDim == 1
 			&& !InProperty->HasAnyPropertyFlags(CPF_Deprecated | CPF_EditorOnly | CPF_BlueprintReadOnly)
 			&& InProperty->HasAnyPropertyFlags(CPF_BlueprintVisible | CPF_BlueprintAssignable)
 #if WITH_EDITOR
@@ -179,6 +185,11 @@ namespace UE::MVVM::BindingHelper
 			{
 				return TEXT("The property is invalid.");
 			}
+			
+			if (InProperty->ArrayDim != 1)
+			{
+				return FString::Printf(TEXT("The property '%s' is a static array."), *InProperty->GetName());
+			}
 
 			if (InProperty->HasAnyPropertyFlags(CPF_Deprecated))
 			{
@@ -422,48 +433,54 @@ namespace UE::MVVM::BindingHelper
 		}
 
 
-		bool IsRealNumberConversion(const FProperty* Source, const FProperty* Destination)
+		bool IsNumericConversionRequired(const FProperty* Source, const FProperty* Destination)
 		{
-			return (CastField<FDoubleProperty>(Source) && CastField<FFloatProperty>(Destination))
-				|| (CastField<FFloatProperty>(Source) && CastField<FDoubleProperty>(Destination));
+			const FNumericProperty* SourceNumericProperty = CastField<const FNumericProperty>(Source);
+			const FNumericProperty* DestinationNumericProperty = CastField<const FNumericProperty>(Destination);
+			if (SourceNumericProperty && DestinationNumericProperty)
+			{
+				const bool bSameType = Destination->SameType(Source);
+				const bool bBothFloatingPoint = SourceNumericProperty->IsFloatingPoint() && DestinationNumericProperty->IsFloatingPoint();
+#if UE_MVVM_ALLOW_AUTO_INTEGRAL_CONVERSION
+				const bool bBothIntegral = SourceNumericProperty->IsInteger() && DestinationNumericProperty->IsInteger();
+				return !bSameType && (bBothFloatingPoint || bBothIntegral);
+#else
+				return !bSameType && bBothFloatingPoint;
+#endif
+			}
+			return false;
 		}
 
-		void ConvertRealNumber(const FProperty* Source, const FProperty* Destination, void* Data)
+		void ConvertNumeric(const FProperty* Source, const FProperty* Destination, void* Data)
 		{
 			check(Source);
 			check(Destination);
 			check(Data);
-			check(Source->ArrayDim == Destination->ArrayDim);
 
-			if (const FDoubleProperty* SrcDoubleProp = CastField<FDoubleProperty>(Source))
+			const FNumericProperty* SourceNumericProperty = CastField<const FNumericProperty>(Source);
+			const FNumericProperty* DestinationNumericProperty = CastField<const FNumericProperty>(Destination);
+			if (SourceNumericProperty && DestinationNumericProperty)
 			{
-				if (const FFloatProperty* DestFloatProp = CastField<FFloatProperty>(Destination))
+				if (SourceNumericProperty->IsFloatingPoint() && DestinationNumericProperty->IsFloatingPoint())
 				{
-					//double to float
-					for (int32 Idx = 0; Idx < Destination->ArrayDim; ++Idx)
-					{
-						const void* SrcElemValue = static_cast<const uint8*>(Data) + (Source->ElementSize * Idx);
-						void* DestElemValue = static_cast<uint8*>(Data) + (Destination->ElementSize * Idx);
-
-						const double Value = SrcDoubleProp->GetFloatingPointPropertyValue(SrcElemValue);
-						DestFloatProp->SetFloatingPointPropertyValue(DestElemValue, Value);
-					}
+					//floating to floating
+					const void* SrcElemValue = static_cast<const uint8*>(Data);
+					void* DestElemValue = static_cast<uint8*>(Data);
+
+					const double Value = SourceNumericProperty->GetFloatingPointPropertyValue(SrcElemValue);
+					DestinationNumericProperty->SetFloatingPointPropertyValue(DestElemValue, Value);
 				}
-			}
-			else if (const FFloatProperty* SrcFloatProp = CastField<FFloatProperty>(Source))
-			{
-				if (const FDoubleProperty* DestDoubleProp = CastField<FDoubleProperty>(Destination))
+#if UE_MVVM_ALLOW_AUTO_INTEGRAL_CONVERSION
+				else if (SourceNumericProperty->IsInteger() && DestinationNumericProperty->IsInteger())
 				{
-					//float to double (backward)
-					for (int32 Idx = Destination->ArrayDim - 1; Idx >= 0; --Idx)
-					{
-						const void* SrcElemValue = static_cast<const uint8*>(Data) + (Source->ElementSize * Idx);
-						void* DestElemValue = static_cast<uint8*>(Data) + (Destination->ElementSize * Idx);
-
-						const double Value = SrcFloatProp->GetFloatingPointPropertyValue(SrcElemValue);
-						DestDoubleProp->SetFloatingPointPropertyValue(DestElemValue, Value);
-					}
+					//integral to integral
+					const void* SrcElemValue = static_cast<const uint8*>(Data);
+					void* DestElemValue = static_cast<uint8*>(Data);
+
+					const int64 Value = SourceNumericProperty->GetSignedIntPropertyValue(SrcElemValue);
+					DestinationNumericProperty->SetIntPropertyValue(DestElemValue, Value);
 				}
+#endif
 			}
 		}
 	} // namespace
@@ -475,11 +492,9 @@ namespace UE::MVVM::BindingHelper
 			return false;
 		}
 
-		return Source->ArrayDim == Destination->ArrayDim
-			&& (Destination->SameType(Source)
-				|| Private::IsRealNumberConversion(Source, Destination)
-				|| Private::IsObjectPropertyCompatible(Source, Destination)
-				);
+		return Destination->SameType(Source)
+				|| Private::IsNumericConversionRequired(Source, Destination)
+				|| Private::IsObjectPropertyCompatible(Source, Destination);
 	}
 
 
@@ -554,7 +569,7 @@ namespace UE::MVVM::BindingHelper
 			Source.GetObjectVariant().GetUObject()->ProcessEvent(Source.GetFieldVariant().GetFunction(), DataPtr);
 		}
 
-		Private::ConvertRealNumber(GetterType, SetterType, DataPtr);
+		Private::ConvertNumeric(GetterType, SetterType, DataPtr);
 
 		if (bIsDestinationBindingIsProperty)
 		{
@@ -593,7 +608,7 @@ namespace UE::MVVM::BindingHelper
 			check(GetterType);
 			check(ArePropertiesCompatible(GetterType, ArgumentConversionProperty));
 
-			if (Private::IsRealNumberConversion(GetterType, ArgumentConversionProperty))
+			if (Private::IsNumericConversionRequired(GetterType, ArgumentConversionProperty))
 			{
 				// we need to do a copy because we may destroy the ReturnConversionProperty (imagine the return type is a TArray)
 				const int32 AllocationSize = FMath::Max(GetterType->GetSize(), ArgumentConversionProperty->GetSize());
@@ -612,7 +627,7 @@ namespace UE::MVVM::BindingHelper
 					Source.GetObjectVariant().GetUObject()->ProcessEvent(Source.GetFieldVariant().GetFunction(), GetterDataPtr);
 				}
 
-				Private::ConvertRealNumber(GetterType, ArgumentConversionProperty, GetterDataPtr);
+				Private::ConvertNumeric(GetterType, ArgumentConversionProperty, GetterDataPtr);
 				ArgumentConversionProperty->CopyCompleteValue(ArgumentConversionProperty->ContainerPtrToValuePtr<void>(ConversionFunctionDataPtr), GetterDataPtr);
 
 				GetterType->DestroyValue(GetterDataPtr);
@@ -646,7 +661,7 @@ namespace UE::MVVM::BindingHelper
 			check(SetterType);
 			check(ArePropertiesCompatible(ReturnConversionProperty, SetterType));
 
-			if (Private::IsRealNumberConversion(ReturnConversionProperty, SetterType))
+			if (Private::IsNumericConversionRequired(ReturnConversionProperty, SetterType))
 			{
 				// we need to do a copy because we may destroy the ReturnConversionProperty (imagine the return type is a TArray)
 				const int32 AllocationSize = FMath::Max(SetterType->GetSize(), ReturnConversionProperty->GetSize());
@@ -655,7 +670,7 @@ namespace UE::MVVM::BindingHelper
 				SetterType->InitializeValue(SetterDataPtr); // probably not needed since they are double/float
 
 				ReturnConversionProperty->CopyCompleteValue(SetterDataPtr, ReturnConversionProperty->ContainerPtrToValuePtr<void>(ConversionFunctionDataPtr));
-				Private::ConvertRealNumber(ReturnConversionProperty, SetterType, SetterDataPtr);
+				Private::ConvertNumeric(ReturnConversionProperty, SetterType, SetterDataPtr);
 
 				if (bIsDestinationBindingIsProperty)
 				{
@@ -674,7 +689,6 @@ namespace UE::MVVM::BindingHelper
 			{
 				// Re use the same buffer, no need to create a new copy
 				void* DestinationDataPtr = ReturnConversionProperty->ContainerPtrToValuePtr<void>(ConversionFunctionDataPtr);
-
 				if (bIsDestinationBindingIsProperty)
 				{
 					SetterType->SetValue_InContainer(Destination.GetObjectVariant().GetData(), DestinationDataPtr);
