MVVM: Prevent static array binding. Add option to auto convert integral property.

diff --git a/Engine/Plugins/Runtime/ModelViewViewModel/Source/ModelViewViewModel/Private/Bindings/MVVMBindingHelper.cpp b/Engine/Plugins/Runtime/ModelViewViewModel/Source/ModelViewViewModel/Private/Bindings/MVVMBindingHelper.cpp
--- a/Engine/Plugins/Runtime/ModelViewViewModel/Source/ModelViewViewModel/Private/Bindings/MVVMBindingHelper.cpp
+++ b/Engine/Plugins/Runtime/ModelViewViewModel/Source/ModelViewViewModel/Private/Bindings/MVVMBindingHelper.cpp
@@ -443,7 +443,8 @@ namespace UE::MVVM::BindingHelper
 				const bool bBothFloatingPoint = SourceNumericProperty->IsFloatingPoint() && DestinationNumericProperty->IsFloatingPoint();
 #if UE_MVVM_ALLOW_AUTO_INTEGRAL_CONVERSION
 				const bool bBothIntegral = SourceNumericProperty->IsInteger() && DestinationNumericProperty->IsInteger();
-				return !bSameType && (bBothFloatingPoint || bBothIntegral);
+				const bool bOneIsEnum = SourceNumericProperty->IsEnum() || DestinationNumericProperty->IsEnum();
+				return !bSameType && (bBothFloatingPoint || (bBothIntegral && !bOneIsEnum));
 #else
 				return !bSameType && bBothFloatingPoint;
 #endif
@@ -451,6 +452,7 @@ namespace UE::MVVM::BindingHelper
 			return false;
 		}
 
+
 		void ConvertNumeric(const FProperty* Source, const FProperty* Destination, void* Data)
 		{
 			check(Source);
