MVVM: Change the description of the MVVM plugin.

diff --git a/Engine/Plugins/Runtime/ModelViewViewModel/ModelViewViewModel.uplugin b/Engine/Plugins/Runtime/ModelViewViewModel/ModelViewViewModel.uplugin
--- a/Engine/Plugins/Runtime/ModelViewViewModel/ModelViewViewModel.uplugin
+++ b/Engine/Plugins/Runtime/ModelViewViewModel/ModelViewViewModel.uplugin
@@ -1,12 +1,12 @@
 {
 	"FileVersion": 3,
 
-	"FriendlyName": "Model View ViewModel",
+	"FriendlyName": "UMG Viewmodel",
 	"Version": 1,
 	"VersionName": "1.0",
 	"CreatedBy": "Epic Games, Inc.",
 	"CreatedByURL": "http://epicgames.com",
-	"Description": "A ViewModel plugin in for a Model View ViewModel pattern.",
+	"Description": "A plugin to support the Viewmodel pattern in UMG.",
 	"Category": "UI",
 	"EnabledByDefault": false,
 	"CanContainContent": false,
