MVVM - You can only add bindings to a widget blueprint that already has viewmodels assigned to it.

It's impossible to create a valid binding without it, so there's not really any reason to allow it.

#rnx

diff --git a/Engine/Plugins/Runtime/ModelViewViewModel/Source/ModelViewViewModelEditor/Private/Widgets/SMVVMViewBindingPanel.cpp b/Engine/Plugins/Runtime/ModelViewViewModel/Source/ModelViewViewModelEditor/Private/Widgets/SMVVMViewBindingPanel.cpp
--- a/Engine/Plugins/Runtime/ModelViewViewModel/Source/ModelViewViewModelEditor/Private/Widgets/SMVVMViewBindingPanel.cpp
+++ b/Engine/Plugins/Runtime/ModelViewViewModel/Source/ModelViewViewModelEditor/Private/Widgets/SMVVMViewBindingPanel.cpp
@@ -164,7 +164,6 @@ void SMVVMViewBindingPanel::OnBindingListSelectionChanged(int32 Index)
 	}
 }
 
-
 void SMVVMViewBindingPanel::AddDefaultBinding()
 {
 	if (UMVVMWidgetBlueprintExtension_View* MVVMExtensionPtr = MVVMExtension.Get())
@@ -177,10 +176,29 @@ void SMVVMViewBindingPanel::AddDefaultBinding()
 	}
 }
 
-
 bool SMVVMViewBindingPanel::CanAddBinding() const
 {
-	return true;
+	if (UMVVMWidgetBlueprintExtension_View* MVVMExtensionPtr = MVVMExtension.Get())
+	{
+		if (UMVVMBlueprintView* BlueprintView = MVVMExtensionPtr->GetBlueprintView())
+		{
+			return BlueprintView->GetViewModels().Num() > 0;
+		}
+	}
+
+	return false;
+}
+
+FText SMVVMViewBindingPanel::GetAddBindingToolTip() const
+{
+	if (CanAddBinding())
+	{
+		return LOCTEXT("AddBindingTooltip", "Add an empty binding.");
+	}
+	else
+	{
+		return LOCTEXT("CannotAddBindingToolTip", "A viewmodel is required before adding bindings.");
+	}
 }
 
 void SMVVMViewBindingPanel::ShowManageViewModelsWindow()
@@ -349,7 +367,7 @@ TSharedRef<SWidget> SMVVMViewBindingPanel::GenerateEditViewWidget()
 			),
 			NAME_None,
 			LOCTEXT("AddBinding", "Add Binding"),
-			LOCTEXT("AddBindingTooltip", "Add an empty binding."),
+			TAttribute<FText>(this, &SMVVMViewBindingPanel::GetAddBindingToolTip),
 			FSlateIcon(FMVVMEditorStyle::Get().GetStyleSetName(), "BindingView.AddBinding"),
 			EUserInterfaceActionType::Button
 		);
