Back out changelist 20791814 to address NonUnity error.

#rnx

diff --git a/Engine/Plugins/Runtime/ModelViewViewModel/Source/ModelViewViewModelEditor/Private/Widgets/SMVVMViewBindingPanel.cpp b/Engine/Plugins/Runtime/ModelViewViewModel/Source/ModelViewViewModelEditor/Private/Widgets/SMVVMViewBindingPanel.cpp
--- a/Engine/Plugins/Runtime/ModelViewViewModel/Source/ModelViewViewModelEditor/Private/Widgets/SMVVMViewBindingPanel.cpp
+++ b/Engine/Plugins/Runtime/ModelViewViewModel/Source/ModelViewViewModelEditor/Private/Widgets/SMVVMViewBindingPanel.cpp
@@ -164,6 +164,7 @@ void SMVVMViewBindingPanel::OnBindingListSelectionChanged(int32 Index)
 	}
 }
 
+
 void SMVVMViewBindingPanel::AddDefaultBinding()
 {
 	if (UMVVMWidgetBlueprintExtension_View* MVVMExtensionPtr = MVVMExtension.Get())
@@ -176,29 +177,10 @@ void SMVVMViewBindingPanel::AddDefaultBinding()
 	}
 }
 
-bool SMVVMViewBindingPanel::CanAddBinding() const
-{
-	if (UMVVMWidgetBlueprintExtension_View* MVVMExtensionPtr = MVVMExtension.Get())
-	{
-		if (UMVVMBlueprintView* BlueprintView = MVVMExtensionPtr->GetBlueprintView())
-		{
-			return BlueprintView->GetViewModels().Num() > 0;
-		}
-	}
-
-	return false;
-}
 
-FText SMVVMViewBindingPanel::GetAddBindingToolTip() const
+bool SMVVMViewBindingPanel::CanAddBinding() const
 {
-	if (CanAddBinding())
-	{
-		return LOCTEXT("AddBindingTooltip", "Add an empty binding.");
-	}
-	else
-	{
-		return LOCTEXT("CannotAddBindingToolTip", "A viewmodel is required before adding bindings.");
-	}
+	return true;
 }
 
 void SMVVMViewBindingPanel::ShowManageViewModelsWindow()
@@ -367,7 +349,7 @@ TSharedRef<SWidget> SMVVMViewBindingPanel::GenerateEditViewWidget()
 			),
 			NAME_None,
 			LOCTEXT("AddBinding", "Add Binding"),
-			TAttribute<FText>(this, &SMVVMViewBindingPanel::GetAddBindingToolTip),
+			LOCTEXT("AddBindingTooltip", "Add an empty binding."),
 			FSlateIcon(FMVVMEditorStyle::Get().GetStyleSetName(), "BindingView.AddBinding"),
 			EUserInterfaceActionType::Button
 		);
