MVVM - Removed IFieldPathHelper. The functionality was moved to the editor subsystem, which made the abstraction unnecessary and unavailable to blueprint. SMVVMFieldSelector, SMVVMSourceSelector and SMVVMPropertyPath all now reference FMVVMBlueprintPropertyPath instead.

FBindingSource no longer contains the IsSelected member. This was very context-specific, and doesn't really make sense in the general case.

#rnx

diff --git a/Engine/Plugins/Runtime/ModelViewViewModel/Source/ModelViewViewModelEditor/Private/MVVMPropertyPathHelpers.cpp b/Engine/Plugins/Runtime/ModelViewViewModel/Source/ModelViewViewModelEditor/Private/MVVMPropertyPathHelpers.cpp
deleted file mode 100644
diff --git a/Engine/Plugins/Runtime/ModelViewViewModel/Source/ModelViewViewModelEditor/Private/MVVMPropertyPathHelpers.h b/Engine/Plugins/Runtime/ModelViewViewModel/Source/ModelViewViewModelEditor/Private/MVVMPropertyPathHelpers.h
deleted file mode 100644
diff --git a/Engine/Plugins/Runtime/ModelViewViewModel/Source/ModelViewViewModelEditor/Private/Types/MVVMBindingSource.h b/Engine/Plugins/Runtime/ModelViewViewModel/Source/ModelViewViewModelEditor/Private/Types/MVVMBindingSource.h
--- a/Engine/Plugins/Runtime/ModelViewViewModel/Source/ModelViewViewModelEditor/Private/Types/MVVMBindingSource.h
+++ b/Engine/Plugins/Runtime/ModelViewViewModel/Source/ModelViewViewModelEditor/Private/Types/MVVMBindingSource.h
@@ -5,19 +5,20 @@
 #include "CoreMinimal.h"
 #include "Framework/Views/TableViewTypeTraits.h"
 
+class UWidgetBlueprint;
+
 namespace UE::MVVM
 {
 	struct FBindingSource
 	{
-		FGuid SourceGuid;
+		FGuid ViewModelId;
 		FName Name;
-		const UClass* Class = nullptr;
+		UClass* Class = nullptr;
 		FText DisplayName;
-		bool IsSelected = false;
 
 		bool operator==(const FBindingSource& Other) const
 		{
-			return SourceGuid == Other.SourceGuid && Name == Other.Name && Class == Other.Class;
+			return ViewModelId == Other.ViewModelId && Name == Other.Name && Class == Other.Class;
 		}
 
 		bool operator!=(const FBindingSource& Other) const
@@ -27,24 +28,27 @@ namespace UE::MVVM
 
 		friend int32 GetTypeHash(const FBindingSource& Source)
 		{
-			uint32 Hash = HashCombine(GetTypeHash(Source.SourceGuid), GetTypeHash(Source.Name));
+			uint32 Hash = HashCombine(GetTypeHash(Source.ViewModelId), GetTypeHash(Source.Name));
 			Hash = HashCombine(Hash, GetTypeHash(Source.Class));
 			return Hash;
 		}
 
 		bool IsValid() const
 		{
-			return SourceGuid.IsValid() || !Name.IsNone();
+			return Class != nullptr && (ViewModelId.IsValid() || !Name.IsNone());
 		}
 
 		void Reset()
 		{
-			SourceGuid = FGuid();
+			ViewModelId = FGuid();
 			Name = FName();
 			Class = nullptr;
 			DisplayName = FText::GetEmpty();
-			IsSelected = false;
 		}
+
+		static FBindingSource CreateForWidget(const UWidgetBlueprint* WidgetBlueprint, FName WidgetName);
+		static FBindingSource CreateForViewModel(const UWidgetBlueprint* WidgetBlueprint, FGuid ViewModelId);
+		static FBindingSource CreateForViewModel(const UWidgetBlueprint* WidgetBlueprint, FName ViewModelName);
 	};
 }
 
diff --git a/Engine/Plugins/Runtime/ModelViewViewModel/Source/ModelViewViewModelEditor/Private/Widgets/SMVVMFieldSelector.cpp b/Engine/Plugins/Runtime/ModelViewViewModel/Source/ModelViewViewModelEditor/Private/Widgets/SMVVMFieldSelector.cpp
--- a/Engine/Plugins/Runtime/ModelViewViewModel/Source/ModelViewViewModelEditor/Private/Widgets/SMVVMFieldSelector.cpp
+++ b/Engine/Plugins/Runtime/ModelViewViewModel/Source/ModelViewViewModelEditor/Private/Widgets/SMVVMFieldSelector.cpp
@@ -27,61 +27,104 @@ namespace UE::MVVM::Private
 		}
 		return LOCTEXT("None", "<None>");
 	}
+}
 
-	FMVVMConstFieldVariant GetSelectedFieldFromHelpers(const TArray<IFieldPathHelper*>& PathHelpers)
-	{
-		bool bFirst = true;
-		FMVVMConstFieldVariant Selected;
 
-		for (const IFieldPathHelper* Helper : PathHelpers)
-		{
-			const FMVVMConstFieldVariant ThisSelection = Helper->GetSelectedField();
-			if (bFirst)
-			{
-				Selected = ThisSelection;
-				bFirst = false;
-			}
-			else if (Selected != ThisSelection)
-			{
-				Selected.Reset();
-				break;
-			}
-		}
+void SMVVMFieldEntry::Construct(const FArguments& InArgs)
+{
+	Field = InArgs._Field;
+	OnValidateField = InArgs._OnValidateField;
 
-		return Selected;
-	}
+	ChildSlot
+		[
+			SNew(SHorizontalBox)
+			+ SHorizontalBox::Slot()
+		.HAlign(HAlign_Left)
+		.VAlign(VAlign_Center)
+		.AutoWidth()
+		[
+			SAssignNew(Icon, SMVVMFieldIcon)
+		]
+	+ SHorizontalBox::Slot()
+		.HAlign(HAlign_Fill)
+		.VAlign(VAlign_Center)
+		.Padding(4, 0, 0, 0)
+		[
+			SAssignNew(Label, STextBlock)
+			.TextStyle(InArgs._TextStyle)
+		]
+		];
+
+	Refresh();
 }
 
-TOptional<FBindingSource> SMVVMFieldSelector::GetSelectedSource() const
+void SMVVMFieldEntry::Refresh()
 {
-	TOptional<FBindingSource> Source;
-	for (const IFieldPathHelper* Helper : PathHelpers)
+	FText ToolTipText = FText::GetEmpty();
+	bool bEnabled = true;
+	if (OnValidateField.IsBound())
 	{
-		if (!Source.IsSet())
+		TValueOrError<bool, FString> Result = OnValidateField.Execute(Field);
+		if (Result.HasError())
 		{
-			Source = Helper->GetSelectedSource();
+			ToolTipText = FText::FromString(Result.GetError());
+			bEnabled = false;
 		}
-		else if (Source != Helper->GetSelectedSource())
+		else
 		{
-			Source.Reset();
-			break;
+			bEnabled = true;
 		}
 	}
 
-	return Source;
+	SetEnabled(bEnabled);
+
+	UE::MVVM::FMVVMConstFieldVariant Variant;
+
+	TArray<UE::MVVM::FMVVMConstFieldVariant> Fields = Field.GetFields();
+	if (Fields.Num() > 0)
+	{
+		Variant = Fields.Last();
+	}
+
+	if (ToolTipText.IsEmpty())
+	{
+		ToolTipText = Variant.IsFunction() ? Variant.GetFunction()->GetToolTipText() :
+			Variant.IsProperty() ? Variant.GetProperty()->GetToolTipText() :
+			FText::GetEmpty();
+	}
+
+	SetToolTipText(ToolTipText);
+
+	Icon->RefreshBinding(Variant);
+	Label->SetText(Private::GetFieldDisplayName(Variant));
+}
+
+void SMVVMFieldEntry::SetField(const FMVVMBlueprintPropertyPath& InField)
+{
+	Field = InField;
+
+	Refresh();
 }
 
 void SMVVMFieldSelector::Construct(const FArguments& InArgs)
 {
-	PathHelpers = InArgs._PathHelpers;
-	CounterpartHelpers = InArgs._CounterpartHelpers;
 	OnValidateFieldDelegate = InArgs._OnValidateField;
 	OnSelectionChangedDelegate = InArgs._OnSelectionChanged;
+	SelectedField = InArgs._SelectedField;
+	check(SelectedField.IsSet());
+
+	SelectedSource = InArgs._SelectedSource;
+	check(SelectedSource.IsSet());
+
+	AvailableFields = InArgs._AvailableFields;
+	check(AvailableFields.IsSet());
+
+	CachedAvailableFields = AvailableFields.Get(TArray<FMVVMBlueprintPropertyPath>());
+
 	BindingMode = InArgs._BindingMode;
 	bIsSource = InArgs._IsSource;
 	TextStyle = InArgs._TextStyle;
 
-	FMVVMConstFieldVariant SelectedField = Private::GetSelectedFieldFromHelpers(PathHelpers);
 	Refresh();
 
 	ChildSlot
@@ -101,24 +144,24 @@ void SMVVMFieldSelector::Construct(const FArguments& InArgs)
 				.TextStyle(FAppStyle::Get(), "HintText")
 				.Visibility_Lambda([this]()
 				{
-					return SelectedSource.IsSet() ? EVisibility::Collapsed : EVisibility::Visible;
+					return SelectedSource.Get(UE::MVVM::FBindingSource()).IsValid() ? EVisibility::Collapsed : EVisibility::Visible;
 				})
 			]
 			+ SOverlay::Slot()
 			[
-				SAssignNew(FieldComboBox, SComboBox<FMVVMConstFieldVariant>)
+				SAssignNew(FieldComboBox, SComboBox<FMVVMBlueprintPropertyPath>)
 				.Visibility_Lambda([this]()
 				{
-					return SelectedSource.IsSet() ? EVisibility::Visible : EVisibility::Hidden;
+					return SelectedSource.Get(UE::MVVM::FBindingSource()).IsValid() ? EVisibility::Visible : EVisibility::Hidden;
 				})
-				.OptionsSource(&AvailableFields)
-				.InitiallySelectedItem(SelectedField)
+				.OptionsSource(&CachedAvailableFields)
+				.InitiallySelectedItem(SelectedField.Get(FMVVMBlueprintPropertyPath()))
 				.OnGenerateWidget(this, &SMVVMFieldSelector::OnGenerateFieldWidget)
 				.OnSelectionChanged(this, &SMVVMFieldSelector::OnComboBoxSelectionChanged)
 				[
-					SAssignNew(SelectedEntry, SMVVMFieldEntry)
+					SAssignNew(SelectedEntryWidget, SMVVMFieldEntry)
 					.TextStyle(TextStyle)
-					.Field(SelectedField)
+					.Field(SelectedField.Get(FMVVMBlueprintPropertyPath()))
 					.OnValidateField(this, &SMVVMFieldSelector::ValidateField)
 				]
 			]
@@ -137,205 +180,48 @@ void SMVVMFieldSelector::Construct(const FArguments& InArgs)
 	];
 }
 
-void SMVVMFieldSelector::OnComboBoxSelectionChanged(FMVVMConstFieldVariant Selected, ESelectInfo::Type SelectionType)
+void SMVVMFieldSelector::OnComboBoxSelectionChanged(FMVVMBlueprintPropertyPath Selected, ESelectInfo::Type SelectionType)
 {
-	SelectedEntry->SetField(Selected);
+	SelectedEntryWidget->SetField(Selected);
 
 	OnSelectionChangedDelegate.ExecuteIfBound(Selected);
 }
 
 void SMVVMFieldSelector::Refresh()
 {
-	AvailableFields.Reset();
-
-	SelectedSource = GetSelectedSource();
-
-	TSet<FMVVMConstFieldVariant> CompatibleFields;
-	TSet<FMVVMConstFieldVariant> IncompatibleFields;
-
-	FMVVMConstFieldVariant SelectedField = Private::GetSelectedFieldFromHelpers(PathHelpers);
-	FMVVMConstFieldVariant CounterpartSelectedField = Private::GetSelectedFieldFromHelpers(CounterpartHelpers);
+	CachedAvailableFields = AvailableFields.Get(TArray<FMVVMBlueprintPropertyPath>());
 
-	bool bFirst = true;
-
-	for (IFieldPathHelper* Helper : PathHelpers)
-	{
-		TSet<FMVVMConstFieldVariant> AllFields;
-		Helper->GetAvailableFields(AllFields);
-
-		if (bFirst)
-		{
-			bFirst = false;
-
-			for (const FMVVMConstFieldVariant& Field : AllFields)
-			{
-				TValueOrError<bool, FString> Result = IsValidBindingForField(Field, CounterpartSelectedField);
-				if (Result.HasValue() && Result.GetValue() == true)
-				{
-					CompatibleFields.Add(Field);
-				}
-				else
-				{
-					IncompatibleFields.Add(Field);
-				}
-			}
-		}
-		else
-		{
-			CompatibleFields.Intersect(AllFields);
-			IncompatibleFields.Intersect(AllFields);
-		}
-	}
-
-	AvailableFields.Append(CompatibleFields.Array());
-
-	// put all incompatible properties at the end so they don't just disappear from the list without explanation
-	AvailableFields.Append(IncompatibleFields.Array());
-	
 	if (FieldComboBox.IsValid())
 	{
 		FieldComboBox->RefreshOptions();
-		FieldComboBox->SetSelectedItem(SelectedField);
+		FieldComboBox->SetSelectedItem(SelectedField.Get(FMVVMBlueprintPropertyPath()));
 	}
 }
 
-TValueOrError<bool, FString> SMVVMFieldSelector::IsValidBindingForField(const FMVVMConstFieldVariant& Field, const FMVVMConstFieldVariant& CounterpartField) const
+TValueOrError<bool, FString> SMVVMFieldSelector::ValidateField(FMVVMBlueprintPropertyPath Field) const
 {
-	// if we have a delegate bound, use that instead
 	if (OnValidateFieldDelegate.IsBound())
 	{
 		return OnValidateFieldDelegate.Execute(Field);
 	}
 
-	if (CounterpartField.IsEmpty() || Field.IsEmpty())
-	{
-		return MakeValue(true);
-	}
-
-	if (!BindingMode.IsSet())
-	{
-		return MakeValue(true);
-	}
-
-	EMVVMBindingMode Mode = BindingMode.Get();
-
-	bool bToDestination = UE::MVVM::IsForwardBinding(Mode);
-	if (!bIsSource)
-	{
-		bToDestination = !bToDestination;
-	}
-
-	UMVVMSubsystem::FConstDirectionalBindingArgs Args;
-	Args.SourceBinding = bToDestination ? Field : CounterpartField;
-	Args.DestinationBinding = bToDestination ? CounterpartField : Field;
-
-	UMVVMSubsystem* Subsystem = GEngine->GetEngineSubsystem<UMVVMSubsystem>();
-
-	TValueOrError<bool, FString> Result = Subsystem->IsBindingValid(Args);
-
-	// check in the other direction as well if this is a two-way binding
-	if (Mode == EMVVMBindingMode::TwoWay &&
-		Result.HasValue() && Result.GetValue() == true)
-	{
-		Args.SourceBinding = bToDestination ? CounterpartField : Field;
-		Args.DestinationBinding = bToDestination ? Field : CounterpartField;
-
-		return Subsystem->IsBindingValid(Args);
-	}
-
-	return Result;
+	return MakeValue(true);
 }
 
-TValueOrError<bool, FString> SMVVMFieldSelector::ValidateField(FMVVMConstFieldVariant Field) const
-{
-	return IsValidBindingForField(Field, Private::GetSelectedFieldFromHelpers(CounterpartHelpers));
-}
-
-TSharedRef<SWidget> SMVVMFieldSelector::OnGenerateFieldWidget(FMVVMConstFieldVariant Field) const
+TSharedRef<SWidget> SMVVMFieldSelector::OnGenerateFieldWidget(FMVVMBlueprintPropertyPath Path) const
 {
 	return SNew(SMVVMFieldEntry)
 		.TextStyle(TextStyle)
-		.Field(Field)
+		.Field(Path)
 		.OnValidateField(this, &SMVVMFieldSelector::ValidateField);
 }
 
-void SMVVMFieldEntry::Construct(const FArguments& InArgs)
-{
-	Field = InArgs._Field;
-	OnValidateField = InArgs._OnValidateField;
-
-	ChildSlot
-	[
-		SNew(SHorizontalBox)
-		+ SHorizontalBox::Slot()
-		.HAlign(HAlign_Left)
-		.VAlign(VAlign_Center)
-		.AutoWidth()
-		[
-			SAssignNew(Icon, SMVVMFieldIcon)
-		]
-		+ SHorizontalBox::Slot()
-		.HAlign(HAlign_Fill)
-		.VAlign(VAlign_Center)
-		.Padding(4, 0, 0, 0)
-		[
-			SAssignNew(Label, STextBlock)
-			.TextStyle(InArgs._TextStyle)
-		]
-	];
-
-	Refresh();
-}
-
-void SMVVMFieldEntry::Refresh()
-{
-	FText ToolTipText = FText::GetEmpty();
-	bool bEnabled = true;
-	if (OnValidateField.IsBound())
-	{
-		TValueOrError<bool, FString> Result = OnValidateField.Execute(Field);
-		if (Result.HasError())
-		{
-			ToolTipText = FText::FromString(Result.GetError());
-			bEnabled = false;
-		}
-		else
-		{
-			bEnabled = true;
-		}
-	}
-
-	SetEnabled(bEnabled);
-
-	if (ToolTipText.IsEmpty())
-	{
-		ToolTipText = Field.IsFunction() ? Field.GetFunction()->GetToolTipText() :
-			Field.IsProperty() ? Field.GetProperty()->GetToolTipText() :
-			FText::GetEmpty();
-	}
-
-	SetToolTipText(ToolTipText);
-
-	Icon->RefreshBinding(Field);
-	Label->SetText(Private::GetFieldDisplayName(Field));
-}
-
-void SMVVMFieldEntry::SetField(const UE::MVVM::FMVVMConstFieldVariant& InField)
-{
-	Field = InField;
-
-	Refresh();
-}
-
 EVisibility SMVVMFieldSelector::GetClearVisibility() const
 {
-	for (const IFieldPathHelper* Helper : PathHelpers)
+	FMVVMBlueprintPropertyPath SelectedPath = SelectedField.Get(FMVVMBlueprintPropertyPath());
+	if (!SelectedPath.IsEmpty())
 	{
-		const FMVVMConstFieldVariant ThisSelection = Helper->GetSelectedField();
-		if (!ThisSelection.IsEmpty())
-		{
-			return EVisibility::Visible;
-		}
+		return EVisibility::Visible;
 	}
 
 	return EVisibility::Collapsed;
@@ -348,6 +234,11 @@ FReply SMVVMFieldSelector::OnClearBinding()
 		FieldComboBox->ClearSelection();
 	}
 
+	if (SelectedEntryWidget.IsValid())
+	{
+		SelectedEntryWidget->Refresh();
+	}
+
 	return FReply::Handled();
 }
 
diff --git a/Engine/Plugins/Runtime/ModelViewViewModel/Source/ModelViewViewModelEditor/Private/Widgets/SMVVMFieldSelector.h b/Engine/Plugins/Runtime/ModelViewViewModel/Source/ModelViewViewModelEditor/Private/Widgets/SMVVMFieldSelector.h
--- a/Engine/Plugins/Runtime/ModelViewViewModel/Source/ModelViewViewModelEditor/Private/Widgets/SMVVMFieldSelector.h
+++ b/Engine/Plugins/Runtime/ModelViewViewModel/Source/ModelViewViewModelEditor/Private/Widgets/SMVVMFieldSelector.h
@@ -3,10 +3,10 @@
 #pragma once
 
 #include "MVVMFieldVariantListTraits.h"
-#include "MVVMPropertyPathHelpers.h" 
 #include "Styling/CoreStyle.h"
 #include "Templates/ValueOrError.h"
 #include "Types/MVVMBindingMode.h"
+#include "Types/MVVMBindingSource.h"
 #include "Types/MVVMFieldVariant.h"
 #include "UObject/UnrealType.h"
 #include "Widgets/DeclarativeSyntaxSupport.h"
@@ -22,7 +22,7 @@ class SMVVMFieldIcon;
 class STextBlock; 
 
 using FIsFieldValidResult = TValueOrError<bool, FString>;
-DECLARE_DELEGATE_RetVal_OneParam(FIsFieldValidResult, FIsFieldValid, UE::MVVM::FMVVMConstFieldVariant);
+DECLARE_DELEGATE_RetVal_OneParam(FIsFieldValidResult, FIsFieldValid, FMVVMBlueprintPropertyPath);
 
 class SMVVMFieldEntry : public SCompoundWidget
 {
@@ -32,16 +32,16 @@ public:
 		_TextStyle( &FCoreStyle::Get().GetWidgetStyle<FTextBlockStyle>( "NormalText" ) )
 		{}
 		SLATE_STYLE_ARGUMENT(FTextBlockStyle, TextStyle)
-		SLATE_ARGUMENT(UE::MVVM::FMVVMConstFieldVariant, Field)
+		SLATE_ARGUMENT(FMVVMBlueprintPropertyPath, Field)
 		SLATE_EVENT(FIsFieldValid, OnValidateField)
 	SLATE_END_ARGS()
 
 	void Construct(const FArguments& InArgs);
 	void Refresh(); 
-	void SetField(const UE::MVVM::FMVVMConstFieldVariant& InField);
+	void SetField(const FMVVMBlueprintPropertyPath& InField);
 
 private:
-	UE::MVVM::FMVVMConstFieldVariant Field;
+	FMVVMBlueprintPropertyPath Field;
 	FIsFieldValid OnValidateField;
 	TSharedPtr<SMVVMFieldIcon> Icon;
 	TSharedPtr<STextBlock> Label;
@@ -50,7 +50,7 @@ private:
 class SMVVMFieldSelector : public SCompoundWidget
 {
 public:
-	DECLARE_DELEGATE_OneParam(FSelectionChanged, UE::MVVM::FMVVMConstFieldVariant);
+	DECLARE_DELEGATE_OneParam(FSelectionChanged, FMVVMBlueprintPropertyPath);
 
 	SLATE_BEGIN_ARGS(SMVVMFieldSelector) :
 		_TextStyle(&FCoreStyle::Get().GetWidgetStyle<FTextBlockStyle>("NormalText")),
@@ -58,11 +58,17 @@ public:
 		{
 		}
 		SLATE_STYLE_ARGUMENT(FTextBlockStyle, TextStyle)
-		SLATE_ARGUMENT(TArray<UE::MVVM::IFieldPathHelper*>, PathHelpers)
-		SLATE_ARGUMENT(TArray<UE::MVVM::IFieldPathHelper*>, CounterpartHelpers)
 		SLATE_ATTRIBUTE(EMVVMBindingMode, BindingMode)
 		SLATE_ARGUMENT(bool, IsSource)
+		SLATE_ATTRIBUTE(UE::MVVM::FBindingSource, SelectedSource)
+		SLATE_ATTRIBUTE(FMVVMBlueprintPropertyPath, SelectedField)
+		SLATE_ATTRIBUTE(TArray<FMVVMBlueprintPropertyPath>, AvailableFields)
 		SLATE_EVENT(FSelectionChanged, OnSelectionChanged)
+
+		/**
+		  * Would the given field be a valid entry for this combo box? 
+		  * The error string returned will be used as a tooltip.
+		  */
 		SLATE_EVENT(FIsFieldValid, OnValidateField)
 	SLATE_END_ARGS()
 
@@ -70,29 +76,28 @@ public:
 	void Refresh();
 
 private:
-	TSharedRef<SWidget> OnGenerateFieldWidget(UE::MVVM::FMVVMConstFieldVariant FieldPath) const;
-	void OnComboBoxSelectionChanged(UE::MVVM::FMVVMConstFieldVariant Selected, ESelectInfo::Type SelectionType);
-	TOptional<UE::MVVM::FBindingSource> GetSelectedSource() const;
+	TSharedRef<SWidget> OnGenerateFieldWidget(FMVVMBlueprintPropertyPath Path) const;
+	void OnComboBoxSelectionChanged(FMVVMBlueprintPropertyPath Selected, ESelectInfo::Type SelectionType);
 
-	TValueOrError<bool, FString> IsValidBindingForField(const UE::MVVM::FMVVMConstFieldVariant& Field, const UE::MVVM::FMVVMConstFieldVariant& CounterpartField) const;
-	TValueOrError<bool, FString> ValidateField(UE::MVVM::FMVVMConstFieldVariant Field) const;
+	TValueOrError<bool, FString> ValidateField(FMVVMBlueprintPropertyPath Field) const;
 
 	EVisibility GetClearVisibility() const;
 	FReply OnClearBinding();
 
 private:
-	TArray<UE::MVVM::IFieldPathHelper*> PathHelpers;
-	TArray<UE::MVVM::IFieldPathHelper*> CounterpartHelpers;
 	TAttribute<EMVVMBindingMode> BindingMode;
-	TSharedPtr<SMVVMFieldEntry> SelectedEntry;
-	TOptional<UE::MVVM::FBindingSource> SelectedSource;
+	TSharedPtr<SMVVMFieldEntry> SelectedEntryWidget;
+	TAttribute<UE::MVVM::FBindingSource> SelectedSource;
+	TAttribute<FMVVMBlueprintPropertyPath> SelectedField;
+	TAttribute<TArray<FMVVMBlueprintPropertyPath>> AvailableFields;
+
 	const FTextBlockStyle* TextStyle = nullptr;
 	bool bIsSource = false;
 
-	TSharedPtr<SComboBox<UE::MVVM::FMVVMConstFieldVariant>> FieldComboBox;
+	TSharedPtr<SComboBox<FMVVMBlueprintPropertyPath>> FieldComboBox;
 
 	FSelectionChanged OnSelectionChangedDelegate;
 	FIsFieldValid OnValidateFieldDelegate;
 
-	TArray<UE::MVVM::FMVVMConstFieldVariant> AvailableFields;
+	TArray<FMVVMBlueprintPropertyPath> CachedAvailableFields;
 }; 
diff --git a/Engine/Plugins/Runtime/ModelViewViewModel/Source/ModelViewViewModelEditor/Private/Widgets/SMVVMPropertyPath.cpp b/Engine/Plugins/Runtime/ModelViewViewModel/Source/ModelViewViewModelEditor/Private/Widgets/SMVVMPropertyPath.cpp
--- a/Engine/Plugins/Runtime/ModelViewViewModel/Source/ModelViewViewModelEditor/Private/Widgets/SMVVMPropertyPath.cpp
+++ b/Engine/Plugins/Runtime/ModelViewViewModel/Source/ModelViewViewModelEditor/Private/Widgets/SMVVMPropertyPath.cpp
@@ -2,10 +2,13 @@
 
 #include "Widgets/SMVVMPropertyPath.h"
 
+#include "Blueprint/WidgetTree.h"
+#include "Editor.h"
 #include "MVVMBlueprintView.h"
+#include "MVVMEditorSubsystem.h"
 #include "MVVMWidgetBlueprintExtension_View.h"
-#include "Widgets/SMVVMFieldIcon.h"
 #include "Styling/MVVMEditorStyle.h"
+#include "Widgets/SMVVMFieldIcon.h"
 #include "Widgets/SBoxPanel.h"
 #include "Widgets/Images/SImage.h"
 #include "Widgets/Layout/SBorder.h"
@@ -14,37 +17,12 @@
 
 #define LOCTEXT_NAMESPACE "SMVVMPropertyPath"
 
-void SMVVMPropertyPathBase::Construct(const FArguments& InArgs, const UWidgetBlueprint* InWidgetBlueprint)
+void SMVVMPropertyPath::Construct(const FArguments& InArgs)
 {
-	WidgetBlueprint = InWidgetBlueprint;
-
-	const UE::MVVM::IFieldPathHelper& PathHelper = GetPathHelper();
-	SelectedSource = PathHelper.GetSelectedSource();
-	SelectedField = PathHelper.GetSelectedField();
-
-	auto GetContextText = [this]()
-	{
-		if (SelectedSource.IsSet())
-		{
-			return SelectedSource.GetValue().DisplayName;
-		}
-		return FText::GetEmpty();
-	};
-
-	auto GetPathText = [this]()
-	{
-		if (SelectedField.IsProperty())
-		{
-			return SelectedField.GetProperty()->GetDisplayNameText();
-		}
-
-		if (SelectedField.IsFunction())
-		{
-			return SelectedField.GetFunction()->GetDisplayNameText();
-		}
-
-		return FText::GetEmpty();
-	};
+	WidgetBlueprint = InArgs._WidgetBlueprint;
+	check(InArgs._WidgetBlueprint != nullptr);
+	PropertyPath = InArgs._PropertyPath;
+	check(PropertyPath);
 
 	ChildSlot
 	[
@@ -55,7 +33,7 @@ void SMVVMPropertyPathBase::Construct(const FArguments& InArgs, const UWidgetBlu
 		[
 			SNew(STextBlock)
 			.TextStyle(FMVVMEditorStyle::Get(), "PropertyPath.ContextText")
-			.Text_Lambda(GetContextText)
+			.Text(this, &SMVVMPropertyPath::GetSourceDisplayName)
 		]
 		+ SHorizontalBox::Slot()
 		.Padding(FMargin(5, 0, 0, 0))
@@ -63,7 +41,7 @@ void SMVVMPropertyPathBase::Construct(const FArguments& InArgs, const UWidgetBlu
 		.AutoWidth()
 		[
 			SNew(SMVVMFieldIcon)
-			.Field(SelectedField)
+			.Field(GetLastField())
 		]
 		+ SHorizontalBox::Slot()
 		.Padding(FMargin(2, 0, 5, 0))
@@ -72,24 +50,38 @@ void SMVVMPropertyPathBase::Construct(const FArguments& InArgs, const UWidgetBlu
 		[
 			SNew(STextBlock)
 			.TextStyle(FMVVMEditorStyle::Get(), "PropertyPath.ContextText")
-			.Text_Lambda(GetPathText)
+			.Text(this, &SMVVMPropertyPath::GetFieldDisplayName)
 		]
 	];
 }
 
-void SMVVMWidgetPropertyPath::Construct(const FArguments& InArgs, const UWidgetBlueprint* InWidgetBlueprint)
+FText SMVVMPropertyPath::GetSourceDisplayName() const
 {
-	PathHelper = UE::MVVM::FWidgetFieldPathHelper(InArgs._WidgetPath, InWidgetBlueprint);
-
-	SMVVMPropertyPathBase::Construct(SMVVMPropertyPathBase::FArguments(), InWidgetBlueprint);
+	if (PropertyPath->IsFromWidget())
+	{
+		const UWidget* Widget = WidgetBlueprint->WidgetTree->FindWidget(PropertyPath->GetWidgetName());
+		return FText::FromString(Widget->GetDisplayLabel());
+	}
+	else if (PropertyPath->IsFromViewModel())
+	{
+		const UMVVMEditorSubsystem* Subsystem = GEditor->GetEditorSubsystem<UMVVMEditorSubsystem>();
+		if (const UMVVMBlueprintView* View = Subsystem->GetView(WidgetBlueprint.Get()))
+		{
+			const FMVVMBlueprintViewModelContext* ViewModel = View->FindViewModel(PropertyPath->GetViewModelId());
+			return ViewModel->GetDisplayName();
+		}
+	}
+	return FText::GetEmpty();
 }
 
-void SMVVMViewModelPropertyPath::Construct(const FArguments& InArgs, const UWidgetBlueprint* InWidgetBlueprint)
+FText SMVVMPropertyPath::GetFieldDisplayName() const
 {
-	PathHelper = UE::MVVM::FViewModelFieldPathHelper(InArgs._ViewModelPath, InWidgetBlueprint);
-
-	SMVVMPropertyPathBase::Construct(SMVVMPropertyPathBase::FArguments(), InWidgetBlueprint);
+	return FText::FromString(PropertyPath->GetBasePropertyPath());
 }
 
+UE::MVVM::FMVVMConstFieldVariant SMVVMPropertyPath::GetLastField() const
+{
+	return PropertyPath->GetFields().Last();
+}
 
 #undef LOCTEXT_NAMESPACE
\ No newline at end of file
diff --git a/Engine/Plugins/Runtime/ModelViewViewModel/Source/ModelViewViewModelEditor/Private/Widgets/SMVVMPropertyPath.h b/Engine/Plugins/Runtime/ModelViewViewModel/Source/ModelViewViewModelEditor/Private/Widgets/SMVVMPropertyPath.h
--- a/Engine/Plugins/Runtime/ModelViewViewModel/Source/ModelViewViewModelEditor/Private/Widgets/SMVVMPropertyPath.h
+++ b/Engine/Plugins/Runtime/ModelViewViewModel/Source/ModelViewViewModelEditor/Private/Widgets/SMVVMPropertyPath.h
@@ -5,62 +5,25 @@
 #include "Widgets/SCompoundWidget.h"
 #include "MVVMBlueprintView.h"
 #include "Widgets/SMVVMFieldIcon.h"
-#include "MVVMPropertyPathHelpers.h"
 #include "Types/MVVMFieldVariant.h"
 
-class SMVVMPropertyPathBase : public SCompoundWidget
+class SMVVMPropertyPath : public SCompoundWidget
 {
 public:
-	SLATE_BEGIN_ARGS(SMVVMPropertyPathBase)
+	SLATE_BEGIN_ARGS(SMVVMPropertyPath)
 	{}
+		SLATE_ARGUMENT(FMVVMBlueprintPropertyPath*, PropertyPath)
+		SLATE_ARGUMENT(const UWidgetBlueprint*, WidgetBlueprint)
 	SLATE_END_ARGS()
 
-	void Construct(const FArguments& InArgs, const UWidgetBlueprint* WidgetBlueprint);
-	virtual void Tick(const FGeometry& AllottedGeometry, const double InCurrentTime, const float InDeltaTime) override;
+	void Construct(const FArguments& InArgs);
 
 private:
-	virtual UE::MVVM::IFieldPathHelper& GetPathHelper() = 0;
+	FText GetSourceDisplayName() const;
+	FText GetFieldDisplayName() const;
+	UE::MVVM::FMVVMConstFieldVariant GetLastField() const;
 
 private:
-	const UWidgetBlueprint* WidgetBlueprint = nullptr;
-	UE::MVVM::FMVVMConstFieldVariant SelectedField;
-	TOptional<UE::MVVM::FBindingSource> SelectedSource;
-};
-
-class SMVVMWidgetPropertyPath : public SMVVMPropertyPathBase
-{
-public:
-	SLATE_BEGIN_ARGS(SMVVMWidgetPropertyPath) {}
-		SLATE_ARGUMENT(FMVVMBlueprintPropertyPath*, WidgetPath)
-	SLATE_END_ARGS()
-
-	void Construct(const FArguments& InArgs, const UWidgetBlueprint* WidgetBlueprint);
-
-private:
-	virtual UE::MVVM::IFieldPathHelper& GetPathHelper() override
-	{ 
-		return PathHelper;
-	}
-
-private:
-	UE::MVVM::FWidgetFieldPathHelper PathHelper;
-};
-
-class SMVVMViewModelPropertyPath : public SMVVMPropertyPathBase
-{
-public:
-	SLATE_BEGIN_ARGS(SMVVMViewModelPropertyPath) {}
-		SLATE_ARGUMENT(FMVVMBlueprintPropertyPath*, ViewModelPath)
-	SLATE_END_ARGS()
-
-	void Construct(const FArguments& InArgs, const UWidgetBlueprint* WidgetBlueprint);
-
-private:
-	virtual UE::MVVM::IFieldPathHelper& GetPathHelper() override
-	{
-		return PathHelper;
-	}
-
-private:
-	UE::MVVM::FViewModelFieldPathHelper PathHelper;
+	FMVVMBlueprintPropertyPath* PropertyPath = nullptr;
+	TWeakObjectPtr<const UWidgetBlueprint> WidgetBlueprint = nullptr;
 };
diff --git a/Engine/Plugins/Runtime/ModelViewViewModel/Source/ModelViewViewModelEditor/Private/Widgets/SMVVMSourceSelector.cpp b/Engine/Plugins/Runtime/ModelViewViewModel/Source/ModelViewViewModelEditor/Private/Widgets/SMVVMSourceSelector.cpp
--- a/Engine/Plugins/Runtime/ModelViewViewModel/Source/ModelViewViewModelEditor/Private/Widgets/SMVVMSourceSelector.cpp
+++ b/Engine/Plugins/Runtime/ModelViewViewModel/Source/ModelViewViewModelEditor/Private/Widgets/SMVVMSourceSelector.cpp
@@ -11,33 +11,13 @@
 
 using namespace UE::MVVM;
 
-static FBindingSource GetSelectedSource(const TArray<IFieldPathHelper*>& Helpers)
-{
-	bool bFirst = true;
-	TOptional<FBindingSource> Result;
-
-	for (const IFieldPathHelper* Helper : Helpers)
-	{
-		TOptional<FBindingSource> ThisSelection = Helper->GetSelectedSource();
-		if (bFirst)
-		{
-			Result = ThisSelection;
-		}
-		else if (Result != ThisSelection)
-		{
-			// all don't share the same value
-			Result.Reset();
-			break;
-		}
-	}
-
-	return Result.Get(FBindingSource());
-}
-
 void SMVVMSourceSelector::Construct(const FArguments& Args)
 {
-	PathHelpers = Args._PathHelpers;
 	TextStyle = Args._TextStyle;
+	AvailableSourcesAttribute = Args._AvailableSources;
+	check(AvailableSourcesAttribute.IsSet());
+	SelectedSourceAttribute = Args._SelectedSource;
+	check(SelectedSourceAttribute.IsSet());
 
 	Refresh();
 
@@ -91,36 +71,8 @@ void SMVVMSourceSelector::OnComboBoxSelectionChanged(FBindingSource Selected, ES
 
 void SMVVMSourceSelector::Refresh()
 {
-	AvailableSources.Reset();
-	SelectedSource = FBindingSource();
-
-	TArray<IFieldPathHelper*> Helpers = PathHelpers.Get(TArray<IFieldPathHelper*>());
-
-	bool bFirst = true;
-
-	TSet<FBindingSource> SourceIntersection;
-	TOptional<FBindingSource> NewSelection;
-
-	for (const IFieldPathHelper* Helper : Helpers)
-	{
-		// get available sources
-		TSet<FBindingSource> Sources;
-		Helper->GetAvailableSources(Sources);
-
-		if (bFirst)
-		{
-			bFirst = false; 
-
-			SourceIntersection.Append(Sources);
-		}
-		else
-		{
-			SourceIntersection.Intersect(Sources);
-		}
-	}
-
-	AvailableSources = SourceIntersection.Array();
-	SelectedSource = GetSelectedSource(Helpers);
+	SelectedSource = SelectedSourceAttribute.Get();
+	AvailableSources = AvailableSourcesAttribute.Get();
 
 	if (SourceComboBox.IsValid())
 	{
@@ -131,16 +83,7 @@ void SMVVMSourceSelector::Refresh()
 
 EVisibility SMVVMSourceSelector::GetClearVisibility() const
 {
-	for (const IFieldPathHelper* Helper : PathHelpers.Get())
-	{
-		const FBindingSource ThisSelection = Helper->GetSelectedSource();
-		if (ThisSelection.IsValid())
-		{
-			return EVisibility::Visible;
-		}
-	}
-
-	return EVisibility::Collapsed;
+	return SelectedSource.IsValid() ? EVisibility::Visible : EVisibility::Collapsed;
 }
 
 FReply SMVVMSourceSelector::OnClearSource()
diff --git a/Engine/Plugins/Runtime/ModelViewViewModel/Source/ModelViewViewModelEditor/Private/Widgets/SMVVMSourceSelector.h b/Engine/Plugins/Runtime/ModelViewViewModel/Source/ModelViewViewModelEditor/Private/Widgets/SMVVMSourceSelector.h
--- a/Engine/Plugins/Runtime/ModelViewViewModel/Source/ModelViewViewModelEditor/Private/Widgets/SMVVMSourceSelector.h
+++ b/Engine/Plugins/Runtime/ModelViewViewModel/Source/ModelViewViewModelEditor/Private/Widgets/SMVVMSourceSelector.h
@@ -3,7 +3,6 @@
 #pragma once
 
 #include "CoreMinimal.h"
-#include "MVVMPropertyPathHelpers.h"
 #include "Types/MVVMBindingSource.h"
 #include "Widgets/DeclarativeSyntaxSupport.h"
 #include "Widgets/SCompoundWidget.h"
@@ -21,7 +20,8 @@ public:
 		{
 		}
 		SLATE_STYLE_ARGUMENT(FTextBlockStyle, TextStyle)
-		SLATE_ATTRIBUTE(TArray<UE::MVVM::IFieldPathHelper*>, PathHelpers)
+		SLATE_ATTRIBUTE(UE::MVVM::FBindingSource, SelectedSource)
+		SLATE_ATTRIBUTE(TArray<UE::MVVM::FBindingSource>, AvailableSources)
 		SLATE_EVENT(FSelectionChanged, OnSelectionChanged)
 	SLATE_END_ARGS()
 
@@ -35,8 +35,9 @@ private:
 	FReply OnClearSource();
 
 private:
-	TAttribute<TArray<UE::MVVM::IFieldPathHelper*>> PathHelpers;
 	FSelectionChanged OnSelectionChanged;
+	TAttribute<TArray<UE::MVVM::FBindingSource>> AvailableSourcesAttribute;
+	TAttribute< UE::MVVM::FBindingSource> SelectedSourceAttribute;
 	TSharedPtr<SComboBox<UE::MVVM::FBindingSource>> SourceComboBox;
 	TArray<UE::MVVM::FBindingSource> AvailableSources;
 	UE::MVVM::FBindingSource SelectedSource;
