Backing out changelist 20830654 due to missing files.

#rnx

diff --git a/Engine/Plugins/Runtime/ModelViewViewModel/Source/ModelViewViewModelEditor/Private/MVVMPropertyPathHelpers.cpp b/Engine/Plugins/Runtime/ModelViewViewModel/Source/ModelViewViewModelEditor/Private/MVVMPropertyPathHelpers.cpp
new file mode 100644
--- /dev/null
+++ b/Engine/Plugins/Runtime/ModelViewViewModel/Source/ModelViewViewModelEditor/Private/MVVMPropertyPathHelpers.cpp
@@ -0,0 +1,263 @@
+// Copyright Epic Games, Inc. All Rights Reserved.
+
+#include "MVVMPropertyPathHelpers.h"
+
+#include "Algo/Transform.h"
+#include "Bindings/MVVMBindingHelper.h"
+#include "Blueprint/WidgetTree.h"
+#include "Components/Widget.h"
+#include "Engine/Engine.h"
+#include "MVVMBlueprintView.h"
+#include "MVVMSubsystem.h"
+#include "MVVMWidgetBlueprintExtension_View.h"
+
+namespace UE::MVVM
+{
+	void FWidgetFieldPathHelper::GetAvailableSources(TSet<FBindingSource>& OutSources) const
+	{
+		const UWidgetTree* WidgetTree = WidgetBlueprint->WidgetTree;
+		if (WidgetTree == nullptr)
+		{
+			return;
+		}
+
+		TArray<UWidget*> AllWidgets;
+		WidgetTree->GetAllWidgets(AllWidgets);
+
+		OutSources.Reserve(AllWidgets.Num() + 1);
+
+		// Add current widget as a possible binding source
+		if (UClass* BPClass = WidgetBlueprint->GeneratedClass)
+		{
+			TArray<FMVVMAvailableBinding> Bindings = GEngine->GetEngineSubsystem<UMVVMSubsystem>()->GetWidgetAvailableBindings(BPClass);
+
+			if (Bindings.Num() > 0)
+			{
+				// at least one valid property, add it to our list
+				FBindingSource Source;
+				Source.Name = WidgetBlueprint->GetFName();
+				Source.DisplayName = FText::FromName(WidgetBlueprint->GetFName());
+				Source.Class = BPClass;
+				Source.IsSelected = Source.Name == Path->GetWidgetName();
+				OutSources.Add(Source);
+			}
+		}
+
+		for (const UWidget* Widget : AllWidgets)
+		{
+			TArray<FMVVMAvailableBinding> Bindings = GEngine->GetEngineSubsystem<UMVVMSubsystem>()->GetWidgetAvailableBindings(Widget->GetClass());
+
+			if (Bindings.Num() > 0)
+			{
+				// at least one valid property, add it to our list
+				FBindingSource Source;
+				Source.Name = Widget->GetFName();
+				Source.DisplayName = Widget->GetLabelText();
+				Source.Class = Widget->GetClass();
+				Source.IsSelected = Source.Name == Path->GetWidgetName();
+				OutSources.Add(Source);
+			}
+		}
+	}
+
+	FBindingSource FWidgetFieldPathHelper::GetSelectedSource() const
+	{
+		TSet<FBindingSource> Sources;
+		GetAvailableSources(Sources);
+
+		for (const FBindingSource& Source : Sources)
+		{
+			if (Source.IsSelected)
+			{
+				return Source;
+			}
+		}
+
+		return FBindingSource();
+	}
+
+	void FWidgetFieldPathHelper::GetAvailableFields(TSet<FMVVMConstFieldVariant>& OutFields) const
+	{
+		const UWidgetTree* WidgetTree = WidgetBlueprint->WidgetTree;
+		if (WidgetTree == nullptr)
+		{
+			return;
+		}
+
+		TSubclassOf<UWidget> WidgetClass = nullptr;
+		if (const UWidget* Widget = WidgetTree->FindWidget(Path->GetWidgetName()))
+		{
+			WidgetClass = Widget->GetClass();
+		}
+		else if (WidgetBlueprint->GetFName() == Path->GetWidgetName())
+		{
+			WidgetClass = WidgetBlueprint->GeneratedClass;
+		}
+
+		if (!WidgetClass)
+		{
+			return;
+		}
+
+		TArray<FMVVMAvailableBinding> Bindings = GEngine->GetEngineSubsystem<UMVVMSubsystem>()->GetWidgetAvailableBindings(WidgetClass);
+		Algo::Transform(Bindings, OutFields, [WidgetClass](const FMVVMAvailableBinding& Binding)
+			{
+				return UE::MVVM::BindingHelper::FindFieldByName(WidgetClass, Binding.GetBindingName());
+			});
+	}
+
+	FMVVMConstFieldVariant FWidgetFieldPathHelper::GetSelectedField() const
+	{
+		FMVVMBindingName BindingName = GetBindingName();
+
+		TSet<FMVVMConstFieldVariant> Fields;
+		GetAvailableFields(Fields);
+
+		for (const FMVVMConstFieldVariant& Field : Fields)
+		{
+			if (Field.GetName() == BindingName.ToName())
+			{
+				return Field;
+			}
+		}
+
+		return FMVVMConstFieldVariant();
+	}
+
+	void FWidgetFieldPathHelper::SetSelectedSource(const FBindingSource& Source) const
+	{
+		if (Source.IsValid())
+		{
+			Path->SetWidgetName(Source.Name);
+		}
+		else
+		{
+			Path->SetWidgetName(FName());
+		}
+	}
+
+	FMVVMBindingName FWidgetFieldPathHelper::GetBindingName() const
+	{
+		TArray<FName> Paths = Path->GetPaths();
+		if (Paths.Num() > 0)
+		{
+			return FMVVMBindingName(Paths[0]);
+		}
+
+		return FMVVMBindingName();
+	}
+
+	void FWidgetFieldPathHelper::SetBindingReference(const UE::MVVM::FMVVMConstFieldVariant& InField) const
+	{
+		Path->SetBasePropertyPath(InField);
+	}
+
+	void FWidgetFieldPathHelper::ResetBinding() const
+	{
+		Path->ResetBasePropertyPath();
+	}
+
+	void FViewModelFieldPathHelper::GetAvailableSources(TSet<FBindingSource>& OutSources) const
+	{
+		const UMVVMWidgetBlueprintExtension_View* MVVMExtensionPtr = UMVVMWidgetBlueprintExtension_View::GetExtension<UMVVMWidgetBlueprintExtension_View>(WidgetBlueprint);
+		const UMVVMBlueprintView* MVVMBlueprintView = MVVMExtensionPtr->GetBlueprintView();
+		const TArrayView<const FMVVMBlueprintViewModelContext>& ViewModels = MVVMBlueprintView->GetViewModels();
+
+		OutSources.Reserve(ViewModels.Num());
+
+		for (const FMVVMBlueprintViewModelContext& ViewModel : ViewModels)
+		{
+			FBindingSource Source;
+			Source.SourceGuid = ViewModel.GetViewModelId();
+			Source.DisplayName = ViewModel.GetDisplayName();
+			Source.Class = ViewModel.GetViewModelClass().Get();
+			Source.IsSelected = (Source.SourceGuid == Path->GetViewModelId());
+
+			OutSources.Add(Source);
+		}
+	}
+
+	FBindingSource FViewModelFieldPathHelper::GetSelectedSource() const
+	{
+		TSet<FBindingSource> Sources;
+		GetAvailableSources(Sources);
+
+		for (const FBindingSource& Source : Sources)
+		{
+			if (Source.IsSelected)
+			{
+				return Source;
+			}
+		}
+
+		return FBindingSource();
+	}
+
+	void FViewModelFieldPathHelper::GetAvailableFields(TSet<FMVVMConstFieldVariant>& OutFields) const
+	{
+		const UMVVMWidgetBlueprintExtension_View* MVVMExtensionPtr = UMVVMWidgetBlueprintExtension_View::GetExtension<UMVVMWidgetBlueprintExtension_View>(WidgetBlueprint);
+		const UMVVMBlueprintView* MVVMBlueprintView = MVVMExtensionPtr->GetBlueprintView();
+
+		const FMVVMBlueprintViewModelContext* ViewModelContext = MVVMBlueprintView->FindViewModel(Path->GetViewModelId());
+		if (ViewModelContext == nullptr)
+		{
+			return;
+		}
+
+		TArray<FMVVMAvailableBinding> Bindings = GEngine->GetEngineSubsystem<UMVVMSubsystem>()->GetViewModelAvailableBindings(ViewModelContext->GetViewModelClass());
+		Algo::Transform(Bindings, OutFields, [ViewModelContext](const FMVVMAvailableBinding& Binding)
+			{
+				return UE::MVVM::BindingHelper::FindFieldByName(ViewModelContext->GetViewModelClass(), Binding.GetBindingName());
+			});
+	}
+
+	FMVVMConstFieldVariant FViewModelFieldPathHelper::GetSelectedField() const
+	{
+		FMVVMBindingName BindingName = GetBindingName();
+
+		TSet<FMVVMConstFieldVariant> Fields;
+		GetAvailableFields(Fields);
+
+		for (const FMVVMConstFieldVariant& Field : Fields)
+		{
+			if (Field.GetName() == BindingName.ToName())
+			{
+				return Field;
+			}
+		}
+
+		return FMVVMConstFieldVariant();
+	}
+
+	void FViewModelFieldPathHelper::SetSelectedSource(const FBindingSource& Source) const
+	{
+		if (Source.IsValid())
+		{
+			Path->SetViewModelId(Source.SourceGuid);
+		}
+		else
+		{
+			Path->SetViewModelId(FGuid());
+		}
+	}
+
+	FMVVMBindingName FViewModelFieldPathHelper::GetBindingName() const
+	{
+		TArray<FName> Paths = Path->GetPaths();
+		if (Paths.Num() > 0)
+		{
+			return FMVVMBindingName(Paths[0]);
+		}
+		return FMVVMBindingName();
+	}
+
+	void FViewModelFieldPathHelper::SetBindingReference(const UE::MVVM::FMVVMConstFieldVariant& InField) const
+	{
+		Path->SetBasePropertyPath(InField);
+	}
+
+	void FViewModelFieldPathHelper::ResetBinding() const
+	{
+		Path->ResetBasePropertyPath();
+	}
+}
diff --git a/Engine/Plugins/Runtime/ModelViewViewModel/Source/ModelViewViewModelEditor/Private/MVVMPropertyPathHelpers.h b/Engine/Plugins/Runtime/ModelViewViewModel/Source/ModelViewViewModelEditor/Private/MVVMPropertyPathHelpers.h
new file mode 100644
--- /dev/null
+++ b/Engine/Plugins/Runtime/ModelViewViewModel/Source/ModelViewViewModelEditor/Private/MVVMPropertyPathHelpers.h
@@ -0,0 +1,91 @@
+// Copyright Epic Games, Inc. All Rights Reserved.
+
+#pragma once
+
+#include "Internationalization/Text.h"
+#include "Misc/Attribute.h"
+#include "Misc/Guid.h"
+#include "MVVMPropertyPath.h"
+#include "Types/MVVMBindingName.h"
+#include "Types/MVVMBindingSource.h"
+#include "Types/MVVMFieldVariant.h"
+#include "UObject/UnrealType.h"
+#include "WidgetBlueprint.h"
+
+namespace UE::MVVM
+{
+	class IFieldPathHelper
+	{
+	public:
+		virtual ~IFieldPathHelper() {}
+
+		virtual void GetAvailableFields(TSet<FMVVMConstFieldVariant>& Fields) const = 0;
+		virtual FMVVMConstFieldVariant GetSelectedField() const = 0;
+
+		virtual void GetAvailableSources(TSet<FBindingSource>& OutSources) const = 0;
+		virtual FBindingSource GetSelectedSource() const = 0;
+		virtual void SetSelectedSource(const FBindingSource& Source) const = 0;
+
+		virtual FMVVMBindingName GetBindingName() const = 0;
+		virtual void SetBindingReference(const FMVVMConstFieldVariant& InField) const = 0;
+		virtual void ResetBinding() const = 0;
+	};
+
+	class FWidgetFieldPathHelper : public IFieldPathHelper
+	{
+	public:
+		FWidgetFieldPathHelper() {}
+		virtual ~FWidgetFieldPathHelper() {}
+		FWidgetFieldPathHelper(FMVVMBlueprintPropertyPath* WidgetPath, const UWidgetBlueprint* InWidgetBlueprint) :
+			Path(WidgetPath),
+			WidgetBlueprint(InWidgetBlueprint)
+		{
+			check(Path);
+			check(Path->IsFromWidget() || Path->IsEmpty());
+		}
+
+		virtual void GetAvailableFields(TSet<FMVVMConstFieldVariant>& Fields) const override;
+		virtual FMVVMConstFieldVariant GetSelectedField() const override;
+
+		virtual void GetAvailableSources(TSet<FBindingSource>& OutSources) const override;
+		virtual FBindingSource GetSelectedSource() const override;
+		virtual void SetSelectedSource(const FBindingSource& Source) const override;
+
+		virtual FMVVMBindingName GetBindingName() const override;
+		virtual void SetBindingReference(const FMVVMConstFieldVariant& InField) const override;
+		virtual void ResetBinding() const override;
+
+	private:
+		FMVVMBlueprintPropertyPath* Path;
+		const UWidgetBlueprint* WidgetBlueprint = nullptr;
+	};
+
+	class FViewModelFieldPathHelper : public IFieldPathHelper
+	{
+	public:
+		FViewModelFieldPathHelper() {}
+		virtual ~FViewModelFieldPathHelper() {}
+		FViewModelFieldPathHelper(FMVVMBlueprintPropertyPath* ViewModelPath, const UWidgetBlueprint* InWidgetBlueprint) :
+			Path(ViewModelPath),
+			WidgetBlueprint(InWidgetBlueprint)
+		{
+			check(Path);
+			check(Path->IsFromViewModel() || Path->IsEmpty());
+		}
+
+		virtual void GetAvailableFields(TSet<FMVVMConstFieldVariant>& Fields) const override;
+		virtual FMVVMConstFieldVariant GetSelectedField() const override;
+
+		virtual void GetAvailableSources(TSet<FBindingSource>& OutSources) const override;
+		virtual FBindingSource GetSelectedSource() const override;
+		virtual void SetSelectedSource(const FBindingSource& Source) const override;
+
+		virtual FMVVMBindingName GetBindingName() const override;
+		virtual void SetBindingReference(const FMVVMConstFieldVariant& InField) const override;
+		virtual void ResetBinding() const override;
+
+	private:
+		FMVVMBlueprintPropertyPath* Path;
+		const UWidgetBlueprint* WidgetBlueprint = nullptr;
+	};
+}
diff --git a/Engine/Plugins/Runtime/ModelViewViewModel/Source/ModelViewViewModelEditor/Private/Types/MVVMBindingSource.h b/Engine/Plugins/Runtime/ModelViewViewModel/Source/ModelViewViewModelEditor/Private/Types/MVVMBindingSource.h
--- a/Engine/Plugins/Runtime/ModelViewViewModel/Source/ModelViewViewModelEditor/Private/Types/MVVMBindingSource.h
+++ b/Engine/Plugins/Runtime/ModelViewViewModel/Source/ModelViewViewModelEditor/Private/Types/MVVMBindingSource.h
@@ -5,20 +5,19 @@
 #include "CoreMinimal.h"
 #include "Framework/Views/TableViewTypeTraits.h"
 
-class UWidgetBlueprint;
-
 namespace UE::MVVM
 {
 	struct FBindingSource
 	{
-		FGuid ViewModelId;
+		FGuid SourceGuid;
 		FName Name;
-		UClass* Class = nullptr;
+		const UClass* Class = nullptr;
 		FText DisplayName;
+		bool IsSelected = false;
 
 		bool operator==(const FBindingSource& Other) const
 		{
-			return ViewModelId == Other.ViewModelId && Name == Other.Name && Class == Other.Class;
+			return SourceGuid == Other.SourceGuid && Name == Other.Name && Class == Other.Class;
 		}
 
 		bool operator!=(const FBindingSource& Other) const
@@ -28,27 +27,24 @@ namespace UE::MVVM
 
 		friend int32 GetTypeHash(const FBindingSource& Source)
 		{
-			uint32 Hash = HashCombine(GetTypeHash(Source.ViewModelId), GetTypeHash(Source.Name));
+			uint32 Hash = HashCombine(GetTypeHash(Source.SourceGuid), GetTypeHash(Source.Name));
 			Hash = HashCombine(Hash, GetTypeHash(Source.Class));
 			return Hash;
 		}
 
 		bool IsValid() const
 		{
-			return Class != nullptr && (ViewModelId.IsValid() || !Name.IsNone());
+			return SourceGuid.IsValid() || !Name.IsNone();
 		}
 
 		void Reset()
 		{
-			ViewModelId = FGuid();
+			SourceGuid = FGuid();
 			Name = FName();
 			Class = nullptr;
 			DisplayName = FText::GetEmpty();
+			IsSelected = false;
 		}
-
-		static FBindingSource CreateForWidget(const UWidgetBlueprint* WidgetBlueprint, FName WidgetName);
-		static FBindingSource CreateForViewModel(const UWidgetBlueprint* WidgetBlueprint, FGuid ViewModelId);
-		static FBindingSource CreateForViewModel(const UWidgetBlueprint* WidgetBlueprint, FName ViewModelName);
 	};
 }
 
diff --git a/Engine/Plugins/Runtime/ModelViewViewModel/Source/ModelViewViewModelEditor/Private/Widgets/SMVVMFieldSelector.cpp b/Engine/Plugins/Runtime/ModelViewViewModel/Source/ModelViewViewModelEditor/Private/Widgets/SMVVMFieldSelector.cpp
--- a/Engine/Plugins/Runtime/ModelViewViewModel/Source/ModelViewViewModelEditor/Private/Widgets/SMVVMFieldSelector.cpp
+++ b/Engine/Plugins/Runtime/ModelViewViewModel/Source/ModelViewViewModelEditor/Private/Widgets/SMVVMFieldSelector.cpp
@@ -27,104 +27,61 @@ namespace UE::MVVM::Private
 		}
 		return LOCTEXT("None", "<None>");
 	}
-}
-
 
-void SMVVMFieldEntry::Construct(const FArguments& InArgs)
-{
-	Field = InArgs._Field;
-	OnValidateField = InArgs._OnValidateField;
+	FMVVMConstFieldVariant GetSelectedFieldFromHelpers(const TArray<IFieldPathHelper*>& PathHelpers)
+	{
+		bool bFirst = true;
+		FMVVMConstFieldVariant Selected;
 
-	ChildSlot
-		[
-			SNew(SHorizontalBox)
-			+ SHorizontalBox::Slot()
-		.HAlign(HAlign_Left)
-		.VAlign(VAlign_Center)
-		.AutoWidth()
-		[
-			SAssignNew(Icon, SMVVMFieldIcon)
-		]
-	+ SHorizontalBox::Slot()
-		.HAlign(HAlign_Fill)
-		.VAlign(VAlign_Center)
-		.Padding(4, 0, 0, 0)
-		[
-			SAssignNew(Label, STextBlock)
-			.TextStyle(InArgs._TextStyle)
-		]
-		];
+		for (const IFieldPathHelper* Helper : PathHelpers)
+		{
+			const FMVVMConstFieldVariant ThisSelection = Helper->GetSelectedField();
+			if (bFirst)
+			{
+				Selected = ThisSelection;
+				bFirst = false;
+			}
+			else if (Selected != ThisSelection)
+			{
+				Selected.Reset();
+				break;
+			}
+		}
 
-	Refresh();
+		return Selected;
+	}
 }
 
-void SMVVMFieldEntry::Refresh()
+TOptional<FBindingSource> SMVVMFieldSelector::GetSelectedSource() const
 {
-	FText ToolTipText = FText::GetEmpty();
-	bool bEnabled = true;
-	if (OnValidateField.IsBound())
+	TOptional<FBindingSource> Source;
+	for (const IFieldPathHelper* Helper : PathHelpers)
 	{
-		TValueOrError<bool, FString> Result = OnValidateField.Execute(Field);
-		if (Result.HasError())
+		if (!Source.IsSet())
 		{
-			ToolTipText = FText::FromString(Result.GetError());
-			bEnabled = false;
+			Source = Helper->GetSelectedSource();
 		}
-		else
+		else if (Source != Helper->GetSelectedSource())
 		{
-			bEnabled = true;
+			Source.Reset();
+			break;
 		}
 	}
 
-	SetEnabled(bEnabled);
-
-	UE::MVVM::FMVVMConstFieldVariant Variant;
-
-	TArray<UE::MVVM::FMVVMConstFieldVariant> Fields = Field.GetFields();
-	if (Fields.Num() > 0)
-	{
-		Variant = Fields.Last();
-	}
-
-	if (ToolTipText.IsEmpty())
-	{
-		ToolTipText = Variant.IsFunction() ? Variant.GetFunction()->GetToolTipText() :
-			Variant.IsProperty() ? Variant.GetProperty()->GetToolTipText() :
-			FText::GetEmpty();
-	}
-
-	SetToolTipText(ToolTipText);
-
-	Icon->RefreshBinding(Variant);
-	Label->SetText(Private::GetFieldDisplayName(Variant));
-}
-
-void SMVVMFieldEntry::SetField(const FMVVMBlueprintPropertyPath& InField)
-{
-	Field = InField;
-
-	Refresh();
+	return Source;
 }
 
 void SMVVMFieldSelector::Construct(const FArguments& InArgs)
 {
+	PathHelpers = InArgs._PathHelpers;
+	CounterpartHelpers = InArgs._CounterpartHelpers;
 	OnValidateFieldDelegate = InArgs._OnValidateField;
 	OnSelectionChangedDelegate = InArgs._OnSelectionChanged;
-	SelectedField = InArgs._SelectedField;
-	check(SelectedField.IsSet());
-
-	SelectedSource = InArgs._SelectedSource;
-	check(SelectedSource.IsSet());
-
-	AvailableFields = InArgs._AvailableFields;
-	check(AvailableFields.IsSet());
-
-	CachedAvailableFields = AvailableFields.Get(TArray<FMVVMBlueprintPropertyPath>());
-
 	BindingMode = InArgs._BindingMode;
 	bIsSource = InArgs._IsSource;
 	TextStyle = InArgs._TextStyle;
 
+	FMVVMConstFieldVariant SelectedField = Private::GetSelectedFieldFromHelpers(PathHelpers);
 	Refresh();
 
 	ChildSlot
@@ -144,24 +101,24 @@ void SMVVMFieldSelector::Construct(const FArguments& InArgs)
 				.TextStyle(FAppStyle::Get(), "HintText")
 				.Visibility_Lambda([this]()
 				{
-					return SelectedSource.Get(UE::MVVM::FBindingSource()).IsValid() ? EVisibility::Collapsed : EVisibility::Visible;
+					return SelectedSource.IsSet() ? EVisibility::Collapsed : EVisibility::Visible;
 				})
 			]
 			+ SOverlay::Slot()
 			[
-				SAssignNew(FieldComboBox, SComboBox<FMVVMBlueprintPropertyPath>)
+				SAssignNew(FieldComboBox, SComboBox<FMVVMConstFieldVariant>)
 				.Visibility_Lambda([this]()
 				{
-					return SelectedSource.Get(UE::MVVM::FBindingSource()).IsValid() ? EVisibility::Visible : EVisibility::Hidden;
+					return SelectedSource.IsSet() ? EVisibility::Visible : EVisibility::Hidden;
 				})
-				.OptionsSource(&CachedAvailableFields)
-				.InitiallySelectedItem(SelectedField.Get(FMVVMBlueprintPropertyPath()))
+				.OptionsSource(&AvailableFields)
+				.InitiallySelectedItem(SelectedField)
 				.OnGenerateWidget(this, &SMVVMFieldSelector::OnGenerateFieldWidget)
 				.OnSelectionChanged(this, &SMVVMFieldSelector::OnComboBoxSelectionChanged)
 				[
-					SAssignNew(SelectedEntryWidget, SMVVMFieldEntry)
+					SAssignNew(SelectedEntry, SMVVMFieldEntry)
 					.TextStyle(TextStyle)
-					.Field(SelectedField.Get(FMVVMBlueprintPropertyPath()))
+					.Field(SelectedField)
 					.OnValidateField(this, &SMVVMFieldSelector::ValidateField)
 				]
 			]
@@ -180,48 +137,205 @@ void SMVVMFieldSelector::Construct(const FArguments& InArgs)
 	];
 }
 
-void SMVVMFieldSelector::OnComboBoxSelectionChanged(FMVVMBlueprintPropertyPath Selected, ESelectInfo::Type SelectionType)
+void SMVVMFieldSelector::OnComboBoxSelectionChanged(FMVVMConstFieldVariant Selected, ESelectInfo::Type SelectionType)
 {
-	SelectedEntryWidget->SetField(Selected);
+	SelectedEntry->SetField(Selected);
 
 	OnSelectionChangedDelegate.ExecuteIfBound(Selected);
 }
 
 void SMVVMFieldSelector::Refresh()
 {
-	CachedAvailableFields = AvailableFields.Get(TArray<FMVVMBlueprintPropertyPath>());
+	AvailableFields.Reset();
+
+	SelectedSource = GetSelectedSource();
+
+	TSet<FMVVMConstFieldVariant> CompatibleFields;
+	TSet<FMVVMConstFieldVariant> IncompatibleFields;
+
+	FMVVMConstFieldVariant SelectedField = Private::GetSelectedFieldFromHelpers(PathHelpers);
+	FMVVMConstFieldVariant CounterpartSelectedField = Private::GetSelectedFieldFromHelpers(CounterpartHelpers);
 
+	bool bFirst = true;
+
+	for (IFieldPathHelper* Helper : PathHelpers)
+	{
+		TSet<FMVVMConstFieldVariant> AllFields;
+		Helper->GetAvailableFields(AllFields);
+
+		if (bFirst)
+		{
+			bFirst = false;
+
+			for (const FMVVMConstFieldVariant& Field : AllFields)
+			{
+				TValueOrError<bool, FString> Result = IsValidBindingForField(Field, CounterpartSelectedField);
+				if (Result.HasValue() && Result.GetValue() == true)
+				{
+					CompatibleFields.Add(Field);
+				}
+				else
+				{
+					IncompatibleFields.Add(Field);
+				}
+			}
+		}
+		else
+		{
+			CompatibleFields.Intersect(AllFields);
+			IncompatibleFields.Intersect(AllFields);
+		}
+	}
+
+	AvailableFields.Append(CompatibleFields.Array());
+
+	// put all incompatible properties at the end so they don't just disappear from the list without explanation
+	AvailableFields.Append(IncompatibleFields.Array());
+	
 	if (FieldComboBox.IsValid())
 	{
 		FieldComboBox->RefreshOptions();
-		FieldComboBox->SetSelectedItem(SelectedField.Get(FMVVMBlueprintPropertyPath()));
+		FieldComboBox->SetSelectedItem(SelectedField);
 	}
 }
 
-TValueOrError<bool, FString> SMVVMFieldSelector::ValidateField(FMVVMBlueprintPropertyPath Field) const
+TValueOrError<bool, FString> SMVVMFieldSelector::IsValidBindingForField(const FMVVMConstFieldVariant& Field, const FMVVMConstFieldVariant& CounterpartField) const
 {
+	// if we have a delegate bound, use that instead
 	if (OnValidateFieldDelegate.IsBound())
 	{
 		return OnValidateFieldDelegate.Execute(Field);
 	}
 
-	return MakeValue(true);
+	if (CounterpartField.IsEmpty() || Field.IsEmpty())
+	{
+		return MakeValue(true);
+	}
+
+	if (!BindingMode.IsSet())
+	{
+		return MakeValue(true);
+	}
+
+	EMVVMBindingMode Mode = BindingMode.Get();
+
+	bool bToDestination = UE::MVVM::IsForwardBinding(Mode);
+	if (!bIsSource)
+	{
+		bToDestination = !bToDestination;
+	}
+
+	UMVVMSubsystem::FConstDirectionalBindingArgs Args;
+	Args.SourceBinding = bToDestination ? Field : CounterpartField;
+	Args.DestinationBinding = bToDestination ? CounterpartField : Field;
+
+	UMVVMSubsystem* Subsystem = GEngine->GetEngineSubsystem<UMVVMSubsystem>();
+
+	TValueOrError<bool, FString> Result = Subsystem->IsBindingValid(Args);
+
+	// check in the other direction as well if this is a two-way binding
+	if (Mode == EMVVMBindingMode::TwoWay &&
+		Result.HasValue() && Result.GetValue() == true)
+	{
+		Args.SourceBinding = bToDestination ? CounterpartField : Field;
+		Args.DestinationBinding = bToDestination ? Field : CounterpartField;
+
+		return Subsystem->IsBindingValid(Args);
+	}
+
+	return Result;
 }
 
-TSharedRef<SWidget> SMVVMFieldSelector::OnGenerateFieldWidget(FMVVMBlueprintPropertyPath Path) const
+TValueOrError<bool, FString> SMVVMFieldSelector::ValidateField(FMVVMConstFieldVariant Field) const
+{
+	return IsValidBindingForField(Field, Private::GetSelectedFieldFromHelpers(CounterpartHelpers));
+}
+
+TSharedRef<SWidget> SMVVMFieldSelector::OnGenerateFieldWidget(FMVVMConstFieldVariant Field) const
 {
 	return SNew(SMVVMFieldEntry)
 		.TextStyle(TextStyle)
-		.Field(Path)
+		.Field(Field)
 		.OnValidateField(this, &SMVVMFieldSelector::ValidateField);
 }
 
+void SMVVMFieldEntry::Construct(const FArguments& InArgs)
+{
+	Field = InArgs._Field;
+	OnValidateField = InArgs._OnValidateField;
+
+	ChildSlot
+	[
+		SNew(SHorizontalBox)
+		+ SHorizontalBox::Slot()
+		.HAlign(HAlign_Left)
+		.VAlign(VAlign_Center)
+		.AutoWidth()
+		[
+			SAssignNew(Icon, SMVVMFieldIcon)
+		]
+		+ SHorizontalBox::Slot()
+		.HAlign(HAlign_Fill)
+		.VAlign(VAlign_Center)
+		.Padding(4, 0, 0, 0)
+		[
+			SAssignNew(Label, STextBlock)
+			.TextStyle(InArgs._TextStyle)
+		]
+	];
+
+	Refresh();
+}
+
+void SMVVMFieldEntry::Refresh()
+{
+	FText ToolTipText = FText::GetEmpty();
+	bool bEnabled = true;
+	if (OnValidateField.IsBound())
+	{
+		TValueOrError<bool, FString> Result = OnValidateField.Execute(Field);
+		if (Result.HasError())
+		{
+			ToolTipText = FText::FromString(Result.GetError());
+			bEnabled = false;
+		}
+		else
+		{
+			bEnabled = true;
+		}
+	}
+
+	SetEnabled(bEnabled);
+
+	if (ToolTipText.IsEmpty())
+	{
+		ToolTipText = Field.IsFunction() ? Field.GetFunction()->GetToolTipText() :
+			Field.IsProperty() ? Field.GetProperty()->GetToolTipText() :
+			FText::GetEmpty();
+	}
+
+	SetToolTipText(ToolTipText);
+
+	Icon->RefreshBinding(Field);
+	Label->SetText(Private::GetFieldDisplayName(Field));
+}
+
+void SMVVMFieldEntry::SetField(const UE::MVVM::FMVVMConstFieldVariant& InField)
+{
+	Field = InField;
+
+	Refresh();
+}
+
 EVisibility SMVVMFieldSelector::GetClearVisibility() const
 {
-	FMVVMBlueprintPropertyPath SelectedPath = SelectedField.Get(FMVVMBlueprintPropertyPath());
-	if (!SelectedPath.IsEmpty())
+	for (const IFieldPathHelper* Helper : PathHelpers)
 	{
-		return EVisibility::Visible;
+		const FMVVMConstFieldVariant ThisSelection = Helper->GetSelectedField();
+		if (!ThisSelection.IsEmpty())
+		{
+			return EVisibility::Visible;
+		}
 	}
 
 	return EVisibility::Collapsed;
@@ -234,11 +348,6 @@ FReply SMVVMFieldSelector::OnClearBinding()
 		FieldComboBox->ClearSelection();
 	}
 
-	if (SelectedEntryWidget.IsValid())
-	{
-		SelectedEntryWidget->Refresh();
-	}
-
 	return FReply::Handled();
 }
 
diff --git a/Engine/Plugins/Runtime/ModelViewViewModel/Source/ModelViewViewModelEditor/Private/Widgets/SMVVMFieldSelector.h b/Engine/Plugins/Runtime/ModelViewViewModel/Source/ModelViewViewModelEditor/Private/Widgets/SMVVMFieldSelector.h
--- a/Engine/Plugins/Runtime/ModelViewViewModel/Source/ModelViewViewModelEditor/Private/Widgets/SMVVMFieldSelector.h
+++ b/Engine/Plugins/Runtime/ModelViewViewModel/Source/ModelViewViewModelEditor/Private/Widgets/SMVVMFieldSelector.h
@@ -3,10 +3,10 @@
 #pragma once
 
 #include "MVVMFieldVariantListTraits.h"
+#include "MVVMPropertyPathHelpers.h" 
 #include "Styling/CoreStyle.h"
 #include "Templates/ValueOrError.h"
 #include "Types/MVVMBindingMode.h"
-#include "Types/MVVMBindingSource.h"
 #include "Types/MVVMFieldVariant.h"
 #include "UObject/UnrealType.h"
 #include "Widgets/DeclarativeSyntaxSupport.h"
@@ -22,7 +22,7 @@ class SMVVMFieldIcon;
 class STextBlock; 
 
 using FIsFieldValidResult = TValueOrError<bool, FString>;
-DECLARE_DELEGATE_RetVal_OneParam(FIsFieldValidResult, FIsFieldValid, FMVVMBlueprintPropertyPath);
+DECLARE_DELEGATE_RetVal_OneParam(FIsFieldValidResult, FIsFieldValid, UE::MVVM::FMVVMConstFieldVariant);
 
 class SMVVMFieldEntry : public SCompoundWidget
 {
@@ -32,16 +32,16 @@ public:
 		_TextStyle( &FCoreStyle::Get().GetWidgetStyle<FTextBlockStyle>( "NormalText" ) )
 		{}
 		SLATE_STYLE_ARGUMENT(FTextBlockStyle, TextStyle)
-		SLATE_ARGUMENT(FMVVMBlueprintPropertyPath, Field)
+		SLATE_ARGUMENT(UE::MVVM::FMVVMConstFieldVariant, Field)
 		SLATE_EVENT(FIsFieldValid, OnValidateField)
 	SLATE_END_ARGS()
 
 	void Construct(const FArguments& InArgs);
 	void Refresh(); 
-	void SetField(const FMVVMBlueprintPropertyPath& InField);
+	void SetField(const UE::MVVM::FMVVMConstFieldVariant& InField);
 
 private:
-	FMVVMBlueprintPropertyPath Field;
+	UE::MVVM::FMVVMConstFieldVariant Field;
 	FIsFieldValid OnValidateField;
 	TSharedPtr<SMVVMFieldIcon> Icon;
 	TSharedPtr<STextBlock> Label;
@@ -50,7 +50,7 @@ private:
 class SMVVMFieldSelector : public SCompoundWidget
 {
 public:
-	DECLARE_DELEGATE_OneParam(FSelectionChanged, FMVVMBlueprintPropertyPath);
+	DECLARE_DELEGATE_OneParam(FSelectionChanged, UE::MVVM::FMVVMConstFieldVariant);
 
 	SLATE_BEGIN_ARGS(SMVVMFieldSelector) :
 		_TextStyle(&FCoreStyle::Get().GetWidgetStyle<FTextBlockStyle>("NormalText")),
@@ -58,17 +58,11 @@ public:
 		{
 		}
 		SLATE_STYLE_ARGUMENT(FTextBlockStyle, TextStyle)
+		SLATE_ARGUMENT(TArray<UE::MVVM::IFieldPathHelper*>, PathHelpers)
+		SLATE_ARGUMENT(TArray<UE::MVVM::IFieldPathHelper*>, CounterpartHelpers)
 		SLATE_ATTRIBUTE(EMVVMBindingMode, BindingMode)
 		SLATE_ARGUMENT(bool, IsSource)
-		SLATE_ATTRIBUTE(UE::MVVM::FBindingSource, SelectedSource)
-		SLATE_ATTRIBUTE(FMVVMBlueprintPropertyPath, SelectedField)
-		SLATE_ATTRIBUTE(TArray<FMVVMBlueprintPropertyPath>, AvailableFields)
 		SLATE_EVENT(FSelectionChanged, OnSelectionChanged)
-
-		/**
-		  * Would the given field be a valid entry for this combo box? 
-		  * The error string returned will be used as a tooltip.
-		  */
 		SLATE_EVENT(FIsFieldValid, OnValidateField)
 	SLATE_END_ARGS()
 
@@ -76,28 +70,29 @@ public:
 	void Refresh();
 
 private:
-	TSharedRef<SWidget> OnGenerateFieldWidget(FMVVMBlueprintPropertyPath Path) const;
-	void OnComboBoxSelectionChanged(FMVVMBlueprintPropertyPath Selected, ESelectInfo::Type SelectionType);
+	TSharedRef<SWidget> OnGenerateFieldWidget(UE::MVVM::FMVVMConstFieldVariant FieldPath) const;
+	void OnComboBoxSelectionChanged(UE::MVVM::FMVVMConstFieldVariant Selected, ESelectInfo::Type SelectionType);
+	TOptional<UE::MVVM::FBindingSource> GetSelectedSource() const;
 
-	TValueOrError<bool, FString> ValidateField(FMVVMBlueprintPropertyPath Field) const;
+	TValueOrError<bool, FString> IsValidBindingForField(const UE::MVVM::FMVVMConstFieldVariant& Field, const UE::MVVM::FMVVMConstFieldVariant& CounterpartField) const;
+	TValueOrError<bool, FString> ValidateField(UE::MVVM::FMVVMConstFieldVariant Field) const;
 
 	EVisibility GetClearVisibility() const;
 	FReply OnClearBinding();
 
 private:
+	TArray<UE::MVVM::IFieldPathHelper*> PathHelpers;
+	TArray<UE::MVVM::IFieldPathHelper*> CounterpartHelpers;
 	TAttribute<EMVVMBindingMode> BindingMode;
-	TSharedPtr<SMVVMFieldEntry> SelectedEntryWidget;
-	TAttribute<UE::MVVM::FBindingSource> SelectedSource;
-	TAttribute<FMVVMBlueprintPropertyPath> SelectedField;
-	TAttribute<TArray<FMVVMBlueprintPropertyPath>> AvailableFields;
-
+	TSharedPtr<SMVVMFieldEntry> SelectedEntry;
+	TOptional<UE::MVVM::FBindingSource> SelectedSource;
 	const FTextBlockStyle* TextStyle = nullptr;
 	bool bIsSource = false;
 
-	TSharedPtr<SComboBox<FMVVMBlueprintPropertyPath>> FieldComboBox;
+	TSharedPtr<SComboBox<UE::MVVM::FMVVMConstFieldVariant>> FieldComboBox;
 
 	FSelectionChanged OnSelectionChangedDelegate;
 	FIsFieldValid OnValidateFieldDelegate;
 
-	TArray<FMVVMBlueprintPropertyPath> CachedAvailableFields;
+	TArray<UE::MVVM::FMVVMConstFieldVariant> AvailableFields;
 }; 
diff --git a/Engine/Plugins/Runtime/ModelViewViewModel/Source/ModelViewViewModelEditor/Private/Widgets/SMVVMPropertyPath.cpp b/Engine/Plugins/Runtime/ModelViewViewModel/Source/ModelViewViewModelEditor/Private/Widgets/SMVVMPropertyPath.cpp
--- a/Engine/Plugins/Runtime/ModelViewViewModel/Source/ModelViewViewModelEditor/Private/Widgets/SMVVMPropertyPath.cpp
+++ b/Engine/Plugins/Runtime/ModelViewViewModel/Source/ModelViewViewModelEditor/Private/Widgets/SMVVMPropertyPath.cpp
@@ -2,13 +2,10 @@
 
 #include "Widgets/SMVVMPropertyPath.h"
 
-#include "Blueprint/WidgetTree.h"
-#include "Editor.h"
 #include "MVVMBlueprintView.h"
-#include "MVVMEditorSubsystem.h"
 #include "MVVMWidgetBlueprintExtension_View.h"
-#include "Styling/MVVMEditorStyle.h"
 #include "Widgets/SMVVMFieldIcon.h"
+#include "Styling/MVVMEditorStyle.h"
 #include "Widgets/SBoxPanel.h"
 #include "Widgets/Images/SImage.h"
 #include "Widgets/Layout/SBorder.h"
@@ -17,12 +14,37 @@
 
 #define LOCTEXT_NAMESPACE "SMVVMPropertyPath"
 
-void SMVVMPropertyPath::Construct(const FArguments& InArgs)
+void SMVVMPropertyPathBase::Construct(const FArguments& InArgs, const UWidgetBlueprint* InWidgetBlueprint)
 {
-	WidgetBlueprint = InArgs._WidgetBlueprint;
-	check(InArgs._WidgetBlueprint != nullptr);
-	PropertyPath = InArgs._PropertyPath;
-	check(PropertyPath);
+	WidgetBlueprint = InWidgetBlueprint;
+
+	const UE::MVVM::IFieldPathHelper& PathHelper = GetPathHelper();
+	SelectedSource = PathHelper.GetSelectedSource();
+	SelectedField = PathHelper.GetSelectedField();
+
+	auto GetContextText = [this]()
+	{
+		if (SelectedSource.IsSet())
+		{
+			return SelectedSource.GetValue().DisplayName;
+		}
+		return FText::GetEmpty();
+	};
+
+	auto GetPathText = [this]()
+	{
+		if (SelectedField.IsProperty())
+		{
+			return SelectedField.GetProperty()->GetDisplayNameText();
+		}
+
+		if (SelectedField.IsFunction())
+		{
+			return SelectedField.GetFunction()->GetDisplayNameText();
+		}
+
+		return FText::GetEmpty();
+	};
 
 	ChildSlot
 	[
@@ -33,7 +55,7 @@ void SMVVMPropertyPath::Construct(const FArguments& InArgs)
 		[
 			SNew(STextBlock)
 			.TextStyle(FMVVMEditorStyle::Get(), "PropertyPath.ContextText")
-			.Text(this, &SMVVMPropertyPath::GetSourceDisplayName)
+			.Text_Lambda(GetContextText)
 		]
 		+ SHorizontalBox::Slot()
 		.Padding(FMargin(5, 0, 0, 0))
@@ -41,7 +63,7 @@ void SMVVMPropertyPath::Construct(const FArguments& InArgs)
 		.AutoWidth()
 		[
 			SNew(SMVVMFieldIcon)
-			.Field(GetLastField())
+			.Field(SelectedField)
 		]
 		+ SHorizontalBox::Slot()
 		.Padding(FMargin(2, 0, 5, 0))
@@ -50,38 +72,24 @@ void SMVVMPropertyPath::Construct(const FArguments& InArgs)
 		[
 			SNew(STextBlock)
 			.TextStyle(FMVVMEditorStyle::Get(), "PropertyPath.ContextText")
-			.Text(this, &SMVVMPropertyPath::GetFieldDisplayName)
+			.Text_Lambda(GetPathText)
 		]
 	];
 }
 
-FText SMVVMPropertyPath::GetSourceDisplayName() const
+void SMVVMWidgetPropertyPath::Construct(const FArguments& InArgs, const UWidgetBlueprint* InWidgetBlueprint)
 {
-	if (PropertyPath->IsFromWidget())
-	{
-		const UWidget* Widget = WidgetBlueprint->WidgetTree->FindWidget(PropertyPath->GetWidgetName());
-		return FText::FromString(Widget->GetDisplayLabel());
-	}
-	else if (PropertyPath->IsFromViewModel())
-	{
-		const UMVVMEditorSubsystem* Subsystem = GEditor->GetEditorSubsystem<UMVVMEditorSubsystem>();
-		if (const UMVVMBlueprintView* View = Subsystem->GetView(WidgetBlueprint.Get()))
-		{
-			const FMVVMBlueprintViewModelContext* ViewModel = View->FindViewModel(PropertyPath->GetViewModelId());
-			return ViewModel->GetDisplayName();
-		}
-	}
-	return FText::GetEmpty();
-}
+	PathHelper = UE::MVVM::FWidgetFieldPathHelper(InArgs._WidgetPath, InWidgetBlueprint);
 
-FText SMVVMPropertyPath::GetFieldDisplayName() const
-{
-	return FText::FromString(PropertyPath->GetBasePropertyPath());
+	SMVVMPropertyPathBase::Construct(SMVVMPropertyPathBase::FArguments(), InWidgetBlueprint);
 }
 
-UE::MVVM::FMVVMConstFieldVariant SMVVMPropertyPath::GetLastField() const
+void SMVVMViewModelPropertyPath::Construct(const FArguments& InArgs, const UWidgetBlueprint* InWidgetBlueprint)
 {
-	return PropertyPath->GetFields().Last();
+	PathHelper = UE::MVVM::FViewModelFieldPathHelper(InArgs._ViewModelPath, InWidgetBlueprint);
+
+	SMVVMPropertyPathBase::Construct(SMVVMPropertyPathBase::FArguments(), InWidgetBlueprint);
 }
 
+
 #undef LOCTEXT_NAMESPACE
\ No newline at end of file
diff --git a/Engine/Plugins/Runtime/ModelViewViewModel/Source/ModelViewViewModelEditor/Private/Widgets/SMVVMPropertyPath.h b/Engine/Plugins/Runtime/ModelViewViewModel/Source/ModelViewViewModelEditor/Private/Widgets/SMVVMPropertyPath.h
--- a/Engine/Plugins/Runtime/ModelViewViewModel/Source/ModelViewViewModelEditor/Private/Widgets/SMVVMPropertyPath.h
+++ b/Engine/Plugins/Runtime/ModelViewViewModel/Source/ModelViewViewModelEditor/Private/Widgets/SMVVMPropertyPath.h
@@ -5,25 +5,62 @@
 #include "Widgets/SCompoundWidget.h"
 #include "MVVMBlueprintView.h"
 #include "Widgets/SMVVMFieldIcon.h"
+#include "MVVMPropertyPathHelpers.h"
 #include "Types/MVVMFieldVariant.h"
 
-class SMVVMPropertyPath : public SCompoundWidget
+class SMVVMPropertyPathBase : public SCompoundWidget
 {
 public:
-	SLATE_BEGIN_ARGS(SMVVMPropertyPath)
+	SLATE_BEGIN_ARGS(SMVVMPropertyPathBase)
 	{}
-		SLATE_ARGUMENT(FMVVMBlueprintPropertyPath*, PropertyPath)
-		SLATE_ARGUMENT(const UWidgetBlueprint*, WidgetBlueprint)
 	SLATE_END_ARGS()
 
-	void Construct(const FArguments& InArgs);
+	void Construct(const FArguments& InArgs, const UWidgetBlueprint* WidgetBlueprint);
+	virtual void Tick(const FGeometry& AllottedGeometry, const double InCurrentTime, const float InDeltaTime) override;
 
 private:
-	FText GetSourceDisplayName() const;
-	FText GetFieldDisplayName() const;
-	UE::MVVM::FMVVMConstFieldVariant GetLastField() const;
+	virtual UE::MVVM::IFieldPathHelper& GetPathHelper() = 0;
 
 private:
-	FMVVMBlueprintPropertyPath* PropertyPath = nullptr;
-	TWeakObjectPtr<const UWidgetBlueprint> WidgetBlueprint = nullptr;
+	const UWidgetBlueprint* WidgetBlueprint = nullptr;
+	UE::MVVM::FMVVMConstFieldVariant SelectedField;
+	TOptional<UE::MVVM::FBindingSource> SelectedSource;
+};
+
+class SMVVMWidgetPropertyPath : public SMVVMPropertyPathBase
+{
+public:
+	SLATE_BEGIN_ARGS(SMVVMWidgetPropertyPath) {}
+		SLATE_ARGUMENT(FMVVMBlueprintPropertyPath*, WidgetPath)
+	SLATE_END_ARGS()
+
+	void Construct(const FArguments& InArgs, const UWidgetBlueprint* WidgetBlueprint);
+
+private:
+	virtual UE::MVVM::IFieldPathHelper& GetPathHelper() override
+	{ 
+		return PathHelper;
+	}
+
+private:
+	UE::MVVM::FWidgetFieldPathHelper PathHelper;
+};
+
+class SMVVMViewModelPropertyPath : public SMVVMPropertyPathBase
+{
+public:
+	SLATE_BEGIN_ARGS(SMVVMViewModelPropertyPath) {}
+		SLATE_ARGUMENT(FMVVMBlueprintPropertyPath*, ViewModelPath)
+	SLATE_END_ARGS()
+
+	void Construct(const FArguments& InArgs, const UWidgetBlueprint* WidgetBlueprint);
+
+private:
+	virtual UE::MVVM::IFieldPathHelper& GetPathHelper() override
+	{
+		return PathHelper;
+	}
+
+private:
+	UE::MVVM::FViewModelFieldPathHelper PathHelper;
 };
diff --git a/Engine/Plugins/Runtime/ModelViewViewModel/Source/ModelViewViewModelEditor/Private/Widgets/SMVVMSourceSelector.cpp b/Engine/Plugins/Runtime/ModelViewViewModel/Source/ModelViewViewModelEditor/Private/Widgets/SMVVMSourceSelector.cpp
--- a/Engine/Plugins/Runtime/ModelViewViewModel/Source/ModelViewViewModelEditor/Private/Widgets/SMVVMSourceSelector.cpp
+++ b/Engine/Plugins/Runtime/ModelViewViewModel/Source/ModelViewViewModelEditor/Private/Widgets/SMVVMSourceSelector.cpp
@@ -11,13 +11,33 @@
 
 using namespace UE::MVVM;
 
+static FBindingSource GetSelectedSource(const TArray<IFieldPathHelper*>& Helpers)
+{
+	bool bFirst = true;
+	TOptional<FBindingSource> Result;
+
+	for (const IFieldPathHelper* Helper : Helpers)
+	{
+		TOptional<FBindingSource> ThisSelection = Helper->GetSelectedSource();
+		if (bFirst)
+		{
+			Result = ThisSelection;
+		}
+		else if (Result != ThisSelection)
+		{
+			// all don't share the same value
+			Result.Reset();
+			break;
+		}
+	}
+
+	return Result.Get(FBindingSource());
+}
+
 void SMVVMSourceSelector::Construct(const FArguments& Args)
 {
+	PathHelpers = Args._PathHelpers;
 	TextStyle = Args._TextStyle;
-	AvailableSourcesAttribute = Args._AvailableSources;
-	check(AvailableSourcesAttribute.IsSet());
-	SelectedSourceAttribute = Args._SelectedSource;
-	check(SelectedSourceAttribute.IsSet());
 
 	Refresh();
 
@@ -71,8 +91,36 @@ void SMVVMSourceSelector::OnComboBoxSelectionChanged(FBindingSource Selected, ES
 
 void SMVVMSourceSelector::Refresh()
 {
-	SelectedSource = SelectedSourceAttribute.Get();
-	AvailableSources = AvailableSourcesAttribute.Get();
+	AvailableSources.Reset();
+	SelectedSource = FBindingSource();
+
+	TArray<IFieldPathHelper*> Helpers = PathHelpers.Get(TArray<IFieldPathHelper*>());
+
+	bool bFirst = true;
+
+	TSet<FBindingSource> SourceIntersection;
+	TOptional<FBindingSource> NewSelection;
+
+	for (const IFieldPathHelper* Helper : Helpers)
+	{
+		// get available sources
+		TSet<FBindingSource> Sources;
+		Helper->GetAvailableSources(Sources);
+
+		if (bFirst)
+		{
+			bFirst = false; 
+
+			SourceIntersection.Append(Sources);
+		}
+		else
+		{
+			SourceIntersection.Intersect(Sources);
+		}
+	}
+
+	AvailableSources = SourceIntersection.Array();
+	SelectedSource = GetSelectedSource(Helpers);
 
 	if (SourceComboBox.IsValid())
 	{
@@ -83,7 +131,16 @@ void SMVVMSourceSelector::Refresh()
 
 EVisibility SMVVMSourceSelector::GetClearVisibility() const
 {
-	return SelectedSource.IsValid() ? EVisibility::Visible : EVisibility::Collapsed;
+	for (const IFieldPathHelper* Helper : PathHelpers.Get())
+	{
+		const FBindingSource ThisSelection = Helper->GetSelectedSource();
+		if (ThisSelection.IsValid())
+		{
+			return EVisibility::Visible;
+		}
+	}
+
+	return EVisibility::Collapsed;
 }
 
 FReply SMVVMSourceSelector::OnClearSource()
diff --git a/Engine/Plugins/Runtime/ModelViewViewModel/Source/ModelViewViewModelEditor/Private/Widgets/SMVVMSourceSelector.h b/Engine/Plugins/Runtime/ModelViewViewModel/Source/ModelViewViewModelEditor/Private/Widgets/SMVVMSourceSelector.h
--- a/Engine/Plugins/Runtime/ModelViewViewModel/Source/ModelViewViewModelEditor/Private/Widgets/SMVVMSourceSelector.h
+++ b/Engine/Plugins/Runtime/ModelViewViewModel/Source/ModelViewViewModelEditor/Private/Widgets/SMVVMSourceSelector.h
@@ -3,6 +3,7 @@
 #pragma once
 
 #include "CoreMinimal.h"
+#include "MVVMPropertyPathHelpers.h"
 #include "Types/MVVMBindingSource.h"
 #include "Widgets/DeclarativeSyntaxSupport.h"
 #include "Widgets/SCompoundWidget.h"
@@ -20,8 +21,7 @@ public:
 		{
 		}
 		SLATE_STYLE_ARGUMENT(FTextBlockStyle, TextStyle)
-		SLATE_ATTRIBUTE(UE::MVVM::FBindingSource, SelectedSource)
-		SLATE_ATTRIBUTE(TArray<UE::MVVM::FBindingSource>, AvailableSources)
+		SLATE_ATTRIBUTE(TArray<UE::MVVM::IFieldPathHelper*>, PathHelpers)
 		SLATE_EVENT(FSelectionChanged, OnSelectionChanged)
 	SLATE_END_ARGS()
 
@@ -35,9 +35,8 @@ private:
 	FReply OnClearSource();
 
 private:
+	TAttribute<TArray<UE::MVVM::IFieldPathHelper*>> PathHelpers;
 	FSelectionChanged OnSelectionChanged;
-	TAttribute<TArray<UE::MVVM::FBindingSource>> AvailableSourcesAttribute;
-	TAttribute< UE::MVVM::FBindingSource> SelectedSourceAttribute;
 	TSharedPtr<SComboBox<UE::MVVM::FBindingSource>> SourceComboBox;
 	TArray<UE::MVVM::FBindingSource> AvailableSources;
 	UE::MVVM::FBindingSource SelectedSource;
