MVVM: Setting a conversion function with multiple arguments now generates a FunctionGraph in the widget blueprint that is used to call the function, which allows us to bind multiple properties to conversion functions.

Removed IFieldPathHelper. The functionality was moved to the editor subsystem, which made the abstraction unnecessary and unavailable to blueprint.

FBindingSource no longer contains the IsSelected member. This was very context-specific, and doesn't really make sense in the general case.

#rnx

diff --git a/Engine/Plugins/Runtime/ModelViewViewModel/Source/ModelViewViewModel/Private/Bindings/MVVMBindingHelper.cpp b/Engine/Plugins/Runtime/ModelViewViewModel/Source/ModelViewViewModel/Private/Bindings/MVVMBindingHelper.cpp
--- a/Engine/Plugins/Runtime/ModelViewViewModel/Source/ModelViewViewModel/Private/Bindings/MVVMBindingHelper.cpp
+++ b/Engine/Plugins/Runtime/ModelViewViewModel/Source/ModelViewViewModel/Private/Bindings/MVVMBindingHelper.cpp
@@ -383,8 +383,7 @@ namespace UE::MVVM::BindingHelper
 			return TryGetPropertyTypeForDestinationBinding(InField.GetFunction());
 		}
 	}
-
-	TValueOrError<FConversionFunctionArguments, FString> TryGetPropertyTypeForConversionFunction(const UFunction* InFunction)
+	TValueOrError<const FProperty*, FString> TryGetReturnTypeForConversionFunction(const UFunction* InFunction)
 	{
 		FString CommonResult = Private::TryGetPropertyTypeCommon(InFunction);
 		if (!CommonResult.IsEmpty())
@@ -397,27 +396,56 @@ namespace UE::MVVM::BindingHelper
 			return MakeError(FString::Printf(TEXT("The function '%s' is not static or is not const and pure."), *InFunction->GetName()));
 		}
 
-		if (InFunction->NumParms != 2)
+		if (InFunction->NumParms < 2)
 		{
 			return MakeError(FString::Printf(TEXT("The function '%s' does not have the correct number of arguments."), *InFunction->GetName()));
 		}
 
-		FConversionFunctionArguments Result;
-		Result.ReturnProperty = GetReturnProperty(InFunction);
-		if (Result.ReturnProperty == nullptr)
+		const FProperty* ReturnProperty = GetReturnProperty(InFunction);
+		if (ReturnProperty == nullptr)
 		{
 			return MakeError(FString::Printf(TEXT("The return value for function '%s' is invalid."), *InFunction->GetName()));
 		}
 
-		Result.ArgumentProperty = GetFirstArgumentProperty(InFunction);
-		if (Result.ArgumentProperty == nullptr)
+		return MakeValue(ReturnProperty);
+	}
+
+	TValueOrError<TArray<const FProperty*>, FString> TryGetArgumentsForConversionFunction(const UFunction* InFunction)
+	{
+		FString CommonResult = Private::TryGetPropertyTypeCommon(InFunction);
+		if (!CommonResult.IsEmpty())
 		{
-			return MakeError(FString::Printf(TEXT("The function '%s' doesn't not have a valid argument."), *InFunction->GetName()));
+			return MakeError(MoveTemp(CommonResult));
 		}
 
-		return MakeValue(Result);
-	}
+		if (!InFunction->HasAllFunctionFlags(FUNC_Static) && !InFunction->HasAnyFunctionFlags(FUNC_Const | FUNC_BlueprintPure))
+		{
+			return MakeError(FString::Printf(TEXT("The function '%s' is not static or is not const and pure."), *InFunction->GetName()));
+		}
+
+		if (InFunction->NumParms < 2)
+		{
+			return MakeError(FString::Printf(TEXT("The function '%s' does not have the correct number of arguments."), *InFunction->GetName()));
+		}
+
+		const FProperty* ReturnProperty = GetReturnProperty(InFunction);
+		if (ReturnProperty == nullptr)
+		{
+			return MakeError(FString::Printf(TEXT("The return value for function '%s' is invalid."), *InFunction->GetName()));
+		}
+
+		TArray<const FProperty*> Arguments;
+		for (TFieldIterator<FProperty> It(InFunction); It && (It->PropertyFlags & CPF_Parm); ++It)
+		{
+			const FProperty* Property = *It;
+			if (IsValidArgumentProperty(Property))
+			{
+				Arguments.Add(Property);
+			}
+		}
 
+		return MakeValue(Arguments);
+	}
 
 	namespace Private
 	{
@@ -520,6 +548,20 @@ namespace UE::MVVM::BindingHelper
 		return Result;
 	}
 
+	bool IsValidArgumentProperty(const FProperty* Property)
+	{
+		if (Property->HasAllPropertyFlags(CPF_Parm) && !Property->HasAnyPropertyFlags(CPF_Deprecated | CPF_EditorOnly | CPF_ReturnParm))
+		{
+			if (Property->HasAllPropertyFlags(CPF_OutParm) && !Property->HasAnyPropertyFlags(CPF_ConstParm | CPF_ReferenceParm))
+			{
+				return false;
+			}
+
+			return true;
+		}
+
+		return false;
+	}
 
 	// We only accept copy argument or const ref
 	const FProperty* GetFirstArgumentProperty(const UFunction* InFunction)
@@ -527,12 +569,8 @@ namespace UE::MVVM::BindingHelper
 		check(InFunction);
 		for (TFieldIterator<FProperty> It(InFunction); It && (It->PropertyFlags & CPF_Parm); ++It)
 		{
-			if (It->HasAllPropertyFlags(CPF_Parm) && !It->HasAnyPropertyFlags(CPF_Deprecated | CPF_EditorOnly | CPF_ReturnParm))
+			if (IsValidArgumentProperty(*It))
 			{
-				if (It->HasAllPropertyFlags(CPF_OutParm) && !It->HasAnyPropertyFlags(CPF_ConstParm | CPF_ReferenceParm))
-				{
-					continue;
-				}
 				return *It;
 			}
 		}
diff --git a/Engine/Plugins/Runtime/ModelViewViewModel/Source/ModelViewViewModel/Private/MVVMSubsystem.cpp b/Engine/Plugins/Runtime/ModelViewViewModel/Source/ModelViewViewModel/Private/MVVMSubsystem.cpp
--- a/Engine/Plugins/Runtime/ModelViewViewModel/Source/ModelViewViewModel/Private/MVVMSubsystem.cpp
+++ b/Engine/Plugins/Runtime/ModelViewViewModel/Source/ModelViewViewModel/Private/MVVMSubsystem.cpp
@@ -128,18 +128,15 @@ namespace UE::MVVM::Private
 		TArray<FMVVMAvailableBinding> Result;
 		Result.Reserve(FieldDescriptors.Num());
 
-		for (TFieldIterator<const UFunction> FunctionItt(Class, EFieldIteratorFlags::IncludeSuper); FunctionItt; ++FunctionItt)
+		for (TFieldIterator<const UFunction> FunctionIt(Class, EFieldIteratorFlags::IncludeSuper); FunctionIt; ++FunctionIt)
 		{
-			const UFunction* Function = *FunctionItt;
+			const UFunction* Function = *FunctionIt;
 			check(Function);
-			if (!Function->HasAnyFunctionFlags(FUNC_Private | FUNC_Protected))
+			const bool bIsReadable = FieldDescriptors.Contains(Function->GetFName()) && BindingHelper::IsValidForSourceBinding(Function);
+			const bool bIsWritable = BindingHelper::IsValidForDestinationBinding(Function);
+			if (bIsReadable || bIsWritable)
 			{
-				const bool bIsReadable = FieldDescriptors.Contains(Function->GetFName()) && BindingHelper::IsValidForSourceBinding(Function);
-				const bool bIsWritable = BindingHelper::IsValidForDestinationBinding(Function);
-				if (bIsReadable || bIsWritable)
-				{
-					Result.Add(FMVVMAvailableBinding(FMVVMBindingName(Function->GetFName()), bIsReadable, bIsWritable));
-				}
+				Result.Add(FMVVMAvailableBinding(FMVVMBindingName(Function->GetFName()), bIsReadable, bIsWritable));
 			}
 		}
 
@@ -254,26 +251,26 @@ TValueOrError<bool, FString> UMVVMSubsystem::IsBindingValid(FDirectionalBindingA
 TValueOrError<bool, FString> UMVVMSubsystem::IsBindingValid(FConstDirectionalBindingArgs Args) const
 {
 	// Test Source
-	TValueOrError<const FProperty*, FString> GetterProperty = UE::MVVM::BindingHelper::TryGetPropertyTypeForSourceBinding(Args.SourceBinding);
-	if (GetterProperty.HasError())
+	TValueOrError<const FProperty*, FString> SourceResult = UE::MVVM::BindingHelper::TryGetPropertyTypeForSourceBinding(Args.SourceBinding);
+	if (SourceResult.HasError())
 	{
-		return MakeError(GetterProperty.StealError());
+		return MakeError(SourceResult.StealError());
 	}
 
-	const FProperty* SourceProperty = GetterProperty.StealValue();
+	const FProperty* SourceProperty = SourceResult.StealValue();
 	if (SourceProperty == nullptr)
 	{
 		return MakeError(TEXT("There is no value to bind at the source."));
 	}
 
 	// Test Destination
-	TValueOrError<const FProperty*, FString> SetterProperty = UE::MVVM::BindingHelper::TryGetPropertyTypeForDestinationBinding(Args.DestinationBinding);
-	if (SetterProperty.HasError())
+	TValueOrError<const FProperty*, FString> DestinationResult = UE::MVVM::BindingHelper::TryGetPropertyTypeForDestinationBinding(Args.DestinationBinding);
+	if (DestinationResult.HasError())
 	{
-		return MakeError(SetterProperty.StealError());
+		return MakeError(DestinationResult.StealError());
 	}
 
-	const FProperty* DestinationProperty = SetterProperty.StealValue();
+	const FProperty* DestinationProperty = DestinationResult.StealValue();
 	if (DestinationProperty == nullptr)
 	{
 		return MakeError(TEXT("There is no value to bind at the destination."));
@@ -282,23 +279,42 @@ TValueOrError<bool, FString> UMVVMSubsystem::IsBindingValid(FConstDirectionalBin
 	// Test the conversion function
 	if (Args.ConversionFunction)
 	{
-		TValueOrError<UE::MVVM::BindingHelper::FConversionFunctionArguments, FString> ConversionResult = UE::MVVM::BindingHelper::TryGetPropertyTypeForConversionFunction(Args.ConversionFunction);
-		if (ConversionResult.HasError())
+		TValueOrError<const FProperty*, FString> ReturnResult = UE::MVVM::BindingHelper::TryGetReturnTypeForConversionFunction(Args.ConversionFunction);
+		if (ReturnResult.HasError())
 		{
-			return MakeError(ConversionResult.StealError());
+			return MakeError(ReturnResult.StealError());
+		}
+
+		TValueOrError<TArray<const FProperty*>, FString> ArgumentsResult = UE::MVVM::BindingHelper::TryGetArgumentsForConversionFunction(Args.ConversionFunction);
+		if (ArgumentsResult.HasError())
+		{
+			return MakeError(ArgumentsResult.StealError());
 		}
 
 		// The compiled version should look like Setter(Conversion(Getter())).
-		const FProperty* ConversionReturnProperty = ConversionResult.GetValue().ReturnProperty;
-		if (!UE::MVVM::BindingHelper::ArePropertiesCompatible(ConversionReturnProperty, DestinationProperty))
+		const FProperty* ReturnProperty = ReturnResult.GetValue();
+		if (!UE::MVVM::BindingHelper::ArePropertiesCompatible(ReturnProperty, DestinationProperty))
+		{
+			return MakeError(FString::Printf(TEXT("The destination property '%s' (%s) does not match the return type of the conversion function (%s)."), *DestinationProperty->GetName(), *DestinationProperty->GetCPPType(), *ReturnProperty->GetCPPType()));
+		}
+
+		bool bAnyCompatible = false;
+
+		const TArray<const FProperty*>& ArgumentProperties = ArgumentsResult.GetValue();
+		for (const FProperty* ArgumentProperty : ArgumentProperties)
 		{
-			return MakeError(FString::Printf(TEXT("The destination property '%s' (%s) does not match the return type of the conversion function (%s)."), *DestinationProperty->GetName(), *DestinationProperty->GetCPPType(), *ConversionReturnProperty->GetCPPType()));
+			if (UE::MVVM::BindingHelper::ArePropertiesCompatible(SourceProperty, ArgumentProperty))
+			{
+				bAnyCompatible = true;
+				break;
+			}
 		}
 
-		const FProperty* ConversionArgProperty = ConversionResult.GetValue().ArgumentProperty;
-		if (!UE::MVVM::BindingHelper::ArePropertiesCompatible(SourceProperty, ConversionArgProperty))
+		if (!bAnyCompatible)
 		{
-			return MakeError(FString::Printf(TEXT("The source property '%s' (%s) does not match the argument type of the conversion function (%s)."), *SourceProperty->GetName(), *SourceProperty->GetCPPType(), *DestinationProperty->GetCPPType()));
+			FString ArgumentsString = FString::JoinBy(ArgumentProperties, TEXT(", "), [](const FProperty* Property) { return Property->GetCPPType(); });
+
+			return MakeError(FString::Printf(TEXT("The source property '%s' (%s) does not match any of the argument types of the conversion function (%s)."), *SourceProperty->GetName(), *SourceProperty->GetCPPType(), *ArgumentsString));
 		}
 	}
 	else if (!UE::MVVM::BindingHelper::ArePropertiesCompatible(SourceProperty, DestinationProperty))
@@ -308,4 +324,3 @@ TValueOrError<bool, FString> UMVVMSubsystem::IsBindingValid(FConstDirectionalBin
 
 	return MakeValue(true);
 }
-
diff --git a/Engine/Plugins/Runtime/ModelViewViewModel/Source/ModelViewViewModel/Public/Bindings/MVVMBindingHelper.h b/Engine/Plugins/Runtime/ModelViewViewModel/Source/ModelViewViewModel/Public/Bindings/MVVMBindingHelper.h
--- a/Engine/Plugins/Runtime/ModelViewViewModel/Source/ModelViewViewModel/Public/Bindings/MVVMBindingHelper.h
+++ b/Engine/Plugins/Runtime/ModelViewViewModel/Source/ModelViewViewModel/Public/Bindings/MVVMBindingHelper.h
@@ -79,13 +79,11 @@ namespace UE::MVVM::BindingHelper
 	UE_NODISCARD MODELVIEWVIEWMODEL_API TValueOrError<const FProperty*, FString> TryGetPropertyTypeForSourceBinding(const UFunction* InFunction);
 	UE_NODISCARD MODELVIEWVIEWMODEL_API TValueOrError<const FProperty*, FString> TryGetPropertyTypeForDestinationBinding(const UFunction* InFunction);
 
-	struct FConversionFunctionArguments
-	{
-		const FProperty* ReturnProperty = nullptr;
-		const FProperty* ArgumentProperty = nullptr;
-	};
-	/** Returns the Return and 1st argument for a conversion function. */
-	UE_NODISCARD MODELVIEWVIEWMODEL_API TValueOrError<FConversionFunctionArguments, FString> TryGetPropertyTypeForConversionFunction(const UFunction* InFunction);
+	/** Returns the return property for a conversion function, or an error if the function is not valid as a conversion function. */
+	UE_NODISCARD MODELVIEWVIEWMODEL_API TValueOrError<const FProperty*, FString> TryGetReturnTypeForConversionFunction(const UFunction* InFunction);
+
+	/** Returns the argument properties for a conversion function, or an error if the function is not valid as a conversion function. */
+	UE_NODISCARD MODELVIEWVIEWMODEL_API TValueOrError<TArray<const FProperty*>, FString> TryGetArgumentsForConversionFunction(const UFunction* InFunction);
 
 	/** Are type the same or could be converted at runtime. */
 	UE_NODISCARD MODELVIEWVIEWMODEL_API bool ArePropertiesCompatible(const FProperty* Source, const FProperty* Destination);
@@ -106,6 +104,11 @@ namespace UE::MVVM::BindingHelper
 	 */
 	UE_NODISCARD MODELVIEWVIEWMODEL_API const FProperty* GetFirstArgumentProperty(const UFunction* InFunction);
 
+	/**
+	 * Returns true if the given property is a valid function argument property.
+	 */
+	UE_NODISCARD MODELVIEWVIEWMODEL_API bool IsValidArgumentProperty(const FProperty* Property);
+
 	/**
 	 * Execute a binding that can be
 	 *   LocalValue = Src.Property; Dst.Property = LocalValue;
diff --git a/Engine/Plugins/Runtime/ModelViewViewModel/Source/ModelViewViewModel/Public/MVVMSubsystem.h b/Engine/Plugins/Runtime/ModelViewViewModel/Source/ModelViewViewModel/Public/MVVMSubsystem.h
--- a/Engine/Plugins/Runtime/ModelViewViewModel/Source/ModelViewViewModel/Public/MVVMSubsystem.h
+++ b/Engine/Plugins/Runtime/ModelViewViewModel/Source/ModelViewViewModel/Public/MVVMSubsystem.h
@@ -4,7 +4,6 @@
 
 #include "Subsystems/EngineSubsystem.h"
 #include "Templates/ValueOrError.h"
-
 #include "Types/MVVMBindingName.h"
 #include "Types/MVVMFieldVariant.h"
 #include "Types/MVVMBindingMode.h"
@@ -30,7 +29,6 @@ public:
 	virtual void Deinitialize() override;
 	//~ End UEngineSubsystem interface
 
-public:
 	UFUNCTION(BlueprintCallable, Category="MVVM")
 	UMVVMView* GetViewFromUserWidget(const UUserWidget* UserWidget) const;
 
@@ -64,15 +62,12 @@ public:
 	UFUNCTION(BlueprintCallable, Category = "MVVM")
 	TArray<FMVVMAvailableBinding> GetAvailableBindings(TSubclassOf<UObject> Class) const;
 
-public:
 	UFUNCTION(BlueprintCallable, Category = "MVVM")
 	UMVVMViewModelCollectionObject* GetGlobalViewModelCollection() const
 	{
 		return GlobalViewModelCollection;
 	}
 
-public:
-
 	struct FConstDirectionalBindingArgs
 	{
 		UE::MVVM::FMVVMConstFieldVariant SourceBinding;
diff --git a/Engine/Plugins/Runtime/ModelViewViewModel/Source/ModelViewViewModelBlueprint/Private/Bindings/MVVMCompiledBindingLibraryCompiler.cpp b/Engine/Plugins/Runtime/ModelViewViewModel/Source/ModelViewViewModelBlueprint/Private/Bindings/MVVMCompiledBindingLibraryCompiler.cpp
--- a/Engine/Plugins/Runtime/ModelViewViewModel/Source/ModelViewViewModelBlueprint/Private/Bindings/MVVMCompiledBindingLibraryCompiler.cpp
+++ b/Engine/Plugins/Runtime/ModelViewViewModel/Source/ModelViewViewModelBlueprint/Private/Bindings/MVVMCompiledBindingLibraryCompiler.cpp
@@ -757,7 +757,7 @@
 		Result.FieldIds.Add(FieldId.IdHandle, FieldId.CompiledFieldId);
 	}
 
-	auto GetCompiledFiledPath = [this](const FFieldPathHandle Handle)
+	auto GetCompiledFieldPath = [this](const FFieldPathHandle Handle)
 	{
 		const Private::FRawFieldPath* FoundBinding = Impl->FieldPaths.FindByPredicate([Handle](const Private::FRawFieldPath& Other)
 			{
@@ -770,20 +770,19 @@
 		return FMVVMVCompiledFieldPath();
 	};
 
-
 	// Create the requested FMVVMVCompiledBinding
 	for (Private::FRawBinding& Binding : Impl->Bindings)
 	{
 		Binding.CompiledBinding.CompiledBindingLibraryId = Result.Library.CompiledBindingLibraryId;
 		check(Binding.CompiledBinding.CompiledBindingLibraryId.IsValid());
 
-		Binding.CompiledBinding.SourceFieldPath = GetCompiledFiledPath(Binding.SourcePathHandles[0]);
+		Binding.CompiledBinding.SourceFieldPath = GetCompiledFieldPath(Binding.SourcePathHandles[0]);
 		check(Binding.CompiledBinding.SourceFieldPath.IsValid());
 
-		Binding.CompiledBinding.DestinationFieldPath = GetCompiledFiledPath(Binding.DestinationPathHandle);
+		Binding.CompiledBinding.DestinationFieldPath = GetCompiledFieldPath(Binding.DestinationPathHandle);
 		check(Binding.CompiledBinding.DestinationFieldPath.IsValid());
 
-		Binding.CompiledBinding.ConversionFunctionFieldPath = GetCompiledFiledPath(Binding.ConversionFunctionPathHandle);
+		Binding.CompiledBinding.ConversionFunctionFieldPath = GetCompiledFieldPath(Binding.ConversionFunctionPathHandle);
 
 		Result.Bindings.Add(Binding.BindingHandle, Binding.CompiledBinding);
 	}
diff --git a/Engine/Plugins/Runtime/ModelViewViewModel/Source/ModelViewViewModelBlueprint/Private/MVVMViewBlueprintCompiler.cpp b/Engine/Plugins/Runtime/ModelViewViewModel/Source/ModelViewViewModelBlueprint/Private/MVVMViewBlueprintCompiler.cpp
--- a/Engine/Plugins/Runtime/ModelViewViewModel/Source/ModelViewViewModelBlueprint/Private/MVVMViewBlueprintCompiler.cpp
+++ b/Engine/Plugins/Runtime/ModelViewViewModel/Source/ModelViewViewModelBlueprint/Private/MVVMViewBlueprintCompiler.cpp
@@ -19,7 +19,7 @@ namespace UE::MVVM::Private
 
 void FMVVMViewBlueprintCompiler::AddErrorForBinding(FMVVMBlueprintViewBinding& Binding, const UMVVMBlueprintView* View, const FString& Message) const
 {
-	const FString BindingName = Binding.GetNameString(View);
+	const FString BindingName = Binding.GetDisplayNameString(View);
 	WidgetBlueprintCompilerContext.MessageLog.Error(*(BindingName + TEXT(": ") + Message));
 	Binding.Errors.Add(FText::FromString(BindingName + TEXT(": ") + Message));
 }
@@ -639,7 +639,7 @@ bool FMVVMViewBlueprintCompiler::PreCompileBindings(UWidgetBlueprintGeneratedCla
 		}
 
 		FMVVMViewBlueprintCompiler* Self = this;
-		auto AddBinding = [Self](UWidgetBlueprintGeneratedClass* Class, const TArrayView<TArray<UE::MVVM::FMVVMConstFieldVariant>> GetterFields, TArrayView<UE::MVVM::FMVVMConstFieldVariant> SetterFields, const FMemberReference& ConversionFunctionReference) -> TValueOrError<FCompilerBinding, FString>
+		auto AddBinding = [Self](UWidgetBlueprintGeneratedClass* Class, const TArrayView<TArray<UE::MVVM::FMVVMConstFieldVariant>> GetterFields, TArrayView<UE::MVVM::FMVVMConstFieldVariant> SetterFields, const UFunction* ConversionFunction) -> TValueOrError<FCompilerBinding, FString>
 		{
 			FCompilerBinding Result;
 
@@ -667,16 +667,8 @@ bool FMVVMViewBlueprintCompiler::PreCompileBindings(UWidgetBlueprintGeneratedCla
 				Result.DestinationWrite = FieldPathResult.GetValue();
 			}
 
-
-			if (!ConversionFunctionReference.GetMemberName().IsNone())
+			if (ConversionFunction != nullptr)
 			{
-				const UFunction* ConversionFunction = ConversionFunctionReference.ResolveMember<UFunction>(Class);
-				if (ConversionFunction == nullptr)
-				{
-					return MakeError(FString::Printf(TEXT("The Conversion Function '%s' could not be resolved.")
-						, *ConversionFunctionReference.GetMemberName().ToString()));
-				}
-
 				TValueOrError<FCompiledBindingLibraryCompiler::FFieldPathHandle, FString> FieldPathResult = Self->BindingLibraryCompiler.AddConversionFunctionFieldPath(Class, ConversionFunction);
 				if (FieldPathResult.HasError())
 				{
@@ -760,7 +752,15 @@ bool FMVVMViewBlueprintCompiler::PreCompileBindings(UWidgetBlueprintGeneratedCla
 				}
 			}
 
-			TValueOrError<FCompilerBinding, FString> AddBindingResult = AddBinding(Class, GetterFields, SetterPath, Binding.Conversion.SourceToDestinationFunction);
+			FMemberReference ConversionFunctionRef = Binding.Conversion.SourceToDestinationFunction;
+			if (!Binding.Conversion.SourceToDestinationWrapper.IsNone())
+			{
+				ConversionFunctionRef.SetSelfMember(Binding.Conversion.SourceToDestinationWrapper);
+			}
+
+			const UFunction* ConversionFunction = ConversionFunctionRef.ResolveMember<UFunction>(Class);
+
+			TValueOrError<FCompilerBinding, FString> AddBindingResult = AddBinding(Class, GetterFields, SetterPath, ConversionFunction);
 			if (AddBindingResult.HasError())
 			{
 				AddErrorForBinding(Binding, BlueprintView, FString::Printf(TEXT("The binding could not be created. %s"), *AddBindingResult.GetError()));
@@ -828,7 +828,15 @@ bool FMVVMViewBlueprintCompiler::PreCompileBindings(UWidgetBlueprintGeneratedCla
 				}
 			}
 
-			TValueOrError<FCompilerBinding, FString> AddBindingResult = AddBinding(Class, GetterFields, SetterPath, Binding.Conversion.DestinationToSourceFunction);
+			FMemberReference ConversionFunctionRef = Binding.Conversion.DestinationToSourceFunction;
+			if (!Binding.Conversion.DestinationToSourceWrapper.IsNone())
+			{
+				ConversionFunctionRef.SetSelfMember(Binding.Conversion.DestinationToSourceWrapper);
+			}
+
+			const UFunction* ConversionFunction = ConversionFunctionRef.ResolveMember<UFunction>(Class);
+
+			TValueOrError<FCompilerBinding, FString> AddBindingResult = AddBinding(Class, GetterFields, SetterPath, ConversionFunction);
 			if (AddBindingResult.HasError())
 			{
 				AddErrorForBinding(Binding, BlueprintView, FString::Printf(TEXT("The binding could not be created. %s"), *AddBindingResult.GetError()));
diff --git a/Engine/Plugins/Runtime/ModelViewViewModel/Source/ModelViewViewModelBlueprint/Public/MVVMBlueprintViewBinding.h b/Engine/Plugins/Runtime/ModelViewViewModel/Source/ModelViewViewModelBlueprint/Public/MVVMBlueprintViewBinding.h
--- a/Engine/Plugins/Runtime/ModelViewViewModel/Source/ModelViewViewModelBlueprint/Public/MVVMBlueprintViewBinding.h
+++ b/Engine/Plugins/Runtime/ModelViewViewModel/Source/ModelViewViewModelBlueprint/Public/MVVMBlueprintViewBinding.h
@@ -13,24 +13,30 @@ class UMVVMBlueprintView;
 *
 */
 USTRUCT(BlueprintType)
-struct FMVVMBlueprintViewConversionPath
+struct MODELVIEWVIEWMODELBLUEPRINT_API FMVVMBlueprintViewConversionPath
 {
 	GENERATED_BODY()
 
 	/** The Conversion function when converting the value from the destination to the source. */
 	UPROPERTY(EditAnywhere, Category = "MVVM", AdvancedDisplay)
 	FMemberReference DestinationToSourceFunction;
+	
+	UPROPERTY()
+	FName DestinationToSourceWrapper;
 
 	/** The Conversion function when converting the value from the source to the destination. */
 	UPROPERTY(EditAnywhere, Category = "MVVM", AdvancedDisplay)
 	FMemberReference SourceToDestinationFunction;
+
+	UPROPERTY()
+	FName SourceToDestinationWrapper;
 };
 
 /**
 *
 */
 USTRUCT(BlueprintType)
-struct FMVVMBlueprintViewBinding
+struct MODELVIEWVIEWMODELBLUEPRINT_API FMVVMBlueprintViewBinding
 {
 	GENERATED_BODY()
 
@@ -63,9 +69,14 @@ struct FMVVMBlueprintViewBinding
 	UPROPERTY(EditAnywhere, Category = "MVVM")
 	bool bCompile = true;
 
+	/**
+	 * Get an internal name. For use in the UI, use GetDisplayNameString()
+	 */
+	FName GetFName(const UMVVMBlueprintView* View) const;
+
 	/** 
 	 * Get a string that identifies this binding. 
 	 * This is of the form: ViewModel.Property -> Widget.Property
 	 */
-	FString GetNameString(const UMVVMBlueprintView* View) const;
+	FString GetDisplayNameString(const UMVVMBlueprintView* View) const;
 };
diff --git a/Engine/Plugins/Runtime/ModelViewViewModel/Source/ModelViewViewModelEditor/ModelViewViewModelEditor.Build.cs b/Engine/Plugins/Runtime/ModelViewViewModel/Source/ModelViewViewModelEditor/ModelViewViewModelEditor.Build.cs
--- a/Engine/Plugins/Runtime/ModelViewViewModel/Source/ModelViewViewModelEditor/ModelViewViewModelEditor.Build.cs
+++ b/Engine/Plugins/Runtime/ModelViewViewModel/Source/ModelViewViewModelEditor/ModelViewViewModelEditor.Build.cs
@@ -24,9 +24,9 @@ public class ModelViewViewModelEditor : ModuleRules
 			new string[]
 			{
 				"AssetTools",
-				
 				"EditorWidgets",
 				"BlueprintGraph",
+				"GraphEditor",
 				"InputCore",
 				"Kismet",
 				"MessageLog",
diff --git a/Engine/Plugins/Runtime/ModelViewViewModel/Source/ModelViewViewModelEditor/Private/Customizations/MVVMConversionPathCustomization.cpp b/Engine/Plugins/Runtime/ModelViewViewModel/Source/ModelViewViewModelEditor/Private/Customizations/MVVMConversionPathCustomization.cpp
--- a/Engine/Plugins/Runtime/ModelViewViewModel/Source/ModelViewViewModelEditor/Private/Customizations/MVVMConversionPathCustomization.cpp
+++ b/Engine/Plugins/Runtime/ModelViewViewModel/Source/ModelViewViewModelEditor/Private/Customizations/MVVMConversionPathCustomization.cpp
@@ -1,177 +1,643 @@
 // Copyright Epic Games, Inc. All Rights Reserved.
 
 #include "MVVMConversionPathCustomization.h"
+
+#include "Algo/Transform.h"
+#include "Bindings/MVVMBindingHelper.h"
+#include "Blueprint/WidgetTree.h"
+#include "Containers/Deque.h"
 #include "DetailWidgetRow.h"
+#include "EdGraphSchema_K2.h"
+#include "EdGraph/EdGraph.h"
+#include "EdGraph/EdGraphNode.h"
+#include "EdGraph/EdGraphPin.h"
+#include "Editor.h"
 #include "IDetailChildrenBuilder.h"
+#include "IDetailGroup.h"
+#include "IPropertyUtilities.h"
+#include "K2Node_CallFunction.h"
+#include "K2Node_VariableGet.h"
 #include "MVVMBlueprintViewBinding.h"
+#include "MVVMEditorSubsystem.h"
+#include "MVVMSubsystem.h"
+#include "NodeFactory.h"
 #include "PropertyHandle.h"
+#include "ScopedTransaction.h"
+#include "SGraphPin.h"
+#include "Styling/MVVMEditorStyle.h"
+#include "UObject/StructOnScope.h"
 #include "WidgetBlueprint.h"
-#include "Widgets/SMVVMConversionPath.h"
+#include "Widgets/Images/SImage.h"
+#include "Widgets/Input/SCheckBox.h"
 #include "Widgets/Input/SEditableTextBox.h"
+#include "Widgets/Layout/SBox.h"
+#include "Widgets/SMVVMConversionPath.h"
+#include "Widgets/SMVVMFieldSelector.h"
+#include "Widgets/SMVVMSourceSelector.h"
 #include "Widgets/Text/STextBlock.h"
 
 #define LOCTEXT_NAMESPACE "MVVMConversionPath"
 
-UE::MVVM::FConversionPathCustomization::FConversionPathCustomization(UWidgetBlueprint* InWidgetBlueprint)
-{
-	WidgetBlueprint = InWidgetBlueprint;
-}
-
-void UE::MVVM::FConversionPathCustomization::CustomizeHeader(TSharedRef<IPropertyHandle> InPropertyHandle, FDetailWidgetRow& HeaderRow, IPropertyTypeCustomizationUtils& CustomizationUtils)
+namespace UE::MVVM
 {
-	SourceToDestinationProperty = InPropertyHandle->GetChildHandle(GET_MEMBER_NAME_CHECKED(FMVVMBlueprintViewConversionPath, SourceToDestinationFunction));
-	DestinationToSourceProperty = InPropertyHandle->GetChildHandle(GET_MEMBER_NAME_CHECKED(FMVVMBlueprintViewConversionPath, DestinationToSourceFunction));
-
-	HeaderRow.NameContent()
-	[
-		InPropertyHandle->CreatePropertyNameWidget()
-	]
-	.ValueContent()
-	[
-		InPropertyHandle->CreatePropertyValueWidget()
-	];
-}
 
-FText UE::MVVM::FConversionPathCustomization::GetSourceToDestinationPath() const
-{
-	void* Value;
-	FPropertyAccess::Result Result = SourceToDestinationProperty->GetValueData(Value);
-
-	if (Result == FPropertyAccess::Success)
+	FConversionPathCustomization::FConversionPathCustomization(UWidgetBlueprint* InWidgetBlueprint)
 	{
-		return FText::FromName(reinterpret_cast<FMemberReference*>(Value)->GetMemberName());
+		check(InWidgetBlueprint);
+		WidgetBlueprint = InWidgetBlueprint;
 	}
 
-	if (Result == FPropertyAccess::MultipleValues)
+	void FConversionPathCustomization::CustomizeHeader(TSharedRef<IPropertyHandle> InPropertyHandle, FDetailWidgetRow& HeaderRow, IPropertyTypeCustomizationUtils& CustomizationUtils)
 	{
-		return LOCTEXT("MultipleValues", "Multiple Values");
-	}
+		WeakPropertyUtilities = CustomizationUtils.GetPropertyUtilities();
 
-	return FText::GetEmpty();
-}
+		InPropertyHandle->SetOnPropertyValueChanged(FSimpleDelegate::CreateSP(this, &FConversionPathCustomization::RefreshDetailsView));
+		InPropertyHandle->SetOnChildPropertyValueChanged(FSimpleDelegate::CreateSP(this, &FConversionPathCustomization::RefreshDetailsView));
 
-FText UE::MVVM::FConversionPathCustomization::GetDestinationToSourcePath() const
-{
-	void* Value;
-	FPropertyAccess::Result Result = DestinationToSourceProperty->GetValueData(Value);
+		ParentHandle = InPropertyHandle->GetParentHandle();
+		BindingModeHandle = ParentHandle->GetChildHandle(GET_MEMBER_NAME_CHECKED(FMVVMBlueprintViewBinding, BindingType));
+
+		HeaderRow
+			.ShouldAutoExpand()
+			.NameContent()
+			[
+				InPropertyHandle->CreatePropertyNameWidget()
+			]
+			.ValueContent()
+			[
+				InPropertyHandle->CreatePropertyValueWidget()
+			];
+	}
 
-	if (Result == FPropertyAccess::Success)
+	void FConversionPathCustomization::CustomizeChildren(TSharedRef<IPropertyHandle> InPropertyHandle, IDetailChildrenBuilder& ChildBuilder, IPropertyTypeCustomizationUtils& CustomizationUtils)
 	{
-		return FText::FromName(reinterpret_cast<FMemberReference*>(Value)->GetMemberName());
+		ArgumentPinWidgets.Reset();
+
+		AddRowForProperty(ChildBuilder, InPropertyHandle, true);
+		AddRowForProperty(ChildBuilder, InPropertyHandle, false);
 	}
 
-	if (Result == FPropertyAccess::MultipleValues)
+	void FConversionPathCustomization::RefreshDetailsView() const
 	{
-		return LOCTEXT("MultipleValues", "Multiple Values");
+		if (TSharedPtr<IPropertyUtilities> PropertyUtilities = WeakPropertyUtilities.Pin())
+		{
+			PropertyUtilities->RequestRefresh();
+		}
 	}
 
-	return FText::GetEmpty();
-}
+	FText FConversionPathCustomization::GetFunctionPathText(TSharedRef<IPropertyHandle> Property) const
+	{
+		void* Value;
+		FPropertyAccess::Result Result = Property->GetValueData(Value);
 
-void UE::MVVM::FConversionPathCustomization::OnTextCommitted(const FText& NewValue, ETextCommit::Type CommitType, bool bSourceToDestination)
-{
-	FString NewString = NewValue.ToString();
+		if (Result == FPropertyAccess::Success)
+		{
+			FMemberReference* MemberReference = reinterpret_cast<FMemberReference*>(Value);
+			UFunction* Function = MemberReference->ResolveMember<UFunction>(WidgetBlueprint);
 
-	UFunction* FoundFunction = nullptr;
-	if (WidgetBlueprint)
-	{
-		FoundFunction = WidgetBlueprint->GetBlueprintClass()->FindFunctionByName(*NewString);
+			return FText::FromString(Function->GetPathName());
+		}
+
+		if (Result == FPropertyAccess::MultipleValues)
+		{
+			return LOCTEXT("MultipleValues", "Multiple Values");
+		}
+
+		return FText::GetEmpty();
 	}
-	if (FoundFunction == nullptr)
+
+	void FConversionPathCustomization::OnTextCommitted(const FText& NewValue, ETextCommit::Type CommitType, TSharedRef<IPropertyHandle> Property, bool bSourceToDestination)
 	{
-		FoundFunction = FindObject<UFunction>(nullptr, *NewString, true);
+		UFunction* FoundFunction = FindObject<UFunction>(nullptr, *NewValue.ToString(), true);
+		OnFunctionPathChanged(FoundFunction, Property, bSourceToDestination);
 	}
-	OnFunctionPathChanged(FoundFunction, bSourceToDestination);
-}
 
-void UE::MVVM::FConversionPathCustomization::OnFunctionPathChanged(const UFunction* NewFunction, bool bSourceToDestination)
-{
-	TSharedPtr<IPropertyHandle> Handle = bSourceToDestination ? SourceToDestinationProperty : DestinationToSourceProperty;
+	void FConversionPathCustomization::OnFunctionPathChanged(const UFunction* NewFunction, TSharedRef<IPropertyHandle> InPropertyHandle, bool bSourceToDestination)
+	{
+		UMVVMEditorSubsystem* EditorSubsystem = GEditor->GetEditorSubsystem<UMVVMEditorSubsystem>();
+		
+		TArray<void*> RawBindings;
+		ParentHandle->AccessRawData(RawBindings);
 
-	FStructProperty* StructProperty = CastFieldChecked<FStructProperty>(Handle->GetProperty());
-	check(StructProperty);
+		for (void* RawBinding : RawBindings)
+		{
+			FMVVMBlueprintViewBinding* ViewBinding = static_cast<FMVVMBlueprintViewBinding*>(RawBinding);
+			if (bSourceToDestination)
+			{
+				EditorSubsystem->SetSourceToDestinationConversionFunction(WidgetBlueprint, *ViewBinding, NewFunction);
+			}
+			else
+			{
+				EditorSubsystem->SetDestinationToSourceConversionFunction(WidgetBlueprint, *ViewBinding, NewFunction);
+			}
+		}
 
-	void* PreviousValue;
-	Handle->GetValueData(PreviousValue);
+		RefreshDetailsView();
+	}
 
-	FMemberReference NewReference;
-	if (NewFunction)
+	void FConversionPathCustomization::AddRowForProperty(IDetailChildrenBuilder& ChildBuilder, const TSharedPtr<IPropertyHandle>& InPropertyHandle, bool bSourceToDestination)
 	{
-		NewReference.SetFromField<UFunction>(NewFunction, WidgetBlueprint ? WidgetBlueprint->SkeletonGeneratedClass : nullptr);
-	}
+		TArray<FMVVMBlueprintViewBinding*> ViewBindings;
+		{
+			TArray<void*> RawData;
+			ParentHandle->AccessRawData(RawData);
 
-	FString TextValue;
-	StructProperty->Struct->ExportText(TextValue, &NewReference, PreviousValue, nullptr, EPropertyPortFlags::PPF_None, nullptr);
-	ensure(Handle->SetValueFromFormattedString(TextValue, EPropertyValueSetFlags::DefaultFlags) == FPropertyAccess::Result::Success);
-}
+			for (void* Data : RawData)
+			{
+				ViewBindings.Add((FMVVMBlueprintViewBinding*) Data);
+			}
+		}
 
-void UE::MVVM::FConversionPathCustomization::CustomizeChildren(TSharedRef<IPropertyHandle> InPropertyHandle, IDetailChildrenBuilder& ChildBuilder, IPropertyTypeCustomizationUtils& CustomizationUtils)
-{
-	TSharedPtr<IPropertyHandle> ParentHandle = InPropertyHandle->GetParentHandle();
-	TArray<void*> RawData;
-	ParentHandle->AccessRawData(RawData);
-
-	TArray<FMVVMBlueprintViewBinding*> ViewBindings;
-	for (void* Data : RawData)
-	{
-		ViewBindings.Add((FMVVMBlueprintViewBinding*) Data);
-	}
-
-	// setter
-	ChildBuilder.AddProperty(DestinationToSourceProperty.ToSharedRef())
-		.CustomWidget()
-		.NameContent()
-		[
-			SNew(STextBlock)
-			.Text(LOCTEXT("SetterLabel", "Setter"))
-		]
-		.ValueContent()
-		[
-			SNew(SHorizontalBox)
-			+ SHorizontalBox::Slot()
-			.VAlign(VAlign_Center)
+		const FName FunctionPropertyName = bSourceToDestination ? GET_MEMBER_NAME_CHECKED(FMVVMBlueprintViewConversionPath, SourceToDestinationFunction) : GET_MEMBER_NAME_CHECKED(FMVVMBlueprintViewConversionPath, DestinationToSourceFunction);
+		TSharedRef<IPropertyHandle> FunctionProperty = InPropertyHandle->GetChildHandle(FunctionPropertyName).ToSharedRef();
+
+		ChildBuilder.AddProperty(FunctionProperty)
+			.IsEnabled(
+				TAttribute<bool>::CreateLambda([this, bSourceToDestination]()
+					{
+						uint8 EnumValue;
+						if (BindingModeHandle->GetValue(EnumValue) == FPropertyAccess::Success)
+						{
+							return bSourceToDestination ? IsForwardBinding((EMVVMBindingMode) EnumValue) : IsBackwardBinding((EMVVMBindingMode) EnumValue);
+						}
+						return true;
+					}))
+			.CustomWidget()
+			.NameContent()
 			[
-				SNew(SEditableTextBox)
-				.Text(this, &FConversionPathCustomization::GetDestinationToSourcePath)
-				.OnTextCommitted(this, &FConversionPathCustomization::OnTextCommitted, false)
+				FunctionProperty->CreatePropertyNameWidget()
 			]
-			+ SHorizontalBox::Slot()
-			.VAlign(VAlign_Center)
-			.AutoWidth()
+			.ValueContent()
 			[
-				SNew(SMVVMConversionPath, WidgetBlueprint, false)
-				.Bindings(ViewBindings)
-				.OnFunctionChanged(this, &FConversionPathCustomization::OnFunctionPathChanged, false)
-			]
-		];
-
-	// getter
-	ChildBuilder.AddProperty(SourceToDestinationProperty.ToSharedRef())
-		.CustomWidget()
-		.NameContent()
-		[
-			SNew(STextBlock)
-			.Text(LOCTEXT("GetterLabel", "Getter"))
-		]
-		.ValueContent()
-		[
-			SNew(SHorizontalBox)
-			+ SHorizontalBox::Slot()
-			.VAlign(VAlign_Center)
+				SNew(SHorizontalBox)
+				+ SHorizontalBox::Slot()
+				[
+					SNew(SEditableTextBox)
+					.Text(this, &FConversionPathCustomization::GetFunctionPathText, FunctionProperty)
+					.OnTextCommitted(this, &FConversionPathCustomization::OnTextCommitted, FunctionProperty, bSourceToDestination)
+				]
+				+ SHorizontalBox::Slot()
+				.AutoWidth()
+				[
+					SNew(SMVVMConversionPath, WidgetBlueprint, bSourceToDestination)
+					.Bindings(ViewBindings)
+					.OnFunctionChanged(this, &FConversionPathCustomization::OnFunctionPathChanged, FunctionProperty, bSourceToDestination)
+				]
+			];
+
+		if (ViewBindings.Num() != 1)
+		{
+			// don't show arguments for multi-selection
+			return;
+		}
+
+		UMVVMEditorSubsystem* EditorSubsystem = GEditor->GetEditorSubsystem<UMVVMEditorSubsystem>();
+
+		UEdGraph* Graph = EditorSubsystem->GetConversionFunctionGraph(WidgetBlueprint, *ViewBindings[0], bSourceToDestination);
+		if (Graph == nullptr)
+		{
+			return;
+		}
+
+		TArray<UK2Node_CallFunction*> FunctionNodes;
+		Graph->GetNodesOfClass<UK2Node_CallFunction>(FunctionNodes);
+
+		if (FunctionNodes.Num() != 1)
+		{
+			// ambiguous result, no idea what our function node is
+			return;
+		}
+
+		UK2Node_CallFunction* CallFunctionNode = FunctionNodes[0];
+		for (UEdGraphPin* Pin : CallFunctionNode->Pins)
+		{
+			// skip output and exec pins, the rest should be arguments
+			if (Pin->Direction != EGPD_Input || Pin->PinName == UEdGraphSchema_K2::PN_Execute || Pin->PinName == UEdGraphSchema_K2::PN_Self)
+			{
+				continue;
+			}
+
+			const FName ArgumentName = Pin->GetFName();
+
+			// create a new pin widget so that we can get the default value widget out of it
+			TSharedPtr<SGraphPin> PinWidget = FNodeFactory::CreateK2PinWidget(Pin);
+
+			TSharedRef<SWidget> ValueWidget = SNullWidget::NullWidget;
+			if (PinWidget.IsValid())
+			{
+				ArgumentPinWidgets.Add(PinWidget);
+				ValueWidget = PinWidget->GetDefaultValueWidget();
+			}
+
+			if (ValueWidget == SNullWidget::NullWidget)
+			{
+				ValueWidget = SNew(STextBlock)
+					.Text(LOCTEXT("DefaultValue", "Default Value"))
+					.TextStyle(FAppStyle::Get(), "HintText");
+			}
+
+			ValueWidget->SetVisibility(TAttribute<EVisibility>::CreateSP(this, &FConversionPathCustomization::GetArgumentWidgetVisibility, ArgumentName, bSourceToDestination, true));
+
+			TSharedPtr<SMVVMFieldSelector> FieldSelector;
+
+			ChildBuilder.AddCustomRow(Pin->GetDisplayName())
+			.NameContent()
 			[
-				SNew(SEditableTextBox)
-				.Text(this, &FConversionPathCustomization::GetSourceToDestinationPath)
-				.OnTextCommitted(this, &FConversionPathCustomization::OnTextCommitted, true)
+				SNew(SBox)
+				.Padding(FMargin(16, 0, 0, 0))
+				[
+					SNew(STextBlock)
+					.Text(Pin->GetDisplayName())
+				]
 			]
-			+ SHorizontalBox::Slot()
-			.VAlign(VAlign_Center)
-			.AutoWidth()
+			.ValueContent()
 			[
-				SNew(SMVVMConversionPath, WidgetBlueprint, true)
-				.Bindings(ViewBindings)
-				.OnFunctionChanged(this, &FConversionPathCustomization::OnFunctionPathChanged, true)
-			]
-		];
+				// TODO (sebastiann): This would be easier to accomplish if this entire thing was a widget
+				SNew(SHorizontalBox)
+				+ SHorizontalBox::Slot()
+				[
+					SNew(SOverlay)
+					+ SOverlay::Slot()
+					[
+						ValueWidget
+					]
+					+ SOverlay::Slot()
+					[
+						SNew(SHorizontalBox)
+						.Visibility(this, &FConversionPathCustomization::GetArgumentWidgetVisibility, ArgumentName, bSourceToDestination, false)
+						+ SHorizontalBox::Slot()
+						[
+							SNew(SMVVMSourceSelector)
+							.SelectedSource(this, &FConversionPathCustomization::OnGetSelectedSource, ArgumentName, bSourceToDestination)
+							.AvailableSources(this, &FConversionPathCustomization::OnGetAvailableSources, bSourceToDestination)
+							.OnSelectionChanged(this, &FConversionPathCustomization::OnSetSource, ArgumentName, bSourceToDestination)
+						]
+						+ SHorizontalBox::Slot()
+						[
+							SAssignNew(FieldSelector, SMVVMFieldSelector)
+							.IsSource(bSourceToDestination)
+							.SelectedSource(this, &FConversionPathCustomization::OnGetSelectedSource, ArgumentName, bSourceToDestination)
+							.AvailableFields(this, &FConversionPathCustomization::OnGetAvailableFields, ArgumentName, bSourceToDestination)
+							.SelectedField(this, &FConversionPathCustomization::OnGetSelectedField, ArgumentName, bSourceToDestination)
+							.OnSelectionChanged(this, &FConversionPathCustomization::OnSetProperty, ArgumentName, bSourceToDestination)
+						]
+					]
+				]
+				+ SHorizontalBox::Slot()
+				.AutoWidth()
+				.HAlign(HAlign_Right)
+				.VAlign(VAlign_Center)
+				.Padding(5, 0, 0, 0)
+				[
+					SNew(SCheckBox)
+					.ToolTipText(LOCTEXT("BindArgument", "Bind this argument to a property."))
+					.Style(FAppStyle::Get(), "ToggleButtonCheckbox")
+					.IsChecked(this, &FConversionPathCustomization::OnGetIsArgumentBound, ArgumentName, bSourceToDestination)
+					.OnCheckStateChanged(this, &FConversionPathCustomization::OnCheckedBindArgument, ArgumentName, bSourceToDestination)
+					[
+						SNew(SBox)
+						.WidthOverride(16)
+						.HeightOverride(16)
+						[
+							SNew(SImage)
+							.Image_Lambda([this, ArgumentName, bSourceToDestination]()
+								{
+									ECheckBoxState CheckState = OnGetIsArgumentBound(ArgumentName, bSourceToDestination);
+									return CheckState == ECheckBoxState::Checked ? FAppStyle::GetBrush("Icons.Link") : FAppStyle::GetBrush("Icons.Unlink");
+								})
+						]
+					]
+				]
+			];
+
+			ArgumentFieldSelectors.Add(FieldSelector);
+		}
+	}
+
+	TArray<FBindingSource> FConversionPathCustomization::OnGetAvailableSources(bool bSourceToDestination) const
+	{
+		UMVVMEditorSubsystem* Subsystem = GEditor->GetEditorSubsystem<UMVVMEditorSubsystem>();
+		if (bSourceToDestination)
+		{
+			return Subsystem->GetAllViewModels(WidgetBlueprint);
+		}
+		else
+		{
+			return Subsystem->GetBindableWidgets(WidgetBlueprint);
+		}
+	}
+
+	FBindingSource FConversionPathCustomization::OnGetSelectedSource(FName ArgumentName, bool bSourceToDestination) const
+	{
+		UMVVMEditorSubsystem* Subsystem = GEditor->GetEditorSubsystem<UMVVMEditorSubsystem>();
+
+		TArray<void*> RawBindings;
+		ParentHandle->AccessRawData(RawBindings);
+
+		FMVVMBlueprintPropertyPath Path;
+		bool bFirst = true;
+
+		for (void* RawBinding : RawBindings)
+		{
+			FMVVMBlueprintViewBinding* Binding = reinterpret_cast<FMVVMBlueprintViewBinding*>(RawBinding);
+			FMVVMBlueprintPropertyPath ThisPath = Subsystem->GetPathForConversionFunctionArgument(WidgetBlueprint, *Binding, ArgumentName, bSourceToDestination);
+
+			if (bFirst)
+			{
+				Path = ThisPath;
+				bFirst = false;
+			}
+			else if (Path != ThisPath)
+			{
+				Path = FMVVMBlueprintPropertyPath();
+			}
+		}
+		
+		if (Path.IsFromViewModel())
+		{
+			return FBindingSource::CreateForViewModel(WidgetBlueprint, Path.GetViewModelId());
+		}
+		else if (Path.IsFromWidget())
+		{
+			return FBindingSource::CreateForWidget(WidgetBlueprint, Path.GetWidgetName());
+		}
+
+		return FBindingSource();
+	}
+
+	void FConversionPathCustomization::OnSetSource(FBindingSource Source, FName ArgumentName, bool bSourceToDestination)
+	{
+		FMVVMBlueprintPropertyPath Path;
+		if (Source.ViewModelId.IsValid())
+		{
+			Path.SetViewModelId(Source.ViewModelId);
+		}
+		else
+		{
+			Path.SetWidgetName(Source.Name);
+		}
+
+		UMVVMEditorSubsystem* Subsystem = GEditor->GetEditorSubsystem<UMVVMEditorSubsystem>();
+
+		TArray<void*> RawBindings;
+		ParentHandle->AccessRawData(RawBindings);
+
+		for (void* RawBinding : RawBindings)
+		{
+			FMVVMBlueprintViewBinding* Binding = reinterpret_cast<FMVVMBlueprintViewBinding*>(RawBinding);
+			Subsystem->SetPathForConversionFunctionArgument(WidgetBlueprint, *Binding, ArgumentName, Path, bSourceToDestination);
+		}
+
+		for (const TSharedPtr<SMVVMFieldSelector>& FieldSelector : ArgumentFieldSelectors)
+		{
+			FieldSelector->Refresh();
+		}
+	}
+
+	void FConversionPathCustomization::OnSetProperty(FMVVMBlueprintPropertyPath NewSelection, FName ArgumentName, bool bSourceToDestination)
+	{
+		UMVVMEditorSubsystem* EditorSubsystem = GEditor->GetEditorSubsystem<UMVVMEditorSubsystem>();
+
+		TArray<void*> RawBindings;
+		ParentHandle->AccessRawData(RawBindings);
+
+		for (void* RawBinding : RawBindings)
+		{
+			FMVVMBlueprintViewBinding* Binding = reinterpret_cast<FMVVMBlueprintViewBinding*>(RawBinding);
+			EditorSubsystem->SetPathForConversionFunctionArgument(WidgetBlueprint, *Binding, ArgumentName, NewSelection, bSourceToDestination);
+		}
+	}
+
+	FMVVMBlueprintPropertyPath FConversionPathCustomization::OnGetSelectedField(FName ArgumentName, bool bSourceToDestination) const
+	{
+		UMVVMEditorSubsystem* EditorSubsystem = GEditor->GetEditorSubsystem<UMVVMEditorSubsystem>();
+		
+		TArray<void*> RawBindings;
+		ParentHandle->AccessRawData(RawBindings);
+
+		FMVVMBlueprintPropertyPath Path;
+		bool bFirst = true;
+
+		for (void* RawBinding : RawBindings)
+		{
+			FMVVMBlueprintViewBinding* Binding = reinterpret_cast<FMVVMBlueprintViewBinding*>(RawBinding);
+			
+			FMVVMBlueprintPropertyPath CurrentPath = EditorSubsystem->GetPathForConversionFunctionArgument(WidgetBlueprint, *Binding, ArgumentName, bSourceToDestination);
+			if (bFirst)
+			{
+				Path = CurrentPath;
+				bFirst = false;
+			}
+			else if (Path != CurrentPath)
+			{
+				return FMVVMBlueprintPropertyPath();
+			}
+		}
+
+		return Path;
+	}
+
+	TArray<FMVVMBlueprintPropertyPath> FConversionPathCustomization::OnGetAvailableFields(FName ArgumentName, bool bSourceToDestination) const
+	{
+		TArray<FMVVMBlueprintPropertyPath> AvailablePaths;
+
+		FBindingSource Source = OnGetSelectedSource(ArgumentName, bSourceToDestination);
+		UClass* SourceClass = nullptr;
+		if (!Source.Name.IsNone())
+		{
+			// widget
+			if (Source.Name.IsEqual(WidgetBlueprint->GetFName()))
+			{
+				SourceClass = WidgetBlueprint->GeneratedClass;
+			}
+			else if (const UWidget* Widget = WidgetBlueprint->WidgetTree->FindWidget(Source.Name))
+			{
+				SourceClass = Widget->GetClass();
+			}
+		}
+		else if (Source.ViewModelId.IsValid())
+		{
+			// viewmodel
+			UMVVMEditorSubsystem* EditorSubsystem = GEditor->GetEditorSubsystem<UMVVMEditorSubsystem>();
+			if (UMVVMBlueprintView* View = EditorSubsystem->GetView(WidgetBlueprint))
+			{
+				if (const FMVVMBlueprintViewModelContext* ViewModel = View->FindViewModel(Source.ViewModelId))
+				{
+					SourceClass = ViewModel->GetViewModelClass();
+				}
+			}
+		}
+
+		if (SourceClass == nullptr)
+		{
+			return AvailablePaths;
+		}
+
+		UMVVMSubsystem* Subsystem = GEngine->GetEngineSubsystem<UMVVMSubsystem>();
+		TArray<FMVVMAvailableBinding> AvailableBindings = Source.ViewModelId.IsValid() ? Subsystem->GetViewModelAvailableBindings(SourceClass) : Subsystem->GetWidgetAvailableBindings(SourceClass);
+
+		const UEdGraphSchema_K2* Schema = GetDefault<UEdGraphSchema_K2>();
+		FEdGraphPinType ArgumentType = GetArgumentPinType(ArgumentName, bSourceToDestination);
+
+		Algo::TransformIf(AvailableBindings, AvailablePaths, 
+			// filter handler
+			[ArgumentType, Schema, SourceClass](const FMVVMAvailableBinding& Binding) -> bool
+			{
+				FName BindingName = Binding.GetBindingName().ToName();
+				
+				if (const FProperty* Property = SourceClass->FindPropertyByName(BindingName))
+				{
+					return true;
+				}
+
+				if (const UFunction* Function = SourceClass->FindFunctionByName(BindingName))
+				{
+					const FProperty* ReturnProperty = UE::MVVM::BindingHelper::GetReturnProperty(Function);
+					if (ReturnProperty != nullptr)
+					{
+						FEdGraphPinType ReturnType;
+						if (Schema->ConvertPropertyToPinType(ReturnProperty, ReturnType))
+						{
+							if (Schema->ArePinTypesCompatible(ReturnType, ArgumentType, SourceClass))
+							{
+								return true;
+							}
+						}
+					}
+				}
+
+				return false;
+			},
+			[SourceClass, Source](const FMVVMAvailableBinding& Binding) -> FMVVMBlueprintPropertyPath
+			{
+				FName BindingName = Binding.GetBindingName().ToName();
+
+				UE::MVVM::FMVVMConstFieldVariant Variant;
+				if (const UFunction* Function = SourceClass->FindFunctionByName(BindingName))
+				{
+					Variant = UE::MVVM::FMVVMConstFieldVariant(Function);
+				}
+				else if (const FProperty* Property = SourceClass->FindPropertyByName(BindingName))
+				{
+					Variant = UE::MVVM::FMVVMConstFieldVariant(Property);
+				}
+
+				FMVVMBlueprintPropertyPath Path;
+				if (!Source.Name.IsNone())
+				{
+					Path.SetWidgetName(Source.Name);
+				}
+				else
+				{
+					Path.SetViewModelId(Source.ViewModelId);
+				}
+				Path.SetBasePropertyPath(Variant);
+
+				return Path;
+			});
+
+		return AvailablePaths;
+	}
+
+	ECheckBoxState FConversionPathCustomization::OnGetIsArgumentBound(FName ArgumentName, bool bSourceToDestination) const
+	{
+		UMVVMEditorSubsystem* EditorSubsystem = GEditor->GetEditorSubsystem<UMVVMEditorSubsystem>();
+
+		TArray<void*> RawBindings;
+		ParentHandle->AccessRawData(RawBindings);
+
+		FMVVMBlueprintPropertyPath Path;
+		bool bFirst = true;
+
+		for (void* RawBinding : RawBindings)
+		{
+			FMVVMBlueprintViewBinding* Binding = reinterpret_cast<FMVVMBlueprintViewBinding*>(RawBinding);
+
+			FMVVMBlueprintPropertyPath CurrentPath = EditorSubsystem->GetPathForConversionFunctionArgument(WidgetBlueprint, *Binding, ArgumentName, bSourceToDestination);
+			if (bFirst)
+			{
+				Path = CurrentPath;
+				bFirst = false;
+			}
+			else if (Path != CurrentPath)
+			{
+				return ECheckBoxState::Undetermined;
+			}
+		}
+
+		return Path.IsEmpty() ? ECheckBoxState::Unchecked : ECheckBoxState::Checked;
+	}
+
+	void FConversionPathCustomization::OnCheckedBindArgument(ECheckBoxState CheckState, FName ArgumentName, bool bSourceToDestination)
+	{
+		UMVVMEditorSubsystem* EditorSubsystem = GEditor->GetEditorSubsystem<UMVVMEditorSubsystem>();
+
+		FMVVMBlueprintPropertyPath Path;
+
+		if (CheckState == ECheckBoxState::Checked)
+		{
+			// HACK: Just set a placeholder viewmodel reference
+			if (const UMVVMBlueprintView* View = EditorSubsystem->GetView(WidgetBlueprint))
+			{
+				const TArrayView<const FMVVMBlueprintViewModelContext> ViewModels = View->GetViewModels();
+				if (ViewModels.Num() > 0)
+				{
+					Path.SetViewModelId(ViewModels[0].GetViewModelId());
+				}
+			}
+		}
+
+		TArray<void*> RawBindings;
+		ParentHandle->AccessRawData(RawBindings);
+
+		for (void* RawBinding : RawBindings)
+		{
+			FMVVMBlueprintViewBinding* Binding = reinterpret_cast<FMVVMBlueprintViewBinding*>(RawBinding);
+			EditorSubsystem->SetPathForConversionFunctionArgument(WidgetBlueprint, *Binding, ArgumentName, Path, bSourceToDestination);
+		}
+	}
+
+	EVisibility FConversionPathCustomization::GetArgumentWidgetVisibility(FName ArgumentName, bool bSourceToDestination, bool bDefaultValue) const
+	{
+		FMVVMBlueprintPropertyPath Path = OnGetSelectedField(ArgumentName, bSourceToDestination);
+		if (Path.IsEmpty())
+		{
+			return bDefaultValue ? EVisibility::Visible : EVisibility::Collapsed;
+		}
+		else
+		{
+			return bDefaultValue ? EVisibility::Collapsed : EVisibility::Visible;
+		}
+	}
+
+	FEdGraphPinType FConversionPathCustomization::GetArgumentPinType(FName ArgumentName, bool bSourceToDestination) const 
+	{
+		TArray<void*> RawData;
+		ParentHandle->AccessRawData(RawData);
+		if (RawData.Num() != 1)
+		{
+			return FEdGraphPinType();
+		}
+
+		UMVVMEditorSubsystem* EditorSubsystem = GEditor->GetEditorSubsystem<UMVVMEditorSubsystem>();
+		UEdGraph* Graph = EditorSubsystem->GetConversionFunctionGraph(WidgetBlueprint, *reinterpret_cast<const FMVVMBlueprintViewBinding*>(RawData[0]), bSourceToDestination);
+		if (Graph == nullptr)
+		{
+			return FEdGraphPinType();
+		}
+
+		TArray<UK2Node_CallFunction*> FunctionNodes;
+		Graph->GetNodesOfClass<UK2Node_CallFunction>(FunctionNodes);
+		if (FunctionNodes.Num() != 1)
+		{
+			// ambiguous result, no idea what our function node is
+			return FEdGraphPinType();
+		}
+
+		UK2Node_CallFunction* CallFunctionNode = FunctionNodes[0];
+		UEdGraphPin* Pin = CallFunctionNode->FindPin(ArgumentName);
+		return Pin != nullptr ? Pin->PinType : FEdGraphPinType();
+	}
 }
 
-#undef LOCTEXT_NAMESPACE
\ No newline at end of file
+#undef LOCTEXT_NAMESPACE 
diff --git a/Engine/Plugins/Runtime/ModelViewViewModel/Source/ModelViewViewModelEditor/Private/Customizations/MVVMConversionPathCustomization.h b/Engine/Plugins/Runtime/ModelViewViewModel/Source/ModelViewViewModelEditor/Private/Customizations/MVVMConversionPathCustomization.h
--- a/Engine/Plugins/Runtime/ModelViewViewModel/Source/ModelViewViewModelEditor/Private/Customizations/MVVMConversionPathCustomization.h
+++ b/Engine/Plugins/Runtime/ModelViewViewModel/Source/ModelViewViewModelEditor/Private/Customizations/MVVMConversionPathCustomization.h
@@ -3,7 +3,17 @@
 #pragma once
 
 #include "IPropertyTypeCustomization.h"
+#include "Layout/Visibility.h"
+#include "Styling/SlateTypes.h"
+#include "Templates/ValueOrError.h"
+#include "Types/MVVMBindingSource.h"
+#include "Types/MVVMFieldVariant.h"
+#include "MVVMPropertyPath.h"
 
+struct FMVVMPropertyPathBase;
+class FStructOnScope;
+class SGraphPin;
+class SMVVMFieldSelector;
 class UWidgetBlueprint;
 
 namespace UE::MVVM
@@ -22,15 +32,35 @@ namespace UE::MVVM
 		virtual void CustomizeChildren(TSharedRef<IPropertyHandle> InPropertyHandle, IDetailChildrenBuilder& ChildBuilder, IPropertyTypeCustomizationUtils& CustomizationUtils) override;
 
 	private:
-		FText GetSourceToDestinationPath() const;
-		FText GetDestinationToSourcePath() const;
+		FText GetFunctionPathText(TSharedRef<IPropertyHandle> Property) const;
+		void AddRowForProperty(IDetailChildrenBuilder& ChildBuilder, const TSharedPtr<IPropertyHandle>& Property, bool bSourceToDestination);
 
-		void OnTextCommitted(const FText& NewValue, ETextCommit::Type CommitType, bool bSourceToDestination);
-		void OnFunctionPathChanged(const UFunction* NewPath, bool bSourceToDestination);
+		void OnTextCommitted(const FText& NewValue, ETextCommit::Type CommitType, TSharedRef<IPropertyHandle> Property, bool bSourceToDestination);
+		void OnFunctionPathChanged(const UFunction* NewFunction, TSharedRef<IPropertyHandle> Property, bool bSourceToDestination);
+		void RefreshDetailsView() const;
+
+		TArray<FBindingSource> OnGetAvailableSources(bool bSourceToDestination) const;
+		FBindingSource OnGetSelectedSource(FName ArgumentName, bool bSourceToDestination) const;
+		void OnSetSource(FBindingSource Source, FName ArgumentName, bool bSourceToDestination);
+
+		TArray<FMVVMBlueprintPropertyPath> OnGetAvailableFields(FName ArgumentName, bool bSourceToDestination) const;
+		FMVVMBlueprintPropertyPath OnGetSelectedField(FName ArgumentName, bool bSourceToDestination) const;
+		void OnSetProperty(FMVVMBlueprintPropertyPath NewSelection, FName ArgumentName, bool bSourceToDestination);
+
+		ECheckBoxState OnGetIsArgumentBound(FName ArgumentName, bool bSourceToDestination) const;
+		void OnCheckedBindArgument(ECheckBoxState CheckState, FName ArgumentName, bool bSourceToDestination);
+		EVisibility GetArgumentWidgetVisibility(FName ArgumentName, bool bSourceToDestination, bool bDefaultValue) const;
+
+		FEdGraphPinType GetArgumentPinType(FName ArgumentName, bool bSourceToDestination) const;
 
 	private:
 		UWidgetBlueprint* WidgetBlueprint = nullptr;
-		TSharedPtr<IPropertyHandle> DestinationToSourceProperty;
-		TSharedPtr<IPropertyHandle> SourceToDestinationProperty;
+
+		TSharedPtr<IPropertyHandle> ParentHandle;
+		TSharedPtr<IPropertyHandle> BindingModeHandle;
+		TWeakPtr<IPropertyUtilities> WeakPropertyUtilities;
+
+		TArray<TSharedPtr<SGraphPin>> ArgumentPinWidgets; 
+		TArray<TSharedPtr<SMVVMFieldSelector>> ArgumentFieldSelectors;
 	};
 }
\ No newline at end of file
diff --git a/Engine/Plugins/Runtime/ModelViewViewModel/Source/ModelViewViewModelEditor/Private/Customizations/MVVMPropertyPathCustomization.cpp b/Engine/Plugins/Runtime/ModelViewViewModel/Source/ModelViewViewModelEditor/Private/Customizations/MVVMPropertyPathCustomization.cpp
--- a/Engine/Plugins/Runtime/ModelViewViewModel/Source/ModelViewViewModelEditor/Private/Customizations/MVVMPropertyPathCustomization.cpp
+++ b/Engine/Plugins/Runtime/ModelViewViewModel/Source/ModelViewViewModelEditor/Private/Customizations/MVVMPropertyPathCustomization.cpp
@@ -7,9 +7,10 @@
 #include "BlueprintEditor.h"
 #include "Components/Widget.h"
 #include "DetailWidgetRow.h"
+#include "MVVMEditorSubsystem.h"
+#include "MVVMSubsystem.h"
 #include "MVVMWidgetBlueprintExtension_View.h"
 #include "PropertyHandle.h"
-#include "ScopedTransaction.h"
 #include "SSimpleButton.h"
 #include "Templates/UnrealTemplate.h"
 #include "WidgetBlueprint.h"
@@ -25,44 +26,6 @@
 
 namespace UE::MVVM
 {
-	namespace Private
-	{
-		// figure out what the selection is - if any is different in the case of a multiselection, then return null
-		// this is required for getting the selection of the other pair
-		TArray<TSharedPtr<IFieldPathHelper>> GetPathHelpersForProperty(const TSharedPtr<IPropertyHandle>& PropertyHandle, UWidgetBlueprint* WidgetBlueprint, bool bIsWidget)
-		{
-			TArray<TSharedPtr<IFieldPathHelper>> Result;
-
-			TArray<void*> RawData;
-			PropertyHandle->AccessRawData(RawData);
-
-			for (void* Datum : RawData)
-			{
-				const FProperty* ThisSelected = nullptr;
-				
-				if (bIsWidget)
-				{
-					Result.Add(MakeShared<FWidgetFieldPathHelper>(reinterpret_cast<FMVVMBlueprintPropertyPath*>(Datum), WidgetBlueprint));
-				}
-				else
-				{
-					Result.Add(MakeShared<FViewModelFieldPathHelper>(reinterpret_cast<FMVVMBlueprintPropertyPath*>(Datum), WidgetBlueprint));
-				}
-			}
-			return Result;
-		}
-
-		TArray<IFieldPathHelper*> GetRawPathHelperPointers(const TArray<TSharedPtr<IFieldPathHelper>>& SharedPtrs)
-		{
-			TArray<IFieldPathHelper*> RawPointers;
-			RawPointers.Reserve(SharedPtrs.Num());
-
-			Algo::Transform(SharedPtrs, RawPointers, [](const TSharedPtr<IFieldPathHelper>& X) { return X.Get(); });
-
-			return RawPointers;
-		}
-	}
-
 	FPropertyPathCustomization::FPropertyPathCustomization(UWidgetBlueprint* InWidgetBlueprint) :
 		WidgetBlueprint(InWidgetBlueprint)
 	{
@@ -80,8 +43,6 @@ namespace UE::MVVM
 	{
 		PropertyHandle = InPropertyHandle;
 
-		bool bIsWidget = false;
-
 		const FName PropertyName = PropertyHandle->GetProperty()->GetFName();
 		if (PropertyName == GET_MEMBER_NAME_CHECKED(FMVVMBlueprintViewBinding, WidgetPath))
 		{
@@ -91,8 +52,10 @@ namespace UE::MVVM
 		{
 			bIsWidget = false;
 		}
-
-		PathHelpers = Private::GetPathHelpersForProperty(PropertyHandle, WidgetBlueprint, bIsWidget);
+		else
+		{
+			ensureMsgf(false, TEXT("MVVMPropertyPathCustomization used in unknown context."));
+		}
 
 		TSharedPtr<IPropertyHandle> ParentHandle = PropertyHandle->GetParentHandle();
 		TSharedPtr<IPropertyHandle> OtherHandle;
@@ -109,7 +72,6 @@ namespace UE::MVVM
 		if (OtherHandle.IsValid() && OtherHandle->IsValidHandle())
 		{
 			OtherHandle->SetOnPropertyValueChanged(FSimpleDelegate::CreateSP(this, &FPropertyPathCustomization::OnOtherPropertyChanged));
-			OtherHelpers = Private::GetPathHelpersForProperty(OtherHandle, WidgetBlueprint, !bIsWidget);
 		}
 		else
 		{
@@ -145,7 +107,8 @@ namespace UE::MVVM
 				[
 					SAssignNew(SourceSelector, SMVVMSourceSelector)
 					.TextStyle(FAppStyle::Get(), "SmallText")
-					.PathHelpers(this, &FPropertyPathCustomization::GetRawPathHelpers)
+					.AvailableSources(this, &FPropertyPathCustomization::OnGetSources)
+					.SelectedSource(this, &FPropertyPathCustomization::OnGetSelectedSource)
 					.OnSelectionChanged(this, &FPropertyPathCustomization::OnSourceSelectionChanged)
 				]
 			+ SHorizontalBox::Slot()
@@ -154,24 +117,15 @@ namespace UE::MVVM
 				[
 					SAssignNew(FieldSelector, SMVVMFieldSelector)
 					.TextStyle(FAppStyle::Get(), "SmallText")
-					.PathHelpers(GetRawPathHelpers())
-					.CounterpartHelpers(GetRawOtherHelpers())
 					.BindingMode(this, &FPropertyPathCustomization::GetCurrentBindingMode)
+					.SelectedSource(this, &FPropertyPathCustomization::OnGetSelectedSource)
 					.OnSelectionChanged(this, &FPropertyPathCustomization::OnPropertySelectionChanged)
+					.AvailableFields(this, &FPropertyPathCustomization::OnGetFields)
+					.SelectedField(this, &FPropertyPathCustomization::OnGetSelectedField)
 				]				
 			];
 	}
 
-	TArray<IFieldPathHelper*> FPropertyPathCustomization::GetRawPathHelpers() const
-	{
-		return Private::GetRawPathHelperPointers(PathHelpers);
-	}
-
-	TArray<IFieldPathHelper*> FPropertyPathCustomization::GetRawOtherHelpers() const
-	{
-		return Private::GetRawPathHelperPointers(OtherHelpers);
-	}
-
 	void FPropertyPathCustomization::OnOtherPropertyChanged()
 	{
 		FieldSelector->Refresh();
@@ -191,14 +145,29 @@ namespace UE::MVVM
 
 	void FPropertyPathCustomization::OnSourceSelectionChanged(FBindingSource Selected)
 	{
-		FScopedTransaction Transaction(LOCTEXT("ChangeSource", "Change binding source."));
-
 		PropertyHandle->NotifyPreChange();
 
-		for (const TSharedPtr<IFieldPathHelper>& Path : PathHelpers)
+		UMVVMEditorSubsystem* Subsystem = GEditor->GetEditorSubsystem<UMVVMEditorSubsystem>();
+		
+		TArray<void*> RawData;
+		PropertyHandle->AccessRawData(RawData);
+
+		for (void* RawPtr : RawData)
 		{
-			Path->SetSelectedSource(Selected);
-			Path->ResetBinding(); // Might make sense to keep this around in case we retarget to a compatible widget or switch back.
+			if (RawPtr == nullptr)
+			{
+				continue;
+			}
+
+			FMVVMBlueprintPropertyPath* PropertyPath = (FMVVMBlueprintPropertyPath*) RawPtr;
+			if (bIsWidget)
+			{
+				PropertyPath->SetWidgetName(Selected.Name);
+			}
+			else
+			{
+				PropertyPath->SetViewModelId(Selected.ViewModelId);
+			}
 		}
 
 		PropertyHandle->NotifyPostChange(EPropertyChangeType::ValueSet);
@@ -206,7 +175,7 @@ namespace UE::MVVM
 		FieldSelector->Refresh();
 	}
 
-	void FPropertyPathCustomization::OnPropertySelectionChanged(FMVVMConstFieldVariant Selected)
+	void FPropertyPathCustomization::OnPropertySelectionChanged(FMVVMBlueprintPropertyPath Selected)
 	{
 		if (bPropertySelectionChanging)
 		{
@@ -214,15 +183,20 @@ namespace UE::MVVM
 		}
 		TGuardValue<bool> ReentrantGuard(bPropertySelectionChanging, true);
 
-		FScopedTransaction Transaction(LOCTEXT("ChangeBindingProperty", "Change binding property."));
-
 		PropertyHandle->NotifyPreChange();
 
-		// This is currently handled in SMVVMFieldSelector, but I wonder if it should be here...
-		TArray<IFieldPathHelper*> Helpers = GetRawPathHelpers();
-		for (const IFieldPathHelper* Helper : Helpers)
+		TArray<void*> RawData;
+		PropertyHandle->AccessRawData(RawData);
+
+		for (void* RawPtr : RawData)
 		{
-			Helper->SetBindingReference(Selected);
+			if (RawPtr == nullptr)
+			{
+				continue;
+			}
+
+			FMVVMBlueprintPropertyPath* PropertyPath = (FMVVMBlueprintPropertyPath*) RawPtr;
+			*PropertyPath = Selected;
 		}
 
 		PropertyHandle->NotifyPostChange(EPropertyChangeType::ValueSet);
@@ -233,6 +207,163 @@ namespace UE::MVVM
 		SourceSelector->Refresh();
 		FieldSelector->Refresh();
 	}
+
+	TArray<UE::MVVM::FBindingSource> FPropertyPathCustomization::OnGetSources() const 
+	{
+		UMVVMEditorSubsystem* Subsystem = GEditor->GetEditorSubsystem<UMVVMEditorSubsystem>();
+		if (bIsWidget)
+		{
+			return Subsystem->GetBindableWidgets(WidgetBlueprint);
+		}
+		else
+		{
+			return Subsystem->GetAllViewModels(WidgetBlueprint);
+		}
+	}
+
+	UE::MVVM::FBindingSource FPropertyPathCustomization::OnGetSelectedSource() const
+	{
+		TArray<void*> RawData;
+		PropertyHandle->AccessRawData(RawData);
+
+		FGuid SelectedGuid;
+		FName SelectedName;
+
+		bool bFirst = true;
+
+		for (void* RawPtr : RawData)
+		{
+			if (RawPtr == nullptr)
+			{
+				continue;
+			}
+
+			FMVVMBlueprintPropertyPath* PropertyPath = (FMVVMBlueprintPropertyPath*) RawPtr;
+			if (bIsWidget)
+			{
+				if (bFirst)
+				{
+					SelectedName = PropertyPath->GetWidgetName();
+				}
+				else if (SelectedName != PropertyPath->GetWidgetName())
+				{
+					SelectedName = FName();
+					break;
+				}
+			}
+			else
+			{
+				if (bFirst)
+				{
+					SelectedGuid = PropertyPath->GetViewModelId();
+				}
+				else if (SelectedGuid != PropertyPath->GetViewModelId())
+				{
+					SelectedGuid = FGuid();
+					break;
+				}
+			}
+
+			bFirst = false;
+		}
+
+		if (bIsWidget)
+		{
+			return UE::MVVM::FBindingSource::CreateForWidget(WidgetBlueprint, SelectedName);
+		}
+		else
+		{
+			return UE::MVVM::FBindingSource::CreateForViewModel(WidgetBlueprint, SelectedGuid);
+		}
+	}
+
+	TArray<FMVVMBlueprintPropertyPath> FPropertyPathCustomization::OnGetFields() const
+	{
+		UE::MVVM::FBindingSource Source = OnGetSelectedSource();
+
+		TArray<FMVVMAvailableBinding> AvailableBindings;
+
+		UClass* Class = Source.Class;
+		UMVVMSubsystem* Subsystem = GEngine->GetEngineSubsystem<UMVVMSubsystem>();
+		if (bIsWidget)
+		{
+			AvailableBindings = Subsystem->GetWidgetAvailableBindings(Class);
+		}
+		else
+		{
+			AvailableBindings = Subsystem->GetViewModelAvailableBindings(Class);
+		}
+
+		TArray<FMVVMBlueprintPropertyPath> AvailablePaths;
+		Algo::Transform(AvailableBindings, AvailablePaths, 
+			[this, Class, Source](const FMVVMAvailableBinding& Binding)
+			{
+				FName BindingName = Binding.GetBindingName().ToName();
+
+				UE::MVVM::FMVVMConstFieldVariant Variant;
+
+				if (const UFunction* Function = Class->FindFunctionByName(BindingName))
+				{
+					Variant = UE::MVVM::FMVVMConstFieldVariant(Function);
+				}
+				else if (const FProperty* Property = Class->FindPropertyByName(BindingName))
+				{
+					Variant = UE::MVVM::FMVVMConstFieldVariant(Property);
+				}
+
+				FMVVMBlueprintPropertyPath Path;
+
+				if (!Variant.IsEmpty())
+				{
+					if (bIsWidget)
+					{
+						Path.SetWidgetName(Source.Name);
+					}
+					else
+					{
+						Path.SetViewModelId(Source.ViewModelId);
+					}
+
+					Path.SetBasePropertyPath(Variant);
+				}
+
+				return Path;
+			});
+
+		return AvailablePaths;
+	}
+
+	FMVVMBlueprintPropertyPath FPropertyPathCustomization::OnGetSelectedField() const
+	{
+		TArray<void*> RawData;
+		PropertyHandle->AccessRawData(RawData);
+
+		FMVVMBlueprintPropertyPath SelectedField;
+
+		bool bFirst = true;
+
+		for (void* RawPtr : RawData)
+		{
+			if (RawPtr == nullptr)
+			{
+				continue;
+			}
+
+			FMVVMBlueprintPropertyPath* PropertyPath = (FMVVMBlueprintPropertyPath*) RawPtr;
+			if (bFirst)
+			{
+				SelectedField = *PropertyPath;
+			}
+			else if (SelectedField != *PropertyPath)
+			{
+				SelectedField = FMVVMBlueprintPropertyPath();
+				break;
+			}
+			bFirst = false;
+		}
+
+		return SelectedField;
+	}
 }
 
 #undef LOCTEXT_NAMESPACE
\ No newline at end of file
diff --git a/Engine/Plugins/Runtime/ModelViewViewModel/Source/ModelViewViewModelEditor/Private/Customizations/MVVMPropertyPathCustomization.h b/Engine/Plugins/Runtime/ModelViewViewModel/Source/ModelViewViewModelEditor/Private/Customizations/MVVMPropertyPathCustomization.h
--- a/Engine/Plugins/Runtime/ModelViewViewModel/Source/ModelViewViewModelEditor/Private/Customizations/MVVMPropertyPathCustomization.h
+++ b/Engine/Plugins/Runtime/ModelViewViewModel/Source/ModelViewViewModelEditor/Private/Customizations/MVVMPropertyPathCustomization.h
@@ -4,7 +4,6 @@
 
 #include "IPropertyTypeCustomization.h"
 #include "MVVMBlueprintView.h"
-#include "MVVMPropertyPathHelpers.h"
 #include "PropertyHandle.h"
 #include "Types/MVVMFieldVariant.h"
 
@@ -33,14 +32,16 @@ namespace UE::MVVM
 
 	private:
 		void OnSourceSelectionChanged(FBindingSource Selected);
-		void OnPropertySelectionChanged(FMVVMConstFieldVariant Selected);
+		void OnPropertySelectionChanged(FMVVMBlueprintPropertyPath Selected);
 
 		void OnOtherPropertyChanged();
 
-		TArray<IFieldPathHelper*> GetRawPathHelpers() const;
-		TArray<IFieldPathHelper*> GetRawOtherHelpers() const;
-		
 		EMVVMBindingMode GetCurrentBindingMode() const;
+		TArray<UE::MVVM::FBindingSource> OnGetSources() const;
+		UE::MVVM::FBindingSource OnGetSelectedSource() const;
+
+		TArray<FMVVMBlueprintPropertyPath> OnGetFields() const;
+		FMVVMBlueprintPropertyPath OnGetSelectedField() const;
 
 		void HandleBlueprintChanged(UBlueprint* Blueprint);
 
@@ -51,9 +52,7 @@ namespace UE::MVVM
 		TSharedPtr<SMVVMSourceSelector> SourceSelector;
 		TSharedPtr<SMVVMFieldSelector> FieldSelector;
 
-		TArray<TSharedPtr<IFieldPathHelper>> PathHelpers;
-		TArray<TSharedPtr<IFieldPathHelper>> OtherHelpers;
-
+		bool bIsWidget = false;
 		bool bPropertySelectionChanging = false;
 
 		FDelegateHandle OnBlueprintChangedHandle;
diff --git a/Engine/Plugins/Runtime/ModelViewViewModel/Source/ModelViewViewModelEditor/Private/MVVMEditorSubsystem.cpp b/Engine/Plugins/Runtime/ModelViewViewModel/Source/ModelViewViewModelEditor/Private/MVVMEditorSubsystem.cpp
--- a/Engine/Plugins/Runtime/ModelViewViewModel/Source/ModelViewViewModelEditor/Private/MVVMEditorSubsystem.cpp
+++ b/Engine/Plugins/Runtime/ModelViewViewModel/Source/ModelViewViewModelEditor/Private/MVVMEditorSubsystem.cpp
@@ -2,23 +2,194 @@
 
 #include "MVVMEditorSubsystem.h"
 
+#include "Algo/Reverse.h"
 #include "Bindings/MVVMBindingHelper.h"
+#include "Blueprint/WidgetTree.h"
 #include "BlueprintActionDatabase.h"
 #include "BlueprintNodeSpawner.h"
+#include "Containers/Deque.h"
 #include "Engine/Engine.h"
+#include "Kismet/BlueprintFunctionLibrary.h"
+#include "Kismet2/BlueprintEditorUtils.h"
+#include "K2Node_CallFunction.h"
+#include "K2Node_FunctionEntry.h"
+#include "K2Node_FunctionResult.h"
+#include "K2Node_VariableGet.h"
 #include "MVVMSubsystem.h"
 #include "MVVMWidgetBlueprintExtension_View.h"
 #include "ScopedTransaction.h"
+#include "Types/MVVMBindingSource.h"
 #include "WidgetBlueprint.h"
 
-#define LOCTEXT_NAMESPACE "UMVVMEditorSubsystem"
+#define LOCTEXT_NAMESPACE "MVVMEditorSubsystem"
 
+namespace UE::MVVM::Private
+{
+	void OnBindingPreEditChange(UMVVMBlueprintView* BlueprintView, FName PropertyName)
+	{
+		FProperty* ChangedProperty = FMVVMBlueprintViewBinding::StaticStruct()->FindPropertyByName(PropertyName);
+		check(ChangedProperty != nullptr);
+
+		FEditPropertyChain EditChain;
+		EditChain.AddTail(UMVVMBlueprintView::StaticClass()->FindPropertyByName("Bindings"));
+		EditChain.AddTail(ChangedProperty);
+		EditChain.SetActivePropertyNode(ChangedProperty);
+
+		BlueprintView->PreEditChange(EditChain);
+	}
+
+	void OnBindingPostEditChange(UMVVMBlueprintView* BlueprintView, FName PropertyName)
+	{
+		FProperty* ChangedProperty = FMVVMBlueprintViewBinding::StaticStruct()->FindPropertyByName(PropertyName);
+		check(ChangedProperty != nullptr);
+
+		FEditPropertyChain EditChain;
+		EditChain.AddTail(UMVVMBlueprintView::StaticClass()->FindPropertyByName("Bindings"));
+		EditChain.AddTail(ChangedProperty);
+		EditChain.SetActivePropertyNode(ChangedProperty);
+
+		FPropertyChangedEvent ChangeEvent(ChangedProperty, EPropertyChangeType::ValueSet);
+		FPropertyChangedChainEvent ChainEvent(EditChain, ChangeEvent);
+		BlueprintView->PostEditChangeChainProperty(ChainEvent);
+	}
+
+	UEdGraph* CreateFunctionGraph(UBlueprint* Blueprint, FName InFunctionName, EFunctionFlags ExtraFunctionFlag, const FStringView Category)
+	{
+		bool bIsEditable = true;
+
+		FName UniqueFunctionName = FBlueprintEditorUtils::FindUniqueKismetName(Blueprint, InFunctionName.ToString());
+		UEdGraph* FunctionGraph = FBlueprintEditorUtils::CreateNewGraph(Blueprint, UniqueFunctionName, UEdGraph::StaticClass(), UEdGraphSchema_K2::StaticClass());
+		Blueprint->FunctionGraphs.Add(FunctionGraph);
+		FunctionGraph->bEditable = bIsEditable;
+
+		const UEdGraphSchema_K2* Schema = CastChecked<UEdGraphSchema_K2>(FunctionGraph->GetSchema());
+		Schema->MarkFunctionEntryAsEditable(FunctionGraph, bIsEditable);
+		Schema->CreateDefaultNodesForGraph(*FunctionGraph);
+
+		// function entry node
+		FGraphNodeCreator<UK2Node_FunctionEntry> FunctionEntryCreator(*FunctionGraph);
+		UK2Node_FunctionEntry* FunctionEntry = FunctionEntryCreator.CreateNode();
+		FunctionEntry->FunctionReference.SetSelfMember(FunctionGraph->GetFName());
+		FunctionEntry->AddExtraFlags(ExtraFunctionFlag);
+		FunctionEntry->bIsEditable = bIsEditable;
+		FunctionEntry->MetaData.Category = FText::FromStringView(Category);
+		FunctionEntry->NodePosX = -500;
+		FunctionEntry->NodePosY = 0;
+		FunctionEntryCreator.Finalize();
+
+		FGraphNodeCreator<UK2Node_FunctionResult> FunctionResultCreator(*FunctionGraph);
+		UK2Node_FunctionResult* FunctionResult = FunctionResultCreator.CreateNode();
+		FunctionResult->FunctionReference.SetSelfMember(FunctionGraph->GetFName());
+		FunctionResult->bIsEditable = bIsEditable;
+		FunctionResult->NodePosX = 500;
+		FunctionResult->NodePosY = 0;
+		FunctionResultCreator.Finalize();
+
+		return FunctionGraph;
+	}
+
+	UK2Node_FunctionEntry* FindFunctionEntry(UEdGraph* Graph)
+	{
+		for (UEdGraphNode* Node : Graph->Nodes)
+		{
+			if (UK2Node_FunctionEntry* FunctionEntry = Cast<UK2Node_FunctionEntry>(Node))
+			{
+				return FunctionEntry;
+			}
+		}
+		return nullptr;
+	}
+
+	UK2Node_FunctionResult* FindFunctionResult(UEdGraph* Graph)
+	{
+		for (UEdGraphNode* Node : Graph->Nodes)
+		{
+			if (UK2Node_FunctionResult* FunctionResult = Cast<UK2Node_FunctionResult>(Node))
+			{
+				return FunctionResult;
+			}
+		}
+		return nullptr;
+	}
+
+	UEdGraphNode* FindConversionNode(UEdGraph* Graph)
+	{
+		UK2Node_FunctionResult* FunctionResult = UE::MVVM::Private::FindFunctionResult(Graph);
+		if (!ensureMsgf(FunctionResult != nullptr, TEXT("Function result node not found in conversion function wrapper!")))
+		{
+			return nullptr;
+		}
+
+		if (!ensureMsgf(FunctionResult->UserDefinedPins.Num() == 1, TEXT("Function result should have exactly one return value.")))
+		{
+			return nullptr;
+		}
+
+		UEdGraphPin* ResultPin = FunctionResult->FindPin(FunctionResult->UserDefinedPins[0]->PinName, EGPD_Input);
+		if (!ensureMsgf(ResultPin != nullptr, TEXT("Function result pin not found.")))
+		{
+			return nullptr;
+		}
+
+		if (!ensureMsgf(ResultPin->LinkedTo.Num() != 0, TEXT("Result pin not linked to anything")))
+		{
+			return nullptr;
+		}
+
+		// finally found our conversion node
+		UEdGraphNode* ConversionNode = ResultPin->LinkedTo[0]->GetOwningNode();
+		return ConversionNode;
+	}
+
+	const TObjectPtr<UEdGraph>* FindExistingConversionFunctionWrapper(const UWidgetBlueprint* WidgetBlueprint, FName WrapperName)
+	{
+		const TObjectPtr<UEdGraph>* ExistingGraph = WidgetBlueprint->FunctionGraphs.FindByPredicate(
+			[&WrapperName](TObjectPtr<UEdGraph> GraphPtr) -> bool
+			{
+				if (UEdGraph* Graph = GraphPtr.Get())
+				{
+					return Graph->GetFName() == WrapperName;
+				}
+				return false;
+			});
+
+		return ExistingGraph;
+	}
+
+	void RemoveConversionFunctionWrapper(UWidgetBlueprint* WidgetBlueprint, FName WrapperName)
+	{
+		const TObjectPtr<UEdGraph>* ExistingGraph = FindExistingConversionFunctionWrapper(WidgetBlueprint, WrapperName);
+		if (ExistingGraph != nullptr)
+		{
+			UEdGraph* Graph = ExistingGraph->Get();
+			Graph->Rename(nullptr, Graph->GetOuter(), REN_DoNotDirty | REN_ForceNoResetLoaders);
+			WidgetBlueprint->FunctionGraphs.Remove(Graph);
+		}
+	}
+}
+
+FName UMVVMEditorSubsystem::GetConversionFunctionWrapperName(const UWidgetBlueprint* WidgetBlueprint, const FMVVMBlueprintViewBinding& Binding, bool bSourceToDestination) const
+{
+	if (UMVVMBlueprintView* View = GetView(WidgetBlueprint))
+	{
+		TStringBuilder<256> StringBuilder;
+		StringBuilder << TEXT("__");
+		Binding.GetFName(View).AppendString(StringBuilder);
+
+		StringBuilder << (bSourceToDestination ? TEXT("_SourceToDest") : TEXT("_DestToSource"));
+
+		return FName(StringBuilder.ToString());
+	}
+
+	return FName();
+}
+
 UMVVMBlueprintView* UMVVMEditorSubsystem::RequestView(UWidgetBlueprint* WidgetBlueprint) const
 {
 	return UMVVMWidgetBlueprintExtension_View::RequestExtension<UMVVMWidgetBlueprintExtension_View>(WidgetBlueprint)->GetBlueprintView();
 }
 
-UMVVMBlueprintView* UMVVMEditorSubsystem::GetView(UWidgetBlueprint* WidgetBlueprint) const
+UMVVMBlueprintView* UMVVMEditorSubsystem::GetView(const UWidgetBlueprint* WidgetBlueprint) const
 {
 	if (UMVVMWidgetBlueprintExtension_View* ExtensionView = UMVVMWidgetBlueprintExtension_View::GetExtension<UMVVMWidgetBlueprintExtension_View>(WidgetBlueprint))
 	{
@@ -27,49 +198,4 @@
 	return nullptr;
 }
 
-TArray<UE::MVVM::FMVVMConstFieldVariant> UMVVMEditorSubsystem::GetChildViewModels(UClass* Class)
-{
-
-	auto IsValidObjectProperty = [](const FProperty* Property)
-	{
-		if (const FObjectPropertyBase* ObjectProperty = CastField<FObjectPropertyBase>(Property))
-		{
-			return ObjectProperty->PropertyClass->ImplementsInterface(UNotifyFieldValueChanged::StaticClass());
-		}
-		return false;
-	};
-
-
-	TArray<UE::MVVM::FMVVMConstFieldVariant> Result;
-	for (TFieldIterator<const FProperty> PropertyIt(Class, EFieldIteratorFlags::IncludeSuper); PropertyIt; ++PropertyIt)
-	{
-		const FProperty* Property = *PropertyIt;
-		if (IsValidObjectProperty(Property) && UE::MVVM::BindingHelper::IsValidForSourceBinding(Property))
-		{
-			Result.Add(UE::MVVM::FMVVMConstFieldVariant(Property));
-		}
-	}
-
-	for (TFieldIterator<const UFunction> FunctionItt(Class, EFieldIteratorFlags::IncludeSuper); FunctionItt; ++FunctionItt)
-	{
-		const UFunction* Function = *FunctionItt;
-		check(Function);
-		if (!Function->HasAnyFunctionFlags(FUNC_Private | FUNC_Protected))
-		{
-			if (UE::MVVM::BindingHelper::IsValidForSourceBinding(Function))
-			{
-				if (const FProperty* ReturnValue = UE::MVVM::BindingHelper::GetReturnProperty(Function))
-				{
-					if (IsValidObjectProperty(ReturnValue))
-					{
-						Result.Add(UE::MVVM::FMVVMConstFieldVariant(Function));
-					}
-				}
-			}
-		}
-	}
-
-	return Result;
-}
-
 void UMVVMEditorSubsystem::RemoveViewModel(UWidgetBlueprint* WidgetBlueprint, FName ViewModel)
@@ -134,9 +260,338 @@
 	}
 }
 
-TArray<const UFunction*> UMVVMEditorSubsystem::GetAvailableConversionFunctions(const UE::MVVM::FMVVMConstFieldVariant& Source, const UE::MVVM::FMVVMConstFieldVariant& Destination) const
+UEdGraph* UMVVMEditorSubsystem::GetConversionFunctionGraph(const UWidgetBlueprint* WidgetBlueprint, const FMVVMBlueprintViewBinding& Binding, bool bSourceToDestination) const
 {
-	TArray<const UFunction*> ConversionFunctions;
+	const FName WrapperName = GetConversionFunctionWrapperName(WidgetBlueprint, Binding, bSourceToDestination);
+	const TObjectPtr<UEdGraph>* EdGraphPtr = UE::MVVM::Private::FindExistingConversionFunctionWrapper(WidgetBlueprint, WrapperName);
+
+	return EdGraphPtr != nullptr ? EdGraphPtr->Get() : nullptr;
+}
+
+UEdGraph* UMVVMEditorSubsystem::CreateConversionFunctionWrapperGraph(UWidgetBlueprint* WidgetBlueprint, const FMVVMBlueprintViewBinding& Binding, const UFunction* ConversionFunction, bool bSourceToDestination)
+{
+	if (UMVVMBlueprintView* View = GetView(WidgetBlueprint))
+	{
+		TValueOrError<TArray<const FProperty*>, FString> ArgumentsResult = UE::MVVM::BindingHelper::TryGetArgumentsForConversionFunction(ConversionFunction);
+		if (ArgumentsResult.HasValue())
+		{
+			TArray<const FProperty*> Arguments = ArgumentsResult.StealValue();
+			if (Arguments.Num() > 1)
+			{
+				const FName WrapperName = GetConversionFunctionWrapperName(WidgetBlueprint, Binding, bSourceToDestination);
+
+				UEdGraph* NewGraph = UE::MVVM::Private::CreateFunctionGraph(
+					WidgetBlueprint,
+					WrapperName,
+					FUNC_BlueprintCallable | FUNC_BlueprintPure | FUNC_Const,
+					TEXT("Conversion Functions")
+				);
+
+				UK2Node_FunctionEntry* FunctionEntry = UE::MVVM::Private::FindFunctionEntry(NewGraph);
+				UK2Node_FunctionResult* FunctionResult = UE::MVVM::Private::FindFunctionResult(NewGraph);
+
+				FGraphNodeCreator<UK2Node_CallFunction> CallFunctionCreator(*NewGraph);
+				UK2Node_CallFunction* CallFunctionNode = CallFunctionCreator.CreateNode();
+				CallFunctionNode->SetFromFunction(ConversionFunction);
+				CallFunctionNode->NodePosX = (FunctionEntry->NodePosX + FunctionResult->NodePosX) / 2;
+				CallFunctionCreator.Finalize();
+
+				UEdGraphPin* FunctionEntryThenPin = FunctionEntry->FindPin(UEdGraphSchema_K2::PN_Then, EGPD_Output);
+				UEdGraphPin* FunctionResultExecPin = FunctionResult->GetExecPin();
+			
+				// hook up exec pins 
+				if (!CallFunctionNode->IsNodePure())
+				{
+					UEdGraphPin* CallFunctionExecPin = CallFunctionNode->GetExecPin();
+					UEdGraphPin* CallFunctionThenPin = CallFunctionNode->FindPin(UEdGraphSchema_K2::PN_Then, EGPD_Output);
+
+					FunctionEntryThenPin->MakeLinkTo(CallFunctionExecPin);
+					CallFunctionThenPin->MakeLinkTo(FunctionResultExecPin);
+
+					CallFunctionNode->NodePosY = 0;
+				}
+				else
+				{
+					FunctionEntryThenPin->MakeLinkTo(FunctionResultExecPin);
+
+					CallFunctionNode->NodePosY = 100;
+				}
+
+				const UEdGraphSchema_K2* K2Schema = GetDefault<UEdGraphSchema_K2>();
+
+				// create return value pin
+				const FProperty* ReturnProperty = ConversionFunction->GetReturnProperty();
+				TSharedPtr<FUserPinInfo> PinInfo = MakeShared<FUserPinInfo>();
+				K2Schema->ConvertPropertyToPinType(ReturnProperty, PinInfo->PinType);
+				PinInfo->PinName = ReturnProperty->GetFName();
+				PinInfo->DesiredPinDirection = EGPD_Input;
+				FunctionResult->UserDefinedPins.Add(PinInfo);
+				FunctionResult->ReconstructNode();
+
+				UEdGraphPin* FunctionReturnPin = CallFunctionNode->FindPin(ReturnProperty->GetName(), EGPD_Output);
+				UEdGraphPin* OutputPin = FunctionResult->FindPin(ReturnProperty->GetName());
+				FunctionReturnPin->MakeLinkTo(OutputPin);
+
+				return NewGraph;
+			}
+		}
+	}
+	return nullptr;
+}
+
+void UMVVMEditorSubsystem::SetSourceToDestinationConversionFunction(UWidgetBlueprint* WidgetBlueprint, FMVVMBlueprintViewBinding& Binding, const UFunction* ConversionFunction)
+{
+	if (UMVVMBlueprintView* View = GetView(WidgetBlueprint))
+	{
+		if (ConversionFunction == nullptr || IsValidConversionFunction(ConversionFunction, Binding.ViewModelPath, Binding.WidgetPath))
+		{
+			FScopedTransaction Transaction(LOCTEXT("SetConversionFunction", "Set Conversion Function"));
+
+			WidgetBlueprint->Modify();
+
+			UE::MVVM::Private::OnBindingPreEditChange(View, GET_MEMBER_NAME_CHECKED(FMVVMBlueprintViewBinding, Conversion));
+
+			if (Binding.Conversion.SourceToDestinationWrapper != FName())
+			{
+				UE::MVVM::Private::RemoveConversionFunctionWrapper(WidgetBlueprint, Binding.Conversion.SourceToDestinationWrapper);
+				Binding.Conversion.SourceToDestinationWrapper = FName();
+			}
+
+			if (ConversionFunction != nullptr)
+			{
+				UEdGraph* WrapperGraph = CreateConversionFunctionWrapperGraph(WidgetBlueprint, Binding, ConversionFunction, true);
+				Binding.Conversion.SourceToDestinationFunction.SetFromField<UFunction>(ConversionFunction, WidgetBlueprint->SkeletonGeneratedClass);
+				Binding.Conversion.SourceToDestinationWrapper = WrapperGraph != nullptr ? WrapperGraph->GetFName() : FName();
+			}
+			else
+			{
+				Binding.Conversion.SourceToDestinationFunction = FMemberReference();
+			}
+
+			UE::MVVM::Private::OnBindingPostEditChange(View, GET_MEMBER_NAME_CHECKED(FMVVMBlueprintViewBinding, Conversion));
+		}
+	}
+}
+
+void UMVVMEditorSubsystem::SetDestinationToSourceConversionFunction(UWidgetBlueprint* WidgetBlueprint, FMVVMBlueprintViewBinding& Binding, const UFunction* ConversionFunction)
+{
+	if (UMVVMBlueprintView* View = GetView(WidgetBlueprint))
+	{
+		if (ConversionFunction == nullptr || IsValidConversionFunction(ConversionFunction, Binding.WidgetPath, Binding.ViewModelPath))
+		{
+			FScopedTransaction Transaction(LOCTEXT("SetConversionFunction", "Set Conversion Function"));
+
+			WidgetBlueprint->Modify();
+
+			UE::MVVM::Private::OnBindingPreEditChange(View, GET_MEMBER_NAME_CHECKED(FMVVMBlueprintViewBinding, Conversion));
+
+			if (Binding.Conversion.DestinationToSourceWrapper != FName())
+			{
+				UE::MVVM::Private::RemoveConversionFunctionWrapper(WidgetBlueprint, Binding.Conversion.DestinationToSourceWrapper);
+				Binding.Conversion.DestinationToSourceWrapper = FName();
+			}
+
+			if (ConversionFunction != nullptr)
+			{
+				UEdGraph* WrapperGraph = CreateConversionFunctionWrapperGraph(WidgetBlueprint, Binding, ConversionFunction, true);
+				Binding.Conversion.DestinationToSourceFunction.SetFromField<UFunction>(ConversionFunction, WidgetBlueprint->SkeletonGeneratedClass);
+				Binding.Conversion.DestinationToSourceWrapper = WrapperGraph != nullptr ? WrapperGraph->GetFName() : FName();
+			}
+			else
+			{
+				Binding.Conversion.DestinationToSourceFunction = FMemberReference();
+			}
+
+			UE::MVVM::Private::OnBindingPostEditChange(View, GET_MEMBER_NAME_CHECKED(FMVVMBlueprintViewBinding, Conversion));
+		}
+	}
+}
+
+void UMVVMEditorSubsystem::SetWidgetForBinding(UWidgetBlueprint* WidgetBlueprint, FMVVMBlueprintViewBinding& Binding, const FName WidgetName)
+{
+	if (UMVVMBlueprintView* View = GetView(WidgetBlueprint))
+	{
+		FScopedTransaction Transaction(LOCTEXT("SetBindingSource", "Set Binding Source"));
+
+		UE::MVVM::Private::OnBindingPreEditChange(View, GET_MEMBER_NAME_CHECKED(FMVVMBlueprintViewBinding, WidgetPath));
+
+		Binding.WidgetPath.SetWidgetName(WidgetName);
+
+		UE::MVVM::Private::OnBindingPostEditChange(View, GET_MEMBER_NAME_CHECKED(FMVVMBlueprintViewBinding, WidgetPath));
+	}
+}
+
+void UMVVMEditorSubsystem::SetWidgetPropertyForBinding(UWidgetBlueprint* WidgetBlueprint, FMVVMBlueprintViewBinding& Binding, FMVVMBlueprintPropertyPath Field)
+{
+	if (UMVVMBlueprintView* View = GetView(WidgetBlueprint))
+	{
+		FScopedTransaction Transaction(LOCTEXT("SetBindingProperty", "Set Binding Property"));
+
+		UE::MVVM::Private::OnBindingPreEditChange(View, GET_MEMBER_NAME_CHECKED(FMVVMBlueprintViewBinding, WidgetPath));
+
+		Binding.WidgetPath = Field;
+
+		UE::MVVM::Private::OnBindingPostEditChange(View, GET_MEMBER_NAME_CHECKED(FMVVMBlueprintViewBinding, WidgetPath));
+	}
+}
+
+void UMVVMEditorSubsystem::SetViewModelForBinding(UWidgetBlueprint* WidgetBlueprint, FMVVMBlueprintViewBinding& Binding, FMVVMBlueprintViewModelContext ViewModel)
+{
+	if (UMVVMBlueprintView* View = GetView(WidgetBlueprint))
+	{
+		FScopedTransaction Transaction(LOCTEXT("SetBindingSource", "Set Binding Source"));
+
+		UE::MVVM::Private::OnBindingPreEditChange(View, GET_MEMBER_NAME_CHECKED(FMVVMBlueprintViewBinding, ViewModelPath));
+
+		Binding.ViewModelPath.SetViewModelId(ViewModel.GetViewModelId());
+
+		UE::MVVM::Private::OnBindingPostEditChange(View, GET_MEMBER_NAME_CHECKED(FMVVMBlueprintViewBinding, ViewModelPath));
+	}
+}
+
+void UMVVMEditorSubsystem::SetViewModelPropertyForBinding(UWidgetBlueprint* WidgetBlueprint, FMVVMBlueprintViewBinding& Binding, FMVVMBlueprintPropertyPath Field)
+{
+	if (UMVVMBlueprintView* View = GetView(WidgetBlueprint))
+	{
+		FScopedTransaction Transaction(LOCTEXT("SetBindingProperty", "Set Binding Property"));
+
+		UE::MVVM::Private::OnBindingPreEditChange(View, GET_MEMBER_NAME_CHECKED(FMVVMBlueprintViewBinding, ViewModelPath));
+
+		Binding.ViewModelPath = Field;
+
+		UE::MVVM::Private::OnBindingPostEditChange(View, GET_MEMBER_NAME_CHECKED(FMVVMBlueprintViewBinding, ViewModelPath));
+	}
+}
+
+void UMVVMEditorSubsystem::SetUpdateModeForBinding(UWidgetBlueprint* WidgetBlueprint, FMVVMBlueprintViewBinding& Binding, EMVVMViewBindingUpdateMode Mode)
+{
+	if (Binding.UpdateMode != Mode)
+	{
+		if (UMVVMBlueprintView* View = GetView(WidgetBlueprint))
+		{
+			FScopedTransaction Transaction(LOCTEXT("SetUpdateMode", "Set Update Mode"));
+
+			UE::MVVM::Private::OnBindingPreEditChange(View, GET_MEMBER_NAME_CHECKED(FMVVMBlueprintViewBinding, UpdateMode));
+
+			Binding.UpdateMode = Mode;
+
+			UE::MVVM::Private::OnBindingPreEditChange(View, GET_MEMBER_NAME_CHECKED(FMVVMBlueprintViewBinding, UpdateMode));
+		}
+	}
+}
+
+void UMVVMEditorSubsystem::SetBindingTypeForBinding(UWidgetBlueprint* WidgetBlueprint, FMVVMBlueprintViewBinding& Binding, EMVVMBindingMode Type)
+{
+	if (Binding.BindingType != Type)
+	{
+		if (UMVVMBlueprintView* View = GetView(WidgetBlueprint))
+		{
+			FScopedTransaction Transaction(LOCTEXT("SetBindingType", "Set Binding Type"));
+
+			UE::MVVM::Private::OnBindingPreEditChange(View, GET_MEMBER_NAME_CHECKED(FMVVMBlueprintViewBinding, BindingType));
+
+			Binding.BindingType = Type;
+
+			UE::MVVM::Private::OnBindingPostEditChange(View, GET_MEMBER_NAME_CHECKED(FMVVMBlueprintViewBinding, BindingType));
+		}
+	}
+}
+
+void UMVVMEditorSubsystem::SetEnabledForBinding(UWidgetBlueprint* WidgetBlueprint, FMVVMBlueprintViewBinding& Binding, bool bEnabled)
+{
+	if (Binding.bEnabled != bEnabled)
+	{
+		if (UMVVMBlueprintView* View = GetView(WidgetBlueprint))
+		{
+			UE::MVVM::Private::OnBindingPreEditChange(View, GET_MEMBER_NAME_CHECKED(FMVVMBlueprintViewBinding, bEnabled));
+
+			Binding.bEnabled = bEnabled;
+
+			UE::MVVM::Private::OnBindingPostEditChange(View, GET_MEMBER_NAME_CHECKED(FMVVMBlueprintViewBinding, bEnabled));
+		}
+	}
+}
+
+void UMVVMEditorSubsystem::SetCompileForBinding(UWidgetBlueprint* WidgetBlueprint, FMVVMBlueprintViewBinding& Binding, bool bCompile)
+{
+	if (Binding.bCompile != bCompile)
+	{
+		if (UMVVMBlueprintView* View = GetView(WidgetBlueprint))
+		{
+			UE::MVVM::Private::OnBindingPreEditChange(View, GET_MEMBER_NAME_CHECKED(FMVVMBlueprintViewBinding, bCompile));
+
+			Binding.bCompile = bCompile;
+
+			UE::MVVM::Private::OnBindingPostEditChange(View, GET_MEMBER_NAME_CHECKED(FMVVMBlueprintViewBinding, bCompile));
+		}
+	}
+}
+
+bool UMVVMEditorSubsystem::IsValidConversionFunction(const UFunction* Function, const FMVVMBlueprintPropertyPath& Source, const FMVVMBlueprintPropertyPath& Destination) const
+{
+	if (Source.IsEmpty() || Destination.IsEmpty())
+	{
+		return false;
+	}
+
+	TValueOrError<const FProperty*, FString> ReturnResult = UE::MVVM::BindingHelper::TryGetReturnTypeForConversionFunction(Function);
+	if (ReturnResult.HasError())
+	{
+		return false;
+	}
+
+	TValueOrError<TArray<const FProperty*>, FString> ArgumentsResult = UE::MVVM::BindingHelper::TryGetArgumentsForConversionFunction(Function);
+	if (ArgumentsResult.HasError())
+	{
+		return false;
+	}
+
+	TArray<UE::MVVM::FMVVMConstFieldVariant> SourceFields = Source.GetFields();
+	TArray<UE::MVVM::FMVVMConstFieldVariant> DestFields = Destination.GetFields();
+	if (SourceFields.Num() == 0 || DestFields.Num() == 0)
+	{
+		return false;
+	}
+
+	const FProperty* ReturnProperty = ReturnResult.GetValue();
+	const FProperty* SourceProperty = SourceFields.Last().IsProperty() ? SourceFields.Last().GetProperty() : UE::MVVM::BindingHelper::GetReturnProperty(SourceFields.Last().GetFunction());
+	const FProperty* DestinationProperty = DestFields.Last().IsProperty() ? DestFields.Last().GetProperty() : UE::MVVM::BindingHelper::GetFirstArgumentProperty(DestFields.Last().GetFunction());
+
+	// check that at least one source -> argument binding is compatible
+	bool bAnyCompatible = false;
+
+	const TArray<const FProperty*>& ConversionArgProperties = ArgumentsResult.GetValue();
+	for (const FProperty* ArgumentProperty : ConversionArgProperties)
+	{
+		if (ArgumentProperty->IsA<FObjectProperty>())
+		{
+			// filter out any functions with UObject properties - they aren't valid conversion functions
+			return false;
+		}
+
+		if (UE::MVVM::BindingHelper::ArePropertiesCompatible(SourceProperty, ArgumentProperty))
+		{
+			bAnyCompatible = true;
+		}
+	}
+
+	if (!bAnyCompatible)
+	{
+		return false;
+	}
+
+	// check that the return -> dest is valid
+	if (!UE::MVVM::BindingHelper::ArePropertiesCompatible(ReturnProperty, DestinationProperty))
+	{
+		return false;
+	}
+
+	return true;
+}
+
+TArray<UFunction*> UMVVMEditorSubsystem::GetAvailableConversionFunctions(const UWidgetBlueprint* WidgetBlueprint, const FMVVMBlueprintPropertyPath& Source, const FMVVMBlueprintPropertyPath& Destination) const
+{
+	TArray<UFunction*> ConversionFunctions;
 
 	if (Source.IsEmpty() || Destination.IsEmpty())
 	{
@@ -158,31 +613,16 @@
 				const UFunction* Function = BlueprintAction.GetAssociatedFunction();
 				if (Function != nullptr)
 				{
-					TValueOrError<UE::MVVM::BindingHelper::FConversionFunctionArguments, FString> FunctionProperties = UE::MVVM::BindingHelper::TryGetPropertyTypeForConversionFunction(Function);
-					if (FunctionProperties.HasError())
+					// functions in the widget blueprint can do anything they want, other functions have to be static functions in a BlueprintFunctionLibrary
+					const UClass* FunctionClass = Function->GetOuterUClass();
+					if ((FunctionClass->ClassGeneratedBy == WidgetBlueprint) ||
+						(FunctionClass->IsChildOf<UBlueprintFunctionLibrary>() && Function->HasAllFunctionFlags(FUNC_Static | FUNC_BlueprintPure)))
 					{
-						continue;
+						if (IsValidConversionFunction(Function, Source, Destination))
+						{
+							ConversionFunctions.Add(const_cast<UFunction*>(Function));
+						}
 					}
-
-					const FProperty* ArgumentProperty = FunctionProperties.GetValue().ArgumentProperty;
-					const FProperty* ReturnProperty = FunctionProperties.GetValue().ReturnProperty;
-
-					const FProperty* SourceProperty = Source.IsProperty() ? Source.GetProperty() : UE::MVVM::BindingHelper::GetReturnProperty(Source.GetFunction());
-					const FProperty* DestinationProperty = Destination.IsProperty() ? Destination.GetProperty() : UE::MVVM::BindingHelper::GetFirstArgumentProperty(Destination.GetFunction());
-
-					// check that the source -> argument is valid
-					if (!UE::MVVM::BindingHelper::ArePropertiesCompatible(SourceProperty, ArgumentProperty))
-					{
-						continue;
-					}
-					
-					// check that the return -> dest is valid
-					if (!UE::MVVM::BindingHelper::ArePropertiesCompatible(ReturnProperty, DestinationProperty))
-					{
-						continue;
-					}
-
-					ConversionFunctions.Add(Function);
 				}
 			}
 		}
@@ -195,4 +635,437 @@
 	return ConversionFunctions;
 }
 
+FMVVMBlueprintPropertyPath UMVVMEditorSubsystem::GetPathForConversionFunctionArgument(const UWidgetBlueprint* WidgetBlueprint, const FMVVMBlueprintViewBinding& Binding, FName ArgumentName, bool bSourceToDestination) const
+{
+	UEdGraph* WrapperGraph = GetConversionFunctionGraph(WidgetBlueprint, Binding, bSourceToDestination);
+	if (WrapperGraph == nullptr)
+	{
+		return FMVVMBlueprintPropertyPath();
+	}
+	 
+	UEdGraphNode* ConversionNode = UE::MVVM::Private::FindConversionNode(WrapperGraph);
+
+	const UEdGraphPin* ArgumentPin = ConversionNode->FindPin(ArgumentName, EGPD_Input);
+	if (!ensureMsgf(ArgumentPin != nullptr, TEXT("Couldn't find argument pin named '%s'."), *ArgumentName.ToString()))
+	{
+		return FMVVMBlueprintPropertyPath();
+	}
+
+	if (ArgumentPin->LinkedTo.Num() == 0)
+	{
+		return FMVVMBlueprintPropertyPath();
+	}
+
+	TArray<FMemberReference> PathParts;
+	TDeque<UEdGraphNode*> NodesToSearch;
+	NodesToSearch.PushLast(ArgumentPin->LinkedTo[0]->GetOwningNode());
+
+	while (NodesToSearch.Num() > 0)
+	{
+		UEdGraphNode* Node = NodesToSearch[0];
+		NodesToSearch.PopFirst();
+
+		if (UK2Node_VariableGet* Getter = Cast<UK2Node_VariableGet>(Node))
+		{
+			PathParts.Add(Getter->VariableReference);
+		}
+
+		for (UEdGraphPin* Pin : Node->Pins)
+		{
+			if (Pin->Direction == EGPD_Input && Pin->LinkedTo.Num() > 0)
+			{
+				NodesToSearch.PushLast(Pin->LinkedTo[0]->GetOwningNode());
+			}
+		}
+	}
+
+	if (PathParts.Num() == 0)
+	{
+		return FMVVMBlueprintPropertyPath();
+	}
+
+	Algo::Reverse(PathParts);
+
+	FMVVMBlueprintPropertyPath ResultPath;
+
+	FMemberReference RootRef = PathParts[0];
+	PathParts.RemoveAt(0);
+
+	if (const FObjectProperty* Property = CastField<FObjectProperty>(RootRef.ResolveMember<FProperty>(WidgetBlueprint->SkeletonGeneratedClass)))
+	{
+		if (Property->PropertyClass->ImplementsInterface(UNotifyFieldValueChanged::StaticClass()))
+		{
+			UE::MVVM::FBindingSource Source = UE::MVVM::FBindingSource::CreateForViewModel(WidgetBlueprint, Property->GetFName());
+			ResultPath.SetViewModelId(Source.ViewModelId);
+		}
+		else if (Property->PropertyClass->IsChildOf<UWidget>() || Property->PropertyClass->IsChildOf<UWidgetBlueprint>())
+		{
+			UE::MVVM::FBindingSource Source = UE::MVVM::FBindingSource::CreateForWidget(WidgetBlueprint, Property->GetFName());
+			ResultPath.SetWidgetName(Source.Name);
+		}
+	}
+
+	if (PathParts.Num() == 0)
+	{
+		return ResultPath;
+	}
+
+	// TODO: This only takes one part, because FMVVMBlueprintPropertyPath only takes one part
+	UE::MVVM::FMVVMConstFieldVariant Variant = UE::MVVM::FMVVMConstFieldVariant(PathParts.Last().ResolveMember<FProperty>(WidgetBlueprint->SkeletonGeneratedClass));
+	ResultPath.SetBasePropertyPath(Variant);
+
+	return ResultPath;
+}
+
+void UMVVMEditorSubsystem::SetPathForConversionFunctionArgument(UWidgetBlueprint* WidgetBlueprint, FMVVMBlueprintViewBinding& Binding, FName ArgumentName, const FMVVMBlueprintPropertyPath& Path, bool bSourceToDestination) const
+{
+	UEdGraph* ConversionFunctionGraph = GetConversionFunctionGraph(WidgetBlueprint, Binding, bSourceToDestination);
+	if (!ensure(ConversionFunctionGraph != nullptr))
+	{
+		return;
+	}
+
+	UEdGraphNode* ConversionNode = UE::MVVM::Private::FindConversionNode(ConversionFunctionGraph);
+	if (!ensure(ConversionNode != nullptr))
+	{
+		return;
+	}
+
+	const UEdGraphSchema_K2* Schema = CastChecked<UEdGraphSchema_K2>(ConversionFunctionGraph->GetSchema());
+
+	// find the argument pin with the name ArgumentName
+	UEdGraphPin* ArgumentPin = nullptr;
+	int32 ArgumentIndex = 0;
+	for (UEdGraphPin* Pin : ConversionNode->Pins)
+	{
+		if (Pin->Direction == EGPD_Input && Pin->PinName != UEdGraphSchema_K2::PN_Execute)
+		{
+			if (Pin->PinName == ArgumentName)
+			{
+				ArgumentPin = Pin;
+				break;
+			}
+
+			++ArgumentIndex;
+		}
+	}
+	
+	if (!ensure(ArgumentPin != nullptr))
+	{
+		return;
+	}
+
+	// delete all linked nodes
+	TArray<UEdGraphPin*> PinsToClear;
+	PinsToClear.Add(ArgumentPin);
+
+	TArray<UEdGraphNode*> NodesToRemove;
+
+	while (PinsToClear.Num() > 0)
+	{
+		UEdGraphPin* Pin = PinsToClear[0];
+		PinsToClear.RemoveAt(0);
+
+		PinsToClear.Append(Pin->LinkedTo);
+		Pin->BreakAllPinLinks();
+
+		// rename the node out of the way
+		UEdGraphNode* Node = Pin->GetOwningNode();
+		if (Node != ConversionNode)
+		{
+			NodesToRemove.AddUnique(Node);
+		}
+	}
+
+	for (UEdGraphNode* Node : NodesToRemove)
+	{
+		ConversionFunctionGraph->RemoveNode(Node, true);
+	}
+
+	if (Path.IsEmpty())
+	{
+		// we were given an empty path and we already cleaned up the path 
+		return;
+	}
+
+	TArray<UE::MVVM::FMVVMConstFieldVariant> Fields = Path.GetFields();
+
+	float PosX = ConversionNode->NodePosX - (300 * Fields.Num());
+	const float PosY = ConversionNode->NodePosY + ArgumentIndex * 50;
+
+	UEdGraphPin* PreviousExecPin = UE::MVVM::Private::FindFunctionEntry(ConversionFunctionGraph)->FindPin(UEdGraphSchema_K2::PN_Then);
+	UEdGraphPin* PreviousDataPin = nullptr;
+	UClass* PreviousClass = nullptr;
+
+	// create the root property getter node, ie. the Widget/ViewModel
+	const FProperty* RootProperty = nullptr;
+	if (Path.IsFromWidget())
+	{
+		RootProperty = WidgetBlueprint->SkeletonGeneratedClass->FindPropertyByName(Path.GetWidgetName());
+	}
+	else if (Path.IsFromViewModel())
+	{
+		UMVVMBlueprintView* View = GetView(WidgetBlueprint);
+		const FMVVMBlueprintViewModelContext* Context = View->FindViewModel(Path.GetViewModelId());
+
+		RootProperty = WidgetBlueprint->SkeletonGeneratedClass->FindPropertyByName(Context->GetViewModelName());
+	}
+
+	if (RootProperty != nullptr)
+	{
+		FGraphNodeCreator<UK2Node_VariableGet> RootGetterCreator(*ConversionFunctionGraph);
+		UK2Node_VariableGet* RootGetterNode = RootGetterCreator.CreateNode();
+		RootGetterNode->NodePosX = PosX;
+		RootGetterNode->NodePosY = PosY;
+		RootGetterNode->VariableReference.SetFromField<FProperty>(RootProperty, true, WidgetBlueprint->SkeletonGeneratedClass);
+		RootGetterCreator.Finalize();
+
+		PreviousDataPin = RootGetterNode->Pins[0];
+
+		PreviousClass = CastField<FObjectProperty>(RootProperty)->PropertyClass;
+
+		PosX += 300;
+	}
+	else
+	{
+		ensureMsgf(false, TEXT("Could not resolve root property!"));
+		return;
+	}
+
+	// create all the subsequent nodes in the path
+	for (int32 Index = 0; Index < Fields.Num(); ++Index)
+	{
+		if (!ensureMsgf(PreviousClass != nullptr, TEXT("Path is too long!")))
+		{
+			return;
+		}
+
+		UEdGraphNode* NewNode = nullptr;
+
+		const UE::MVVM::FMVVMConstFieldVariant& Field = Fields[Index];
+		if (Field.IsProperty())
+		{
+			const FProperty* Property = Field.GetProperty();
+
+			if (const FObjectProperty* ObjectProperty = CastField<FObjectProperty>(Property))
+			{
+				PreviousClass = ObjectProperty->PropertyClass;
+
+				FGraphNodeCreator<UK2Node_VariableGet> GetterCreator(*ConversionFunctionGraph);
+				UK2Node_VariableGet* GetterNode = GetterCreator.CreateNode();
+				GetterNode->NodePosX = PosX;
+				GetterNode->NodePosY = PosY;
+				GetterNode->SetFromProperty(Property, false, PreviousClass);
+				GetterNode->AllocateDefaultPins();
+				GetterCreator.Finalize();
+
+				NewNode = GetterNode;
+			}
+			else if (const FStructProperty* StructProperty = CastField<FStructProperty>(Property))
+			{
+				// TODO (sebastiann): StructProperty should be supported as well, but that's not a VariableGet
+				ensureMsgf(false, TEXT("Invalid path, structs are not supported yet."));
+				return;
+			}
+			else
+			{
+				PreviousClass = nullptr;
+			}
+		}
+		else if (Field.IsFunction())
+		{
+			const UFunction* Function = Field.GetFunction();
+
+			FGraphNodeCreator<UK2Node_CallFunction> CallFunctionCreator(*ConversionFunctionGraph);
+			UK2Node_CallFunction* FunctionNode = CallFunctionCreator.CreateNode();
+			FunctionNode->NodePosX = PosX;
+			FunctionNode->NodePosY = PosY;
+			FunctionNode->SetFromFunction(Function);
+			FunctionNode->AllocateDefaultPins();
+			CallFunctionCreator.Finalize();
+
+			const FProperty* ReturnProperty = UE::MVVM::BindingHelper::GetReturnProperty(Function);
+			if (const FObjectProperty* ObjectProperty = CastField<FObjectProperty>(ReturnProperty))
+			{
+				PreviousClass = ObjectProperty->PropertyClass;
+			}
+			else if (const FStructProperty* StructProperty = CastField<FStructProperty>(ReturnProperty))
+			{
+				// TODO (sebastiann): StructProperty should be supported as well, but that's not a VariableGet
+				ensureMsgf(false, TEXT("Invalid path, structs are not supported yet."));
+				return;
+			}
+			else
+			{
+				PreviousClass = nullptr;
+			}
+
+			NewNode = FunctionNode;
+		}
+		else
+		{
+			ensureMsgf(false, TEXT("Invalid path, empty field in path."));
+			return;
+		}
+		
+		// first create new connections
+		for (UEdGraphPin* Pin : NewNode->Pins)
+		{
+			if (Pin->Direction == EGPD_Input)
+			{
+				if (Pin->PinName == UEdGraphSchema_K2::PN_Execute)
+				{
+					PreviousExecPin->BreakAllPinLinks();
+					Pin->MakeLinkTo(PreviousExecPin);
+				}
+				else if (Schema->ArePinsCompatible(PreviousDataPin, Pin, PreviousClass))
+				{
+					Pin->MakeLinkTo(PreviousDataPin);
+				}
+			}
+		}
+
+		// then update our previous pin pointers
+		for (UEdGraphPin* Pin : NewNode->Pins)
+		{
+			if (Pin->Direction == EGPD_Output)
+			{
+				if (Pin->PinName == UEdGraphSchema_K2::PN_Then)
+				{
+					PreviousExecPin = Pin;
+				}
+				else
+				{
+					PreviousDataPin = Pin;
+				}
+			}
+		}
+
+		PosX += 300; 
+	}
+
+	// finish by linking to our ultimate destinations
+	UEdGraphPin* LastExecPin = ConversionNode->FindPin(UEdGraphSchema_K2::PN_Execute);
+	if (LastExecPin == nullptr)
+	{
+		LastExecPin = UE::MVVM::Private::FindFunctionResult(ConversionFunctionGraph)->FindPin(UEdGraphSchema_K2::PN_Execute);
+	}
+	PreviousExecPin->MakeLinkTo(LastExecPin);
+	PreviousDataPin->MakeLinkTo(ArgumentPin);
+}
+
+TArray<UE::MVVM::FBindingSource> UMVVMEditorSubsystem::GetBindableWidgets(const UWidgetBlueprint* WidgetBlueprint) const
+{
+	TArray<UE::MVVM::FBindingSource> Sources;
+
+	const UWidgetTree* WidgetTree = WidgetBlueprint->WidgetTree;
+	if (WidgetTree != nullptr)
+	{
+		return Sources;
+	}
+
+	TArray<UWidget*> AllWidgets;
+	WidgetTree->GetAllWidgets(AllWidgets);
+
+	Sources.Reserve(AllWidgets.Num() + 1);
+
+	// Add current widget as a possible binding source
+	if (UClass* BPClass = WidgetBlueprint->GeneratedClass)
+	{
+		TArray<FMVVMAvailableBinding> Bindings = GEngine->GetEngineSubsystem<UMVVMSubsystem>()->GetWidgetAvailableBindings(BPClass);
+		if (Bindings.Num() > 0)
+		{
+			// at least one valid property, add it to our list
+			UE::MVVM::FBindingSource Source;
+			Source.Name = WidgetBlueprint->GetFName();
+			Source.DisplayName = FText::FromName(WidgetBlueprint->GetFName());
+			Source.Class = BPClass;
+			Sources.Add(Source);
+		}
+	}
+
+	for (const UWidget* Widget : AllWidgets)
+	{
+		TArray<FMVVMAvailableBinding> Bindings = GEngine->GetEngineSubsystem<UMVVMSubsystem>()->GetWidgetAvailableBindings(Widget->GetClass());
+		if (Bindings.Num() > 0)
+		{
+			// at least one valid property, add it to our list
+			UE::MVVM::FBindingSource Source;
+			Source.Name = Widget->GetFName();
+			Source.DisplayName = Widget->GetLabelText();
+			Source.Class = Widget->GetClass();
+			Sources.Add(Source);
+		}
+	}
+
+	return Sources;
+}
+
+TArray<UE::MVVM::FBindingSource> UMVVMEditorSubsystem::GetAllViewModels(const UWidgetBlueprint* WidgetBlueprint) const
+{
+	TArray<UE::MVVM::FBindingSource> Sources;
+
+	if (UMVVMBlueprintView* View = GetView(WidgetBlueprint))
+	{
+		const TArrayView<const FMVVMBlueprintViewModelContext> ViewModels = View->GetViewModels();
+		Sources.Reserve(ViewModels.Num());
+
+		for (const FMVVMBlueprintViewModelContext& ViewModel : ViewModels)
+		{
+			UE::MVVM::FBindingSource Source;
+			Source.ViewModelId = ViewModel.GetViewModelId();
+			Source.DisplayName = ViewModel.GetDisplayName();
+			Source.Class = ViewModel.GetViewModelClass().Get();
+			Sources.Add(Source);
+		}
+	}
+
+	return Sources;
+}
+
+
+TArray<UE::MVVM::FMVVMConstFieldVariant> UMVVMEditorSubsystem::GetChildViewModels(const UClass* Class)
+{
+	auto IsValidObjectProperty = [](const FProperty* Property)
+	{
+		if (const FObjectPropertyBase* ObjectProperty = CastField<FObjectPropertyBase>(Property))
+		{
+			return ObjectProperty->PropertyClass->ImplementsInterface(UNotifyFieldValueChanged::StaticClass());
+		}
+		return false;
+	};
+
+
+	TArray<UE::MVVM::FMVVMConstFieldVariant> Result;
+	for (TFieldIterator<const FProperty> PropertyIt(Class, EFieldIteratorFlags::IncludeSuper); PropertyIt; ++PropertyIt)
+	{
+		const FProperty* Property = *PropertyIt;
+		if (IsValidObjectProperty(Property) && UE::MVVM::BindingHelper::IsValidForSourceBinding(Property))
+		{
+			Result.Add(UE::MVVM::FMVVMConstFieldVariant(Property));
+		}
+	}
+
+	for (TFieldIterator<const UFunction> FunctionItt(Class, EFieldIteratorFlags::IncludeSuper); FunctionItt; ++FunctionItt)
+	{
+		const UFunction* Function = *FunctionItt;
+		check(Function);
+		if (!Function->HasAnyFunctionFlags(FUNC_Private | FUNC_Protected))
+		{
+			if (UE::MVVM::BindingHelper::IsValidForSourceBinding(Function))
+			{
+				if (const FProperty* ReturnValue = UE::MVVM::BindingHelper::GetReturnProperty(Function))
+				{
+					if (IsValidObjectProperty(ReturnValue))
+					{
+						Result.Add(UE::MVVM::FMVVMConstFieldVariant(Function));
+					}
+				}
+			}
+		}
+	}
+
+	return Result;
+}
+
 #undef LOCTEXT_NAMESPACE
diff --git a/Engine/Plugins/Runtime/ModelViewViewModel/Source/ModelViewViewModelEditor/Private/MVVMPropertyPathHelpers.cpp b/Engine/Plugins/Runtime/ModelViewViewModel/Source/ModelViewViewModelEditor/Private/MVVMPropertyPathHelpers.cpp
deleted file mode 100644
diff --git a/Engine/Plugins/Runtime/ModelViewViewModel/Source/ModelViewViewModelEditor/Private/MVVMPropertyPathHelpers.h b/Engine/Plugins/Runtime/ModelViewViewModel/Source/ModelViewViewModelEditor/Private/MVVMPropertyPathHelpers.h
deleted file mode 100644
diff --git a/Engine/Plugins/Runtime/ModelViewViewModel/Source/ModelViewViewModelEditor/Private/Types/MVVMBindingSource.cpp b/Engine/Plugins/Runtime/ModelViewViewModel/Source/ModelViewViewModelEditor/Private/Types/MVVMBindingSource.cpp
new file mode 100644
--- /dev/null
+++ b/Engine/Plugins/Runtime/ModelViewViewModel/Source/ModelViewViewModelEditor/Private/Types/MVVMBindingSource.cpp
@@ -0,0 +1,67 @@
+// Copyright Epic Games, Inc. All Rights Reserved.
+
+#pragma once
+
+#include "Types/MVVMBindingSource.h"
+#include "Blueprint/WidgetTree.h"
+#include "Components/Widget.h"
+#include "Editor.h"
+#include "MVVMEditorSubsystem.h"
+#include "WidgetBlueprint.h"
+
+UE::MVVM::FBindingSource UE::MVVM::FBindingSource::CreateForWidget(const UWidgetBlueprint* WidgetBlueprint, FName WidgetName)
+{
+	UE::MVVM::FBindingSource Source;
+
+	Source.Name = WidgetName;
+
+	if (Source.Name == WidgetBlueprint->GetFName())
+	{
+		Source.DisplayName = FText::FromString(WidgetBlueprint->GetName());
+		Source.Class = WidgetBlueprint->GeneratedClass;
+	}
+	else if (UWidget* Widget = WidgetBlueprint->WidgetTree->FindWidget(Source.Name))
+	{
+		Source.DisplayName = Widget->GetLabelText();
+		Source.Class = Widget->GetClass();
+	}
+
+	return Source;
+}
+
+UE::MVVM::FBindingSource UE::MVVM::FBindingSource::CreateForViewModel(const UWidgetBlueprint* WidgetBlueprint, FGuid ViewModelId)
+{
+	UE::MVVM::FBindingSource Source;
+	Source.ViewModelId = ViewModelId;
+
+	UMVVMEditorSubsystem* Subsystem = GEditor->GetEditorSubsystem<UMVVMEditorSubsystem>();
+	if (UMVVMBlueprintView* View = Subsystem->GetView(WidgetBlueprint))
+	{
+		if (const FMVVMBlueprintViewModelContext* ViewModel = View->FindViewModel(Source.ViewModelId))
+		{
+			Source.DisplayName = ViewModel->GetDisplayName();
+			Source.Class = ViewModel->GetViewModelClass().Get();
+		}
+	}
+
+	return Source;
+}
+
+UE::MVVM::FBindingSource UE::MVVM::FBindingSource::CreateForViewModel(const UWidgetBlueprint* WidgetBlueprint, FName ViewModelName)
+{
+	UE::MVVM::FBindingSource Source;
+
+	UMVVMEditorSubsystem* Subsystem = GEditor->GetEditorSubsystem<UMVVMEditorSubsystem>();
+	if (UMVVMBlueprintView* View = Subsystem->GetView(WidgetBlueprint))
+	{
+		if (const FMVVMBlueprintViewModelContext* ViewModel = View->FindViewModel(ViewModelName))
+		{
+			Source.ViewModelId = ViewModel->GetViewModelId();
+			Source.DisplayName = ViewModel->GetDisplayName();
+			Source.Class = ViewModel->GetViewModelClass().Get();
+		}
+	}
+
+	return Source;
+}
+
diff --git a/Engine/Plugins/Runtime/ModelViewViewModel/Source/ModelViewViewModelEditor/Private/Types/MVVMBindingSource.h b/Engine/Plugins/Runtime/ModelViewViewModel/Source/ModelViewViewModelEditor/Private/Types/MVVMBindingSource.h
--- a/Engine/Plugins/Runtime/ModelViewViewModel/Source/ModelViewViewModelEditor/Private/Types/MVVMBindingSource.h
+++ b/Engine/Plugins/Runtime/ModelViewViewModel/Source/ModelViewViewModelEditor/Private/Types/MVVMBindingSource.h
@@ -5,19 +5,20 @@
 #include "CoreMinimal.h"
 #include "Framework/Views/TableViewTypeTraits.h"
 
+class UWidgetBlueprint;
+
 namespace UE::MVVM
 {
 	struct FBindingSource
 	{
-		FGuid SourceGuid;
+		FGuid ViewModelId;
 		FName Name;
-		const UClass* Class = nullptr;
+		UClass* Class = nullptr;
 		FText DisplayName;
-		bool IsSelected = false;
 
 		bool operator==(const FBindingSource& Other) const
 		{
-			return SourceGuid == Other.SourceGuid && Name == Other.Name && Class == Other.Class;
+			return ViewModelId == Other.ViewModelId && Name == Other.Name && Class == Other.Class;
 		}
 
 		bool operator!=(const FBindingSource& Other) const
@@ -27,24 +28,27 @@ namespace UE::MVVM
 
 		friend int32 GetTypeHash(const FBindingSource& Source)
 		{
-			uint32 Hash = HashCombine(GetTypeHash(Source.SourceGuid), GetTypeHash(Source.Name));
+			uint32 Hash = HashCombine(GetTypeHash(Source.ViewModelId), GetTypeHash(Source.Name));
 			Hash = HashCombine(Hash, GetTypeHash(Source.Class));
 			return Hash;
 		}
 
 		bool IsValid() const
 		{
-			return SourceGuid.IsValid() || !Name.IsNone();
+			return Class != nullptr && (ViewModelId.IsValid() || !Name.IsNone());
 		}
 
 		void Reset()
 		{
-			SourceGuid = FGuid();
+			ViewModelId = FGuid();
 			Name = FName();
 			Class = nullptr;
 			DisplayName = FText::GetEmpty();
-			IsSelected = false;
 		}
+
+		static FBindingSource CreateForWidget(const UWidgetBlueprint* WidgetBlueprint, FName WidgetName);
+		static FBindingSource CreateForViewModel(const UWidgetBlueprint* WidgetBlueprint, FGuid ViewModelId);
+		static FBindingSource CreateForViewModel(const UWidgetBlueprint* WidgetBlueprint, FName ViewModelName);
 	};
 }
 
diff --git a/Engine/Plugins/Runtime/ModelViewViewModel/Source/ModelViewViewModelEditor/Private/Widgets/SMVVMConversionPath.cpp b/Engine/Plugins/Runtime/ModelViewViewModel/Source/ModelViewViewModelEditor/Private/Widgets/SMVVMConversionPath.cpp
--- a/Engine/Plugins/Runtime/ModelViewViewModel/Source/ModelViewViewModelEditor/Private/Widgets/SMVVMConversionPath.cpp
+++ b/Engine/Plugins/Runtime/ModelViewViewModel/Source/ModelViewViewModelEditor/Private/Widgets/SMVVMConversionPath.cpp
@@ -6,7 +6,6 @@
 #include "Framework/MultiBox/MultiBoxBuilder.h"
 #include "MVVMBlueprintViewBinding.h"
 #include "MVVMEditorSubsystem.h"
-#include "MVVMPropertyPathHelpers.h"
 #include "MVVMSubsystem.h"
 #include "Styling/MVVMEditorStyle.h"
 #include "Styling/StyleColors.h"
@@ -146,38 +145,6 @@ FReply SMVVMConversionPath::OnButtonClicked() const
 
 void SMVVMConversionPath::SetConversionFunction(const UFunction* Function)
 {
-	TArray<FMVVMBlueprintViewBinding*> ViewBindings = Bindings.Get(TArray<FMVVMBlueprintViewBinding*>());
-	if (ViewBindings.Num() == 0)
-	{
-		return;
-	}
-
-	for (FMVVMBlueprintViewBinding* Binding : ViewBindings)
-	{
-		if (bSourceToDestination)
-		{
-			if (Function)
-			{
-				Binding->Conversion.SourceToDestinationFunction.SetFromField<UFunction>(Function, WidgetBlueprint->SkeletonGeneratedClass);
-			}
-			else
-			{
-				Binding->Conversion.SourceToDestinationFunction = FMemberReference();
-			}
-		}
-		else
-		{
-			if (Function)
-			{
-				Binding->Conversion.DestinationToSourceFunction.SetFromField<UFunction>(Function, WidgetBlueprint->SkeletonGeneratedClass);
-			}
-			else
-			{
-				Binding->Conversion.DestinationToSourceFunction = FMemberReference();
-			}
-		}
-	}
-
 	OnFunctionChanged.ExecuteIfBound(Function);
 }
 
@@ -209,30 +176,23 @@ TSharedRef<SWidget> SMVVMConversionPath::GetFunctionMenuContent()
 		return SNullWidget::NullWidget;
 	}
 
-	TSet<const UFunction*> ConversionFunctions;
+	TSet<UFunction*> ConversionFunctions;
 
-	for (FMVVMBlueprintViewBinding* Binding : ViewBindings)
+	for (const FMVVMBlueprintViewBinding* Binding : ViewBindings)
 	{
-		UE::MVVM::FViewModelFieldPathHelper ViewModelHelper(&Binding->ViewModelPath, WidgetBlueprint);
-		UE::MVVM::FMVVMConstFieldVariant ViewModelField = ViewModelHelper.GetSelectedField();
-
-		UE::MVVM::FWidgetFieldPathHelper WidgetHelper(&Binding->WidgetPath, WidgetBlueprint);
-		UE::MVVM::FMVVMConstFieldVariant WidgetField = WidgetHelper.GetSelectedField();
-
-		UMVVMSubsystem::FConstDirectionalBindingArgs Args;
-		Args.SourceBinding = bSourceToDestination ? ViewModelField : WidgetField;
-		Args.DestinationBinding = bSourceToDestination ? WidgetField : ViewModelField;
+		FMVVMBlueprintPropertyPath SourcePath = bSourceToDestination ? Binding->ViewModelPath : Binding->WidgetPath;
+		FMVVMBlueprintPropertyPath DestPath = bSourceToDestination ? Binding->WidgetPath : Binding->ViewModelPath;
 
 		UMVVMEditorSubsystem* EditorSubsystem = GEditor->GetEditorSubsystem<UMVVMEditorSubsystem>();
-		TArray<const UFunction*> FunctionsForThis = EditorSubsystem->GetAvailableConversionFunctions(Args.SourceBinding, Args.DestinationBinding);
+		TArray<UFunction*> FunctionsForThis = EditorSubsystem->GetAvailableConversionFunctions(WidgetBlueprint, SourcePath, DestPath);
 
 		if (ConversionFunctions.Num() > 0)
 		{
-			ConversionFunctions = ConversionFunctions.Intersect(TSet<const UFunction*>(FunctionsForThis));
+			ConversionFunctions = ConversionFunctions.Intersect(TSet<UFunction*>(FunctionsForThis));
 		}
 		else
 		{
-			ConversionFunctions = TSet<const UFunction*>(FunctionsForThis);
+			ConversionFunctions = TSet<UFunction*>(FunctionsForThis);
 		}
 	}
 
diff --git a/Engine/Plugins/Runtime/ModelViewViewModel/Source/ModelViewViewModelEditor/Private/Widgets/SMVVMFieldSelector.cpp b/Engine/Plugins/Runtime/ModelViewViewModel/Source/ModelViewViewModelEditor/Private/Widgets/SMVVMFieldSelector.cpp
--- a/Engine/Plugins/Runtime/ModelViewViewModel/Source/ModelViewViewModelEditor/Private/Widgets/SMVVMFieldSelector.cpp
+++ b/Engine/Plugins/Runtime/ModelViewViewModel/Source/ModelViewViewModelEditor/Private/Widgets/SMVVMFieldSelector.cpp
@@ -27,61 +27,104 @@ namespace UE::MVVM::Private
 		}
 		return LOCTEXT("None", "<None>");
 	}
+}
 
-	FMVVMConstFieldVariant GetSelectedFieldFromHelpers(const TArray<IFieldPathHelper*>& PathHelpers)
-	{
-		bool bFirst = true;
-		FMVVMConstFieldVariant Selected;
 
-		for (const IFieldPathHelper* Helper : PathHelpers)
-		{
-			const FMVVMConstFieldVariant ThisSelection = Helper->GetSelectedField();
-			if (bFirst)
-			{
-				Selected = ThisSelection;
-				bFirst = false;
-			}
-			else if (Selected != ThisSelection)
-			{
-				Selected.Reset();
-				break;
-			}
-		}
+void SMVVMFieldEntry::Construct(const FArguments& InArgs)
+{
+	Field = InArgs._Field;
+	OnValidateField = InArgs._OnValidateField;
 
-		return Selected;
-	}
+	ChildSlot
+		[
+			SNew(SHorizontalBox)
+			+ SHorizontalBox::Slot()
+		.HAlign(HAlign_Left)
+		.VAlign(VAlign_Center)
+		.AutoWidth()
+		[
+			SAssignNew(Icon, SMVVMFieldIcon)
+		]
+	+ SHorizontalBox::Slot()
+		.HAlign(HAlign_Fill)
+		.VAlign(VAlign_Center)
+		.Padding(4, 0, 0, 0)
+		[
+			SAssignNew(Label, STextBlock)
+			.TextStyle(InArgs._TextStyle)
+		]
+		];
+
+	Refresh();
 }
 
-TOptional<FBindingSource> SMVVMFieldSelector::GetSelectedSource() const
+void SMVVMFieldEntry::Refresh()
 {
-	TOptional<FBindingSource> Source;
-	for (const IFieldPathHelper* Helper : PathHelpers)
+	FText ToolTipText = FText::GetEmpty();
+	bool bEnabled = true;
+	if (OnValidateField.IsBound())
 	{
-		if (!Source.IsSet())
+		TValueOrError<bool, FString> Result = OnValidateField.Execute(Field);
+		if (Result.HasError())
 		{
-			Source = Helper->GetSelectedSource();
+			ToolTipText = FText::FromString(Result.GetError());
+			bEnabled = false;
 		}
-		else if (Source != Helper->GetSelectedSource())
+		else
 		{
-			Source.Reset();
-			break;
+			bEnabled = true;
 		}
 	}
 
-	return Source;
+	SetEnabled(bEnabled);
+
+	UE::MVVM::FMVVMConstFieldVariant Variant;
+
+	TArray<UE::MVVM::FMVVMConstFieldVariant> Fields = Field.GetFields();
+	if (Fields.Num() > 0)
+	{
+		Variant = Fields.Last();
+	}
+
+	if (ToolTipText.IsEmpty())
+	{
+		ToolTipText = Variant.IsFunction() ? Variant.GetFunction()->GetToolTipText() :
+			Variant.IsProperty() ? Variant.GetProperty()->GetToolTipText() :
+			FText::GetEmpty();
+	}
+
+	SetToolTipText(ToolTipText);
+
+	Icon->RefreshBinding(Variant);
+	Label->SetText(Private::GetFieldDisplayName(Variant));
+}
+
+void SMVVMFieldEntry::SetField(const FMVVMBlueprintPropertyPath& InField)
+{
+	Field = InField;
+
+	Refresh();
 }
 
 void SMVVMFieldSelector::Construct(const FArguments& InArgs)
 {
-	PathHelpers = InArgs._PathHelpers;
-	CounterpartHelpers = InArgs._CounterpartHelpers;
 	OnValidateFieldDelegate = InArgs._OnValidateField;
 	OnSelectionChangedDelegate = InArgs._OnSelectionChanged;
+	SelectedField = InArgs._SelectedField;
+	check(SelectedField.IsSet());
+
+	SelectedSource = InArgs._SelectedSource;
+	check(SelectedSource.IsSet());
+
+	AvailableFields = InArgs._AvailableFields;
+	check(AvailableFields.IsSet());
+
+	CachedAvailableFields = AvailableFields.Get(TArray<FMVVMBlueprintPropertyPath>());
+
 	BindingMode = InArgs._BindingMode;
 	bIsSource = InArgs._IsSource;
 	TextStyle = InArgs._TextStyle;
 
-	FMVVMConstFieldVariant SelectedField = Private::GetSelectedFieldFromHelpers(PathHelpers);
 	Refresh();
 
 	ChildSlot
@@ -101,24 +144,24 @@ void SMVVMFieldSelector::Construct(const FArguments& InArgs)
 				.TextStyle(FAppStyle::Get(), "HintText")
 				.Visibility_Lambda([this]()
 				{
-					return SelectedSource.IsSet() ? EVisibility::Collapsed : EVisibility::Visible;
+					return SelectedSource.Get(UE::MVVM::FBindingSource()).IsValid() ? EVisibility::Collapsed : EVisibility::Visible;
 				})
 			]
 			+ SOverlay::Slot()
 			[
-				SAssignNew(FieldComboBox, SComboBox<FMVVMConstFieldVariant>)
+				SAssignNew(FieldComboBox, SComboBox<FMVVMBlueprintPropertyPath>)
 				.Visibility_Lambda([this]()
 				{
-					return SelectedSource.IsSet() ? EVisibility::Visible : EVisibility::Hidden;
+					return SelectedSource.Get(UE::MVVM::FBindingSource()).IsValid() ? EVisibility::Visible : EVisibility::Hidden;
 				})
-				.OptionsSource(&AvailableFields)
-				.InitiallySelectedItem(SelectedField)
+				.OptionsSource(&CachedAvailableFields)
+				.InitiallySelectedItem(SelectedField.Get(FMVVMBlueprintPropertyPath()))
 				.OnGenerateWidget(this, &SMVVMFieldSelector::OnGenerateFieldWidget)
 				.OnSelectionChanged(this, &SMVVMFieldSelector::OnComboBoxSelectionChanged)
 				[
-					SAssignNew(SelectedEntry, SMVVMFieldEntry)
+					SAssignNew(SelectedEntryWidget, SMVVMFieldEntry)
 					.TextStyle(TextStyle)
-					.Field(SelectedField)
+					.Field(SelectedField.Get(FMVVMBlueprintPropertyPath()))
 					.OnValidateField(this, &SMVVMFieldSelector::ValidateField)
 				]
 			]
@@ -137,205 +180,48 @@ void SMVVMFieldSelector::Construct(const FArguments& InArgs)
 	];
 }
 
-void SMVVMFieldSelector::OnComboBoxSelectionChanged(FMVVMConstFieldVariant Selected, ESelectInfo::Type SelectionType)
+void SMVVMFieldSelector::OnComboBoxSelectionChanged(FMVVMBlueprintPropertyPath Selected, ESelectInfo::Type SelectionType)
 {
-	SelectedEntry->SetField(Selected);
+	SelectedEntryWidget->SetField(Selected);
 
 	OnSelectionChangedDelegate.ExecuteIfBound(Selected);
 }
 
 void SMVVMFieldSelector::Refresh()
 {
-	AvailableFields.Reset();
-
-	SelectedSource = GetSelectedSource();
-
-	TSet<FMVVMConstFieldVariant> CompatibleFields;
-	TSet<FMVVMConstFieldVariant> IncompatibleFields;
-
-	FMVVMConstFieldVariant SelectedField = Private::GetSelectedFieldFromHelpers(PathHelpers);
-	FMVVMConstFieldVariant CounterpartSelectedField = Private::GetSelectedFieldFromHelpers(CounterpartHelpers);
+	CachedAvailableFields = AvailableFields.Get(TArray<FMVVMBlueprintPropertyPath>());
 
-	bool bFirst = true;
-
-	for (IFieldPathHelper* Helper : PathHelpers)
-	{
-		TSet<FMVVMConstFieldVariant> AllFields;
-		Helper->GetAvailableFields(AllFields);
-
-		if (bFirst)
-		{
-			bFirst = false;
-
-			for (const FMVVMConstFieldVariant& Field : AllFields)
-			{
-				TValueOrError<bool, FString> Result = IsValidBindingForField(Field, CounterpartSelectedField);
-				if (Result.HasValue() && Result.GetValue() == true)
-				{
-					CompatibleFields.Add(Field);
-				}
-				else
-				{
-					IncompatibleFields.Add(Field);
-				}
-			}
-		}
-		else
-		{
-			CompatibleFields.Intersect(AllFields);
-			IncompatibleFields.Intersect(AllFields);
-		}
-	}
-
-	AvailableFields.Append(CompatibleFields.Array());
-
-	// put all incompatible properties at the end so they don't just disappear from the list without explanation
-	AvailableFields.Append(IncompatibleFields.Array());
-	
 	if (FieldComboBox.IsValid())
 	{
 		FieldComboBox->RefreshOptions();
-		FieldComboBox->SetSelectedItem(SelectedField);
+		FieldComboBox->SetSelectedItem(SelectedField.Get(FMVVMBlueprintPropertyPath()));
 	}
 }
 
-TValueOrError<bool, FString> SMVVMFieldSelector::IsValidBindingForField(const FMVVMConstFieldVariant& Field, const FMVVMConstFieldVariant& CounterpartField) const
+TValueOrError<bool, FString> SMVVMFieldSelector::ValidateField(FMVVMBlueprintPropertyPath Field) const
 {
-	// if we have a delegate bound, use that instead
 	if (OnValidateFieldDelegate.IsBound())
 	{
 		return OnValidateFieldDelegate.Execute(Field);
 	}
 
-	if (CounterpartField.IsEmpty() || Field.IsEmpty())
-	{
-		return MakeValue(true);
-	}
-
-	if (!BindingMode.IsSet())
-	{
-		return MakeValue(true);
-	}
-
-	EMVVMBindingMode Mode = BindingMode.Get();
-
-	bool bToDestination = UE::MVVM::IsForwardBinding(Mode);
-	if (!bIsSource)
-	{
-		bToDestination = !bToDestination;
-	}
-
-	UMVVMSubsystem::FConstDirectionalBindingArgs Args;
-	Args.SourceBinding = bToDestination ? Field : CounterpartField;
-	Args.DestinationBinding = bToDestination ? CounterpartField : Field;
-
-	UMVVMSubsystem* Subsystem = GEngine->GetEngineSubsystem<UMVVMSubsystem>();
-
-	TValueOrError<bool, FString> Result = Subsystem->IsBindingValid(Args);
-
-	// check in the other direction as well if this is a two-way binding
-	if (Mode == EMVVMBindingMode::TwoWay &&
-		Result.HasValue() && Result.GetValue() == true)
-	{
-		Args.SourceBinding = bToDestination ? CounterpartField : Field;
-		Args.DestinationBinding = bToDestination ? Field : CounterpartField;
-
-		return Subsystem->IsBindingValid(Args);
-	}
-
-	return Result;
+	return MakeValue(true);
 }
 
-TValueOrError<bool, FString> SMVVMFieldSelector::ValidateField(FMVVMConstFieldVariant Field) const
-{
-	return IsValidBindingForField(Field, Private::GetSelectedFieldFromHelpers(CounterpartHelpers));
-}
-
-TSharedRef<SWidget> SMVVMFieldSelector::OnGenerateFieldWidget(FMVVMConstFieldVariant Field) const
+TSharedRef<SWidget> SMVVMFieldSelector::OnGenerateFieldWidget(FMVVMBlueprintPropertyPath Path) const
 {
 	return SNew(SMVVMFieldEntry)
 		.TextStyle(TextStyle)
-		.Field(Field)
+		.Field(Path)
 		.OnValidateField(this, &SMVVMFieldSelector::ValidateField);
 }
 
-void SMVVMFieldEntry::Construct(const FArguments& InArgs)
-{
-	Field = InArgs._Field;
-	OnValidateField = InArgs._OnValidateField;
-
-	ChildSlot
-	[
-		SNew(SHorizontalBox)
-		+ SHorizontalBox::Slot()
-		.HAlign(HAlign_Left)
-		.VAlign(VAlign_Center)
-		.AutoWidth()
-		[
-			SAssignNew(Icon, SMVVMFieldIcon)
-		]
-		+ SHorizontalBox::Slot()
-		.HAlign(HAlign_Fill)
-		.VAlign(VAlign_Center)
-		.Padding(4, 0, 0, 0)
-		[
-			SAssignNew(Label, STextBlock)
-			.TextStyle(InArgs._TextStyle)
-		]
-	];
-
-	Refresh();
-}
-
-void SMVVMFieldEntry::Refresh()
-{
-	FText ToolTipText = FText::GetEmpty();
-	bool bEnabled = true;
-	if (OnValidateField.IsBound())
-	{
-		TValueOrError<bool, FString> Result = OnValidateField.Execute(Field);
-		if (Result.HasError())
-		{
-			ToolTipText = FText::FromString(Result.GetError());
-			bEnabled = false;
-		}
-		else
-		{
-			bEnabled = true;
-		}
-	}
-
-	SetEnabled(bEnabled);
-
-	if (ToolTipText.IsEmpty())
-	{
-		ToolTipText = Field.IsFunction() ? Field.GetFunction()->GetToolTipText() :
-			Field.IsProperty() ? Field.GetProperty()->GetToolTipText() :
-			FText::GetEmpty();
-	}
-
-	SetToolTipText(ToolTipText);
-
-	Icon->RefreshBinding(Field);
-	Label->SetText(Private::GetFieldDisplayName(Field));
-}
-
-void SMVVMFieldEntry::SetField(const UE::MVVM::FMVVMConstFieldVariant& InField)
-{
-	Field = InField;
-
-	Refresh();
-}
-
 EVisibility SMVVMFieldSelector::GetClearVisibility() const
 {
-	for (const IFieldPathHelper* Helper : PathHelpers)
+	FMVVMBlueprintPropertyPath SelectedPath = SelectedField.Get(FMVVMBlueprintPropertyPath());
+	if (!SelectedPath.IsEmpty())
 	{
-		const FMVVMConstFieldVariant ThisSelection = Helper->GetSelectedField();
-		if (!ThisSelection.IsEmpty())
-		{
-			return EVisibility::Visible;
-		}
+		return EVisibility::Visible;
 	}
 
 	return EVisibility::Collapsed;
@@ -348,6 +234,11 @@ FReply SMVVMFieldSelector::OnClearBinding()
 		FieldComboBox->ClearSelection();
 	}
 
+	if (SelectedEntryWidget.IsValid())
+	{
+		SelectedEntryWidget->Refresh();
+	}
+
 	return FReply::Handled();
 }
 
diff --git a/Engine/Plugins/Runtime/ModelViewViewModel/Source/ModelViewViewModelEditor/Private/Widgets/SMVVMFieldSelector.h b/Engine/Plugins/Runtime/ModelViewViewModel/Source/ModelViewViewModelEditor/Private/Widgets/SMVVMFieldSelector.h
--- a/Engine/Plugins/Runtime/ModelViewViewModel/Source/ModelViewViewModelEditor/Private/Widgets/SMVVMFieldSelector.h
+++ b/Engine/Plugins/Runtime/ModelViewViewModel/Source/ModelViewViewModelEditor/Private/Widgets/SMVVMFieldSelector.h
@@ -3,10 +3,10 @@
 #pragma once
 
 #include "MVVMFieldVariantListTraits.h"
-#include "MVVMPropertyPathHelpers.h" 
 #include "Styling/CoreStyle.h"
 #include "Templates/ValueOrError.h"
 #include "Types/MVVMBindingMode.h"
+#include "Types/MVVMBindingSource.h"
 #include "Types/MVVMFieldVariant.h"
 #include "UObject/UnrealType.h"
 #include "Widgets/DeclarativeSyntaxSupport.h"
@@ -22,7 +22,7 @@ class SMVVMFieldIcon;
 class STextBlock; 
 
 using FIsFieldValidResult = TValueOrError<bool, FString>;
-DECLARE_DELEGATE_RetVal_OneParam(FIsFieldValidResult, FIsFieldValid, UE::MVVM::FMVVMConstFieldVariant);
+DECLARE_DELEGATE_RetVal_OneParam(FIsFieldValidResult, FIsFieldValid, FMVVMBlueprintPropertyPath);
 
 class SMVVMFieldEntry : public SCompoundWidget
 {
@@ -32,16 +32,16 @@ public:
 		_TextStyle( &FCoreStyle::Get().GetWidgetStyle<FTextBlockStyle>( "NormalText" ) )
 		{}
 		SLATE_STYLE_ARGUMENT(FTextBlockStyle, TextStyle)
-		SLATE_ARGUMENT(UE::MVVM::FMVVMConstFieldVariant, Field)
+		SLATE_ARGUMENT(FMVVMBlueprintPropertyPath, Field)
 		SLATE_EVENT(FIsFieldValid, OnValidateField)
 	SLATE_END_ARGS()
 
 	void Construct(const FArguments& InArgs);
 	void Refresh(); 
-	void SetField(const UE::MVVM::FMVVMConstFieldVariant& InField);
+	void SetField(const FMVVMBlueprintPropertyPath& InField);
 
 private:
-	UE::MVVM::FMVVMConstFieldVariant Field;
+	FMVVMBlueprintPropertyPath Field;
 	FIsFieldValid OnValidateField;
 	TSharedPtr<SMVVMFieldIcon> Icon;
 	TSharedPtr<STextBlock> Label;
@@ -50,7 +50,7 @@ private:
 class SMVVMFieldSelector : public SCompoundWidget
 {
 public:
-	DECLARE_DELEGATE_OneParam(FSelectionChanged, UE::MVVM::FMVVMConstFieldVariant);
+	DECLARE_DELEGATE_OneParam(FSelectionChanged, FMVVMBlueprintPropertyPath);
 
 	SLATE_BEGIN_ARGS(SMVVMFieldSelector) :
 		_TextStyle(&FCoreStyle::Get().GetWidgetStyle<FTextBlockStyle>("NormalText")),
@@ -58,11 +58,17 @@ public:
 		{
 		}
 		SLATE_STYLE_ARGUMENT(FTextBlockStyle, TextStyle)
-		SLATE_ARGUMENT(TArray<UE::MVVM::IFieldPathHelper*>, PathHelpers)
-		SLATE_ARGUMENT(TArray<UE::MVVM::IFieldPathHelper*>, CounterpartHelpers)
 		SLATE_ATTRIBUTE(EMVVMBindingMode, BindingMode)
 		SLATE_ARGUMENT(bool, IsSource)
+		SLATE_ATTRIBUTE(UE::MVVM::FBindingSource, SelectedSource)
+		SLATE_ATTRIBUTE(FMVVMBlueprintPropertyPath, SelectedField)
+		SLATE_ATTRIBUTE(TArray<FMVVMBlueprintPropertyPath>, AvailableFields)
 		SLATE_EVENT(FSelectionChanged, OnSelectionChanged)
+
+		/**
+		  * Would the given field be a valid entry for this combo box? 
+		  * The error string returned will be used as a tooltip.
+		  */
 		SLATE_EVENT(FIsFieldValid, OnValidateField)
 	SLATE_END_ARGS()
 
@@ -70,29 +76,28 @@ public:
 	void Refresh();
 
 private:
-	TSharedRef<SWidget> OnGenerateFieldWidget(UE::MVVM::FMVVMConstFieldVariant FieldPath) const;
-	void OnComboBoxSelectionChanged(UE::MVVM::FMVVMConstFieldVariant Selected, ESelectInfo::Type SelectionType);
-	TOptional<UE::MVVM::FBindingSource> GetSelectedSource() const;
+	TSharedRef<SWidget> OnGenerateFieldWidget(FMVVMBlueprintPropertyPath Path) const;
+	void OnComboBoxSelectionChanged(FMVVMBlueprintPropertyPath Selected, ESelectInfo::Type SelectionType);
 
-	TValueOrError<bool, FString> IsValidBindingForField(const UE::MVVM::FMVVMConstFieldVariant& Field, const UE::MVVM::FMVVMConstFieldVariant& CounterpartField) const;
-	TValueOrError<bool, FString> ValidateField(UE::MVVM::FMVVMConstFieldVariant Field) const;
+	TValueOrError<bool, FString> ValidateField(FMVVMBlueprintPropertyPath Field) const;
 
 	EVisibility GetClearVisibility() const;
 	FReply OnClearBinding();
 
 private:
-	TArray<UE::MVVM::IFieldPathHelper*> PathHelpers;
-	TArray<UE::MVVM::IFieldPathHelper*> CounterpartHelpers;
 	TAttribute<EMVVMBindingMode> BindingMode;
-	TSharedPtr<SMVVMFieldEntry> SelectedEntry;
-	TOptional<UE::MVVM::FBindingSource> SelectedSource;
+	TSharedPtr<SMVVMFieldEntry> SelectedEntryWidget;
+	TAttribute<UE::MVVM::FBindingSource> SelectedSource;
+	TAttribute<FMVVMBlueprintPropertyPath> SelectedField;
+	TAttribute<TArray<FMVVMBlueprintPropertyPath>> AvailableFields;
+
 	const FTextBlockStyle* TextStyle = nullptr;
 	bool bIsSource = false;
 
-	TSharedPtr<SComboBox<UE::MVVM::FMVVMConstFieldVariant>> FieldComboBox;
+	TSharedPtr<SComboBox<FMVVMBlueprintPropertyPath>> FieldComboBox;
 
 	FSelectionChanged OnSelectionChangedDelegate;
 	FIsFieldValid OnValidateFieldDelegate;
 
-	TArray<UE::MVVM::FMVVMConstFieldVariant> AvailableFields;
+	TArray<FMVVMBlueprintPropertyPath> CachedAvailableFields;
 }; 
diff --git a/Engine/Plugins/Runtime/ModelViewViewModel/Source/ModelViewViewModelEditor/Private/Widgets/SMVVMPropertyPath.cpp b/Engine/Plugins/Runtime/ModelViewViewModel/Source/ModelViewViewModelEditor/Private/Widgets/SMVVMPropertyPath.cpp
--- a/Engine/Plugins/Runtime/ModelViewViewModel/Source/ModelViewViewModelEditor/Private/Widgets/SMVVMPropertyPath.cpp
+++ b/Engine/Plugins/Runtime/ModelViewViewModel/Source/ModelViewViewModelEditor/Private/Widgets/SMVVMPropertyPath.cpp
@@ -2,10 +2,13 @@
 
 #include "Widgets/SMVVMPropertyPath.h"
 
+#include "Blueprint/WidgetTree.h"
+#include "Editor.h"
 #include "MVVMBlueprintView.h"
+#include "MVVMEditorSubsystem.h"
 #include "MVVMWidgetBlueprintExtension_View.h"
-#include "Widgets/SMVVMFieldIcon.h"
 #include "Styling/MVVMEditorStyle.h"
+#include "Widgets/SMVVMFieldIcon.h"
 #include "Widgets/SBoxPanel.h"
 #include "Widgets/Images/SImage.h"
 #include "Widgets/Layout/SBorder.h"
@@ -14,37 +17,12 @@
 
 #define LOCTEXT_NAMESPACE "SMVVMPropertyPath"
 
-void SMVVMPropertyPathBase::Construct(const FArguments& InArgs, const UWidgetBlueprint* InWidgetBlueprint)
+void SMVVMPropertyPath::Construct(const FArguments& InArgs)
 {
-	WidgetBlueprint = InWidgetBlueprint;
-
-	const UE::MVVM::IFieldPathHelper& PathHelper = GetPathHelper();
-	SelectedSource = PathHelper.GetSelectedSource();
-	SelectedField = PathHelper.GetSelectedField();
-
-	auto GetContextText = [this]()
-	{
-		if (SelectedSource.IsSet())
-		{
-			return SelectedSource.GetValue().DisplayName;
-		}
-		return FText::GetEmpty();
-	};
-
-	auto GetPathText = [this]()
-	{
-		if (SelectedField.IsProperty())
-		{
-			return SelectedField.GetProperty()->GetDisplayNameText();
-		}
-
-		if (SelectedField.IsFunction())
-		{
-			return SelectedField.GetFunction()->GetDisplayNameText();
-		}
-
-		return FText::GetEmpty();
-	};
+	WidgetBlueprint = InArgs._WidgetBlueprint;
+	check(InArgs._WidgetBlueprint != nullptr);
+	PropertyPath = InArgs._PropertyPath;
+	check(PropertyPath);
 
 	ChildSlot
 	[
@@ -55,7 +33,7 @@ void SMVVMPropertyPathBase::Construct(const FArguments& InArgs, const UWidgetBlu
 		[
 			SNew(STextBlock)
 			.TextStyle(FMVVMEditorStyle::Get(), "PropertyPath.ContextText")
-			.Text_Lambda(GetContextText)
+			.Text(this, &SMVVMPropertyPath::GetSourceDisplayName)
 		]
 		+ SHorizontalBox::Slot()
 		.Padding(FMargin(5, 0, 0, 0))
@@ -63,7 +41,7 @@ void SMVVMPropertyPathBase::Construct(const FArguments& InArgs, const UWidgetBlu
 		.AutoWidth()
 		[
 			SNew(SMVVMFieldIcon)
-			.Field(SelectedField)
+			.Field(GetLastField())
 		]
 		+ SHorizontalBox::Slot()
 		.Padding(FMargin(2, 0, 5, 0))
@@ -72,24 +50,38 @@ void SMVVMPropertyPathBase::Construct(const FArguments& InArgs, const UWidgetBlu
 		[
 			SNew(STextBlock)
 			.TextStyle(FMVVMEditorStyle::Get(), "PropertyPath.ContextText")
-			.Text_Lambda(GetPathText)
+			.Text(this, &SMVVMPropertyPath::GetFieldDisplayName)
 		]
 	];
 }
 
-void SMVVMWidgetPropertyPath::Construct(const FArguments& InArgs, const UWidgetBlueprint* InWidgetBlueprint)
+FText SMVVMPropertyPath::GetSourceDisplayName() const
 {
-	PathHelper = UE::MVVM::FWidgetFieldPathHelper(InArgs._WidgetPath, InWidgetBlueprint);
-
-	SMVVMPropertyPathBase::Construct(SMVVMPropertyPathBase::FArguments(), InWidgetBlueprint);
+	if (PropertyPath->IsFromWidget())
+	{
+		const UWidget* Widget = WidgetBlueprint->WidgetTree->FindWidget(PropertyPath->GetWidgetName());
+		return FText::FromString(Widget->GetDisplayLabel());
+	}
+	else if (PropertyPath->IsFromViewModel())
+	{
+		const UMVVMEditorSubsystem* Subsystem = GEditor->GetEditorSubsystem<UMVVMEditorSubsystem>();
+		if (const UMVVMBlueprintView* View = Subsystem->GetView(WidgetBlueprint.Get()))
+		{
+			const FMVVMBlueprintViewModelContext* ViewModel = View->FindViewModel(PropertyPath->GetViewModelId());
+			return ViewModel->GetDisplayName();
+		}
+	}
+	return FText::GetEmpty();
 }
 
-void SMVVMViewModelPropertyPath::Construct(const FArguments& InArgs, const UWidgetBlueprint* InWidgetBlueprint)
+FText SMVVMPropertyPath::GetFieldDisplayName() const
 {
-	PathHelper = UE::MVVM::FViewModelFieldPathHelper(InArgs._ViewModelPath, InWidgetBlueprint);
-
-	SMVVMPropertyPathBase::Construct(SMVVMPropertyPathBase::FArguments(), InWidgetBlueprint);
+	return FText::FromString(PropertyPath->GetBasePropertyPath());
 }
 
+UE::MVVM::FMVVMConstFieldVariant SMVVMPropertyPath::GetLastField() const
+{
+	return PropertyPath->GetFields().Last();
+}
 
 #undef LOCTEXT_NAMESPACE
\ No newline at end of file
diff --git a/Engine/Plugins/Runtime/ModelViewViewModel/Source/ModelViewViewModelEditor/Private/Widgets/SMVVMPropertyPath.h b/Engine/Plugins/Runtime/ModelViewViewModel/Source/ModelViewViewModelEditor/Private/Widgets/SMVVMPropertyPath.h
--- a/Engine/Plugins/Runtime/ModelViewViewModel/Source/ModelViewViewModelEditor/Private/Widgets/SMVVMPropertyPath.h
+++ b/Engine/Plugins/Runtime/ModelViewViewModel/Source/ModelViewViewModelEditor/Private/Widgets/SMVVMPropertyPath.h
@@ -5,62 +5,25 @@
 #include "Widgets/SCompoundWidget.h"
 #include "MVVMBlueprintView.h"
 #include "Widgets/SMVVMFieldIcon.h"
-#include "MVVMPropertyPathHelpers.h"
 #include "Types/MVVMFieldVariant.h"
 
-class SMVVMPropertyPathBase : public SCompoundWidget
+class SMVVMPropertyPath : public SCompoundWidget
 {
 public:
-	SLATE_BEGIN_ARGS(SMVVMPropertyPathBase)
+	SLATE_BEGIN_ARGS(SMVVMPropertyPath)
 	{}
+		SLATE_ARGUMENT(FMVVMBlueprintPropertyPath*, PropertyPath)
+		SLATE_ARGUMENT(const UWidgetBlueprint*, WidgetBlueprint)
 	SLATE_END_ARGS()
 
-	void Construct(const FArguments& InArgs, const UWidgetBlueprint* WidgetBlueprint);
-	virtual void Tick(const FGeometry& AllottedGeometry, const double InCurrentTime, const float InDeltaTime) override;
+	void Construct(const FArguments& InArgs);
 
 private:
-	virtual UE::MVVM::IFieldPathHelper& GetPathHelper() = 0;
+	FText GetSourceDisplayName() const;
+	FText GetFieldDisplayName() const;
+	UE::MVVM::FMVVMConstFieldVariant GetLastField() const;
 
 private:
-	const UWidgetBlueprint* WidgetBlueprint = nullptr;
-	UE::MVVM::FMVVMConstFieldVariant SelectedField;
-	TOptional<UE::MVVM::FBindingSource> SelectedSource;
-};
-
-class SMVVMWidgetPropertyPath : public SMVVMPropertyPathBase
-{
-public:
-	SLATE_BEGIN_ARGS(SMVVMWidgetPropertyPath) {}
-		SLATE_ARGUMENT(FMVVMBlueprintPropertyPath*, WidgetPath)
-	SLATE_END_ARGS()
-
-	void Construct(const FArguments& InArgs, const UWidgetBlueprint* WidgetBlueprint);
-
-private:
-	virtual UE::MVVM::IFieldPathHelper& GetPathHelper() override
-	{ 
-		return PathHelper;
-	}
-
-private:
-	UE::MVVM::FWidgetFieldPathHelper PathHelper;
-};
-
-class SMVVMViewModelPropertyPath : public SMVVMPropertyPathBase
-{
-public:
-	SLATE_BEGIN_ARGS(SMVVMViewModelPropertyPath) {}
-		SLATE_ARGUMENT(FMVVMBlueprintPropertyPath*, ViewModelPath)
-	SLATE_END_ARGS()
-
-	void Construct(const FArguments& InArgs, const UWidgetBlueprint* WidgetBlueprint);
-
-private:
-	virtual UE::MVVM::IFieldPathHelper& GetPathHelper() override
-	{
-		return PathHelper;
-	}
-
-private:
-	UE::MVVM::FViewModelFieldPathHelper PathHelper;
+	FMVVMBlueprintPropertyPath* PropertyPath = nullptr;
+	TWeakObjectPtr<const UWidgetBlueprint> WidgetBlueprint = nullptr;
 };
diff --git a/Engine/Plugins/Runtime/ModelViewViewModel/Source/ModelViewViewModelEditor/Private/Widgets/SMVVMSourceEntry.h b/Engine/Plugins/Runtime/ModelViewViewModel/Source/ModelViewViewModelEditor/Private/Widgets/SMVVMSourceEntry.h
--- a/Engine/Plugins/Runtime/ModelViewViewModel/Source/ModelViewViewModelEditor/Private/Widgets/SMVVMSourceEntry.h
+++ b/Engine/Plugins/Runtime/ModelViewViewModel/Source/ModelViewViewModelEditor/Private/Widgets/SMVVMSourceEntry.h
@@ -2,12 +2,11 @@
 
 #pragma once
 
-#include "MVVMPropertyPathHelpers.h"
 #include "Styling/CoreStyle.h"
 #include "Styling/ISlateStyle.h"
 #include "Styling/SlateTypes.h"
 #include "Styling/SlateWidgetStyleAsset.h"
-#include "Types/MVVMFieldVariant.h"
+#include "Types/MVVMBindingSource.h"
 #include "Widgets/DeclarativeSyntaxSupport.h"
 #include "Widgets/SCompoundWidget.h"
 
diff --git a/Engine/Plugins/Runtime/ModelViewViewModel/Source/ModelViewViewModelEditor/Private/Widgets/SMVVMSourceSelector.cpp b/Engine/Plugins/Runtime/ModelViewViewModel/Source/ModelViewViewModelEditor/Private/Widgets/SMVVMSourceSelector.cpp
--- a/Engine/Plugins/Runtime/ModelViewViewModel/Source/ModelViewViewModelEditor/Private/Widgets/SMVVMSourceSelector.cpp
+++ b/Engine/Plugins/Runtime/ModelViewViewModel/Source/ModelViewViewModelEditor/Private/Widgets/SMVVMSourceSelector.cpp
@@ -11,33 +11,13 @@
 
 using namespace UE::MVVM;
 
-static FBindingSource GetSelectedSource(const TArray<IFieldPathHelper*>& Helpers)
-{
-	bool bFirst = true;
-	TOptional<FBindingSource> Result;
-
-	for (const IFieldPathHelper* Helper : Helpers)
-	{
-		TOptional<FBindingSource> ThisSelection = Helper->GetSelectedSource();
-		if (bFirst)
-		{
-			Result = ThisSelection;
-		}
-		else if (Result != ThisSelection)
-		{
-			// all don't share the same value
-			Result.Reset();
-			break;
-		}
-	}
-
-	return Result.Get(FBindingSource());
-}
-
 void SMVVMSourceSelector::Construct(const FArguments& Args)
 {
-	PathHelpers = Args._PathHelpers;
 	TextStyle = Args._TextStyle;
+	AvailableSourcesAttribute = Args._AvailableSources;
+	check(AvailableSourcesAttribute.IsSet());
+	SelectedSourceAttribute = Args._SelectedSource;
+	check(SelectedSourceAttribute.IsSet());
 
 	Refresh();
 
@@ -91,36 +71,8 @@ void SMVVMSourceSelector::OnComboBoxSelectionChanged(FBindingSource Selected, ES
 
 void SMVVMSourceSelector::Refresh()
 {
-	AvailableSources.Reset();
-	SelectedSource = FBindingSource();
-
-	TArray<IFieldPathHelper*> Helpers = PathHelpers.Get(TArray<IFieldPathHelper*>());
-
-	bool bFirst = true;
-
-	TSet<FBindingSource> SourceIntersection;
-	TOptional<FBindingSource> NewSelection;
-
-	for (const IFieldPathHelper* Helper : Helpers)
-	{
-		// get available sources
-		TSet<FBindingSource> Sources;
-		Helper->GetAvailableSources(Sources);
-
-		if (bFirst)
-		{
-			bFirst = false; 
-
-			SourceIntersection.Append(Sources);
-		}
-		else
-		{
-			SourceIntersection.Intersect(Sources);
-		}
-	}
-
-	AvailableSources = SourceIntersection.Array();
-	SelectedSource = GetSelectedSource(Helpers);
+	SelectedSource = SelectedSourceAttribute.Get();
+	AvailableSources = AvailableSourcesAttribute.Get();
 
 	if (SourceComboBox.IsValid())
 	{
@@ -131,16 +83,7 @@ void SMVVMSourceSelector::Refresh()
 
 EVisibility SMVVMSourceSelector::GetClearVisibility() const
 {
-	for (const IFieldPathHelper* Helper : PathHelpers.Get())
-	{
-		const FBindingSource ThisSelection = Helper->GetSelectedSource();
-		if (ThisSelection.IsValid())
-		{
-			return EVisibility::Visible;
-		}
-	}
-
-	return EVisibility::Collapsed;
+	return SelectedSource.IsValid() ? EVisibility::Visible : EVisibility::Collapsed;
 }
 
 FReply SMVVMSourceSelector::OnClearSource()
diff --git a/Engine/Plugins/Runtime/ModelViewViewModel/Source/ModelViewViewModelEditor/Private/Widgets/SMVVMSourceSelector.h b/Engine/Plugins/Runtime/ModelViewViewModel/Source/ModelViewViewModelEditor/Private/Widgets/SMVVMSourceSelector.h
--- a/Engine/Plugins/Runtime/ModelViewViewModel/Source/ModelViewViewModelEditor/Private/Widgets/SMVVMSourceSelector.h
+++ b/Engine/Plugins/Runtime/ModelViewViewModel/Source/ModelViewViewModelEditor/Private/Widgets/SMVVMSourceSelector.h
@@ -3,7 +3,6 @@
 #pragma once
 
 #include "CoreMinimal.h"
-#include "MVVMPropertyPathHelpers.h"
 #include "Types/MVVMBindingSource.h"
 #include "Widgets/DeclarativeSyntaxSupport.h"
 #include "Widgets/SCompoundWidget.h"
@@ -21,7 +20,8 @@ public:
 		{
 		}
 		SLATE_STYLE_ARGUMENT(FTextBlockStyle, TextStyle)
-		SLATE_ATTRIBUTE(TArray<UE::MVVM::IFieldPathHelper*>, PathHelpers)
+		SLATE_ATTRIBUTE(UE::MVVM::FBindingSource, SelectedSource)
+		SLATE_ATTRIBUTE(TArray<UE::MVVM::FBindingSource>, AvailableSources)
 		SLATE_EVENT(FSelectionChanged, OnSelectionChanged)
 	SLATE_END_ARGS()
 
@@ -35,8 +35,9 @@ private:
 	FReply OnClearSource();
 
 private:
-	TAttribute<TArray<UE::MVVM::IFieldPathHelper*>> PathHelpers;
 	FSelectionChanged OnSelectionChanged;
+	TAttribute<TArray<UE::MVVM::FBindingSource>> AvailableSourcesAttribute;
+	TAttribute< UE::MVVM::FBindingSource> SelectedSourceAttribute;
 	TSharedPtr<SComboBox<UE::MVVM::FBindingSource>> SourceComboBox;
 	TArray<UE::MVVM::FBindingSource> AvailableSources;
 	UE::MVVM::FBindingSource SelectedSource;
diff --git a/Engine/Plugins/Runtime/ModelViewViewModel/Source/ModelViewViewModelEditor/Private/Widgets/SMVVMViewBindingListView.cpp b/Engine/Plugins/Runtime/ModelViewViewModel/Source/ModelViewViewModelEditor/Private/Widgets/SMVVMViewBindingListView.cpp
--- a/Engine/Plugins/Runtime/ModelViewViewModel/Source/ModelViewViewModelEditor/Private/Widgets/SMVVMViewBindingListView.cpp
+++ b/Engine/Plugins/Runtime/ModelViewViewModel/Source/ModelViewViewModelEditor/Private/Widgets/SMVVMViewBindingListView.cpp
@@ -2,7 +2,9 @@
 
 #include "Widgets/SMVVMViewBindingListView.h"
 
+#include "Algo/Transform.h"
 #include "Bindings/MVVMBindingHelper.h"
+#include "Blueprint/WidgetTree.h"
 #include "Dialog/SCustomDialog.h"
 #include "Editor.h"
 #include "Styling/AppStyle.h"
@@ -13,10 +15,8 @@
 #include "MVVMBlueprintView.h"
 #include "MVVMBlueprintViewBinding.h"
 #include "MVVMEditorSubsystem.h"
-#include "MVVMPropertyPathHelpers.h"
 #include "MVVMSubsystem.h"
 #include "MVVMWidgetBlueprintExtension_View.h"
-#include "ScopedTransaction.h"
 #include "SEnumCombo.h"
 #include "SSimpleButton.h" 
 #include "Styling/MVVMEditorStyle.h"
@@ -72,9 +72,6 @@ public:
 
 		FMVVMBlueprintViewBinding* ViewModelBinding = GetThisViewBinding();
 
-		ViewModelHelper = MakeUnique<UE::MVVM::FViewModelFieldPathHelper>(&ViewModelBinding->ViewModelPath, WidgetBlueprint);
-		WidgetHelper = MakeUnique<UE::MVVM::FWidgetFieldPathHelper>(&ViewModelBinding->WidgetPath, WidgetBlueprint);
-
 		SMultiColumnTableRow<FMVVMViewBindingListEntryPtr>::Construct(
 			FSuperRowType::FArguments()
 			.Padding(1.0f),
@@ -137,7 +134,8 @@ public:
 				.VAlign(VAlign_Center)
 				[
 					SAssignNew(ViewModelSourceSelector, SMVVMSourceSelector)
-					.PathHelpers(this, &SMVVMViewBindingListEntryRow::GetViewModelHelpers)
+					.AvailableSources(this, &SMVVMViewBindingListEntryRow::GetAvailableViewModels)
+					.SelectedSource(this, &SMVVMViewBindingListEntryRow::GetSelectedViewModel)
 					.OnSelectionChanged(this, &SMVVMViewBindingListEntryRow::OnViewModelSelectionChanged)
 				];
 		}
@@ -149,8 +147,9 @@ public:
 				.VAlign(VAlign_Center)
 				[
 					SAssignNew(ViewModelFieldSelector, SMVVMFieldSelector)
-					.PathHelpers(GetViewModelHelpers())
-					.CounterpartHelpers(GetWidgetHelpers())
+					.SelectedSource(this, &SMVVMViewBindingListEntryRow::GetSelectedViewModel)
+					.AvailableFields(this, &SMVVMViewBindingListEntryRow::GetAvailableViewModelProperties)
+					.SelectedField(this, &SMVVMViewBindingListEntryRow::GetSelectedViewModelProperty)
 					.BindingMode(this, &SMVVMViewBindingListEntryRow::GetCurrentBindingMode)
 					.IsSource(true)
 					.OnSelectionChanged(this, &SMVVMViewBindingListEntryRow::OnViewModelPropertySelectionChanged)
@@ -199,7 +198,8 @@ public:
 				.VAlign(VAlign_Center)
 				[
 					SAssignNew(WidgetSourceSelector, SMVVMSourceSelector)
-					.PathHelpers(this, &SMVVMViewBindingListEntryRow::GetWidgetHelpers)
+					.AvailableSources(this, &SMVVMViewBindingListEntryRow::GetAvailableWidgets)
+					.SelectedSource(this, &SMVVMViewBindingListEntryRow::GetSelectedWidget)
 					.OnSelectionChanged(this, &SMVVMViewBindingListEntryRow::OnWidgetSelectionChanged)
 				];
 		}
@@ -210,8 +210,9 @@ public:
 				.VAlign(VAlign_Center)
 				[
 					SAssignNew(WidgetFieldSelector, SMVVMFieldSelector)
-					.PathHelpers(GetWidgetHelpers())
-					.CounterpartHelpers(GetViewModelHelpers())
+					.SelectedSource(this, &SMVVMViewBindingListEntryRow::GetSelectedWidget)
+					.AvailableFields(this, &SMVVMViewBindingListEntryRow::GetAvailableWidgetProperties)
+					.SelectedField(this, &SMVVMViewBindingListEntryRow::GetSelectedWidgetProperty)
 					.BindingMode(this, &SMVVMViewBindingListEntryRow::GetCurrentBindingMode)
 					.IsSource(false)
 					.OnSelectionChanged(this, &SMVVMViewBindingListEntryRow::OnWidgetPropertySelectionChanged)
@@ -343,21 +344,18 @@ private:
 
 	UMVVMSubsystem::FConstDirectionalBindingArgs GetCurrentBindingArgs(bool bIsGetter) const
 	{
-		UE::MVVM::FMVVMConstFieldVariant ViewModelField = ViewModelHelper->GetSelectedField();
-		UE::MVVM::FMVVMConstFieldVariant WidgetField = WidgetHelper->GetSelectedField();
-
 		const FMVVMBlueprintViewBinding* ViewModelBinding = GetThisViewBinding();
 
 		UMVVMSubsystem::FConstDirectionalBindingArgs Args;
 		if (bIsGetter)
 		{
-			Args.SourceBinding = ViewModelField;
-			Args.DestinationBinding = WidgetField;
+			Args.SourceBinding = ViewModelBinding->ViewModelPath.GetFields().Last();
+			Args.DestinationBinding = ViewModelBinding->WidgetPath.GetFields().Last();
 		}
 		else
 		{
-			Args.SourceBinding = WidgetField;
-			Args.DestinationBinding = ViewModelField;
+			Args.SourceBinding = ViewModelBinding->WidgetPath.GetFields().Last();
+			Args.DestinationBinding = ViewModelBinding->ViewModelPath.GetFields().Last();
 		}
 
 		return MoveTemp(Args);
@@ -374,176 +372,244 @@ private:
 			];
 	}
 
-	void OnPreEditChange(FName PropertyName) const
+	void OnConversionFunctionChanged(const UFunction* Function, bool bSourceToDest)
 	{
-		if (UMVVMBlueprintView* BlueprintViewPtr = BlueprintView.Get())
+		if (FMVVMBlueprintViewBinding* ViewBinding = GetThisViewBinding())
 		{
-			FProperty* ChangedProperty = FMVVMBlueprintViewBinding::StaticStruct()->FindPropertyByName(PropertyName);
-			check(ChangedProperty != nullptr);
+			UMVVMEditorSubsystem* EditorSubsystem = GEditor->GetEditorSubsystem<UMVVMEditorSubsystem>();
+
+			if (bSourceToDest)
+			{
+				EditorSubsystem->SetSourceToDestinationConversionFunction(WidgetBlueprint, *ViewBinding, Function);
+			} 
+			else
+			{
+				EditorSubsystem->SetDestinationToSourceConversionFunction(WidgetBlueprint, *ViewBinding, Function);
+			}
+		}
+	}
 
-			FEditPropertyChain EditChain;
-			EditChain.AddTail(UMVVMBlueprintView::StaticClass()->FindPropertyByName("Bindings"));
-			EditChain.AddTail(ChangedProperty);
-			EditChain.SetActivePropertyNode(ChangedProperty);
+	TArray<UE::MVVM::FBindingSource> GetAvailableViewModels() const
+	{
+		UMVVMEditorSubsystem* EditorSubsystem = GEditor->GetEditorSubsystem<UMVVMEditorSubsystem>();
+		return EditorSubsystem->GetAllViewModels(WidgetBlueprint);
+	}
 
-			BlueprintViewPtr->PreEditChange(EditChain);
+	UE::MVVM::FBindingSource GetSelectedViewModel() const
+	{
+		if (FMVVMBlueprintViewBinding* ViewBinding = GetThisViewBinding())
+		{
+			return UE::MVVM::FBindingSource::CreateForViewModel(WidgetBlueprint, ViewBinding->ViewModelPath.GetViewModelId());
 		}
+
+		return UE::MVVM::FBindingSource();
 	}
 
-	void OnPostEditChange(FName PropertyName) const
+	TArray<FMVVMBlueprintPropertyPath> GetAvailableViewModelProperties() const
 	{
-		if (UMVVMBlueprintView* BlueprintViewPtr = BlueprintView.Get())
+		TArray<FMVVMBlueprintPropertyPath> AvailablePaths;
+
+		if (FMVVMBlueprintViewBinding* ViewBinding = GetThisViewBinding())
 		{
-			FProperty* ChangedProperty = FMVVMBlueprintViewBinding::StaticStruct()->FindPropertyByName(PropertyName);
-			check(ChangedProperty != nullptr);
+			UMVVMBlueprintView* BlueprintViewPtr = BlueprintView.Get();
+			if (const FMVVMBlueprintViewModelContext* ViewModel = BlueprintViewPtr->FindViewModel(ViewBinding->ViewModelPath.GetViewModelId()))
+			{
+				UMVVMSubsystem* Subsystem = GEngine->GetEngineSubsystem<UMVVMSubsystem>();
+				TArray<FMVVMAvailableBinding> AvailableBindings = Subsystem->GetViewModelAvailableBindings(ViewModel->GetViewModelClass());
+
+				Algo::Transform(AvailableBindings, AvailablePaths, [ViewModel](const FMVVMAvailableBinding& Binding)
+					{
+						const UClass* Class = ViewModel->GetViewModelClass();
+						FName BindingName = Binding.GetBindingName().ToName();
+						
+						UE::MVVM::FMVVMConstFieldVariant Variant; 
+
+						if (const UFunction* Function = Class->FindFunctionByName(BindingName))
+						{
+							Variant = UE::MVVM::FMVVMConstFieldVariant(Function);
+						}
+						else if (const FProperty* Property = Class->FindPropertyByName(BindingName))
+						{
+							Variant = UE::MVVM::FMVVMConstFieldVariant(Property);
+						}
+
+						FMVVMBlueprintPropertyPath Path;
+
+						if (!Variant.IsEmpty())
+						{
+							Path.SetViewModelId(ViewModel->GetViewModelId());
+							Path.SetBasePropertyPath(Variant);
+						}
+
+						return Path;
+					});
+			}
+		}
 
-			FEditPropertyChain EditChain;
-			EditChain.AddTail(UMVVMBlueprintView::StaticClass()->FindPropertyByName("Bindings"));
-			EditChain.AddTail(ChangedProperty);
-			EditChain.SetActivePropertyNode(ChangedProperty);
+		return AvailablePaths;
+	}
 
-			FPropertyChangedEvent ChangeEvent(ChangedProperty, EPropertyChangeType::ValueSet);
-			FPropertyChangedChainEvent ChainEvent(EditChain, ChangeEvent);
-			BlueprintViewPtr->PostEditChangeChainProperty(ChainEvent);
+	FMVVMBlueprintPropertyPath GetSelectedViewModelProperty() const
+	{
+		if (FMVVMBlueprintViewBinding* ViewBinding = GetThisViewBinding())
+		{
+			return ViewBinding->ViewModelPath;
 		}
+
+		return FMVVMBlueprintPropertyPath();
 	}
 
 	void OnViewModelSelectionChanged(UE::MVVM::FBindingSource Source)
 	{
-		OnSourceSelectionChanged(Source, ViewModelHelper.Get(), false);
+		if (FMVVMBlueprintViewBinding* ViewBinding = GetThisViewBinding())
+		{
+			const FMVVMBlueprintViewModelContext* ViewModel = BlueprintView.Get()->FindViewModel(Source.ViewModelId);
+
+			UMVVMEditorSubsystem* EditorSubsystem = GEditor->GetEditorSubsystem<UMVVMEditorSubsystem>();
+			EditorSubsystem->SetViewModelForBinding(WidgetBlueprint, *ViewBinding, ViewModel != nullptr ? *ViewModel : FMVVMBlueprintViewModelContext());
+		}
+
 		if (ViewModelFieldSelector.IsValid())
 		{
 			ViewModelFieldSelector->Refresh();
 		}
 	}
 
-	void OnWidgetSelectionChanged(UE::MVVM::FBindingSource Source)
+	TArray<UE::MVVM::FBindingSource> GetAvailableWidgets() const
 	{
-		OnSourceSelectionChanged(Source, WidgetHelper.Get(), true);
-		if (WidgetFieldSelector.IsValid())
-		{
-			WidgetFieldSelector->Refresh();
-		}
+		UMVVMEditorSubsystem* EditorSubsystem = GEditor->GetEditorSubsystem<UMVVMEditorSubsystem>();
+		return EditorSubsystem->GetBindableWidgets(WidgetBlueprint);
 	}
-
-	void OnSourceSelectionChanged(UE::MVVM::FBindingSource SelectedSource, UE::MVVM::IFieldPathHelper* PathHelper, bool bIsWidget)
+	
+	TArray<FMVVMBlueprintPropertyPath> GetAvailableWidgetProperties() const
 	{
-		if (UMVVMBlueprintView* BlueprintViewPtr = BlueprintView.Get())
+		TArray<FMVVMBlueprintPropertyPath> AvailablePaths;
+
+		if (FMVVMBlueprintViewBinding* ViewBinding = GetThisViewBinding())
 		{
-			FScopedTransaction Transaction(LOCTEXT("SetBindingSource", "Set Binding Source"));
+			UClass* Class = nullptr;
 
-			FName ChangedProperty;
-			if (bIsWidget)
+			const FName WidgetName = ViewBinding->WidgetPath.GetWidgetName();
+			if (WidgetName.IsEqual(WidgetBlueprint->GetFName()))
 			{
-				ChangedProperty = GET_MEMBER_NAME_CHECKED(FMVVMBlueprintViewBinding, WidgetPath);
+				Class = WidgetBlueprint->GeneratedClass;
 			}
-			else
+			else if (const UWidget* Widget = WidgetBlueprint->WidgetTree->FindWidget(WidgetName))
 			{
-				ChangedProperty = GET_MEMBER_NAME_CHECKED(FMVVMBlueprintViewBinding, ViewModelPath);
+				Class = Widget->GetClass();
 			}
 
-			OnPreEditChange(ChangedProperty);
+			if (Class != nullptr)
+			{
+				UMVVMSubsystem* Subsystem = GEngine->GetEngineSubsystem<UMVVMSubsystem>();
+				TArray<FMVVMAvailableBinding> AvailableBindings = Subsystem->GetWidgetAvailableBindings(Class);
+
+				Algo::Transform(AvailableBindings, AvailablePaths, [Class, WidgetName](const FMVVMAvailableBinding& Binding)
+					{
+						FName BindingName = Binding.GetBindingName().ToName();
+
+						UE::MVVM::FMVVMConstFieldVariant Variant;
+
+						if (const UFunction* Function = Class->FindFunctionByName(BindingName))
+						{
+							Variant = UE::MVVM::FMVVMConstFieldVariant(Function);
+						}
+						else if (const FProperty* Property = Class->FindPropertyByName(BindingName))
+						{
+							Variant = UE::MVVM::FMVVMConstFieldVariant(Property);
+						}
 
-			PathHelper->SetSelectedSource(SelectedSource);
-			PathHelper->ResetBinding(); // Might make sense to keep this around in case we retarget to a compatible widget or switch back.
+						FMVVMBlueprintPropertyPath Path;
 
-			OnPostEditChange(ChangedProperty);
+						if (!Variant.IsEmpty())
+						{
+							Path.SetWidgetName(WidgetName);
+							Path.SetBasePropertyPath(Variant);
+						}
+
+						return Path;
+					});
+			}
 		}
+
+		return AvailablePaths;
 	}
 
-	void OnViewModelPropertySelectionChanged(UE::MVVM::FMVVMConstFieldVariant SelectedField)
+	FMVVMBlueprintPropertyPath GetSelectedWidgetProperty() const
 	{
-		UE::MVVM::FMVVMConstFieldVariant CurrentField = ViewModelHelper->GetSelectedField();
-		if (CurrentField != SelectedField)
+		if (FMVVMBlueprintViewBinding* ViewBinding = GetThisViewBinding())
 		{
-			OnPropertySelectionChanged(SelectedField, ViewModelHelper.Get(), false);
-
-			if (WidgetFieldSelector.IsValid())
-			{
-				WidgetFieldSelector->Refresh();
-			}
+			return ViewBinding->WidgetPath;
 		}
+
+		return FMVVMBlueprintPropertyPath();
 	}
 
-	void OnWidgetPropertySelectionChanged(UE::MVVM::FMVVMConstFieldVariant SelectedField)
+	UE::MVVM::FBindingSource GetSelectedWidget() const
 	{
-		UE::MVVM::FMVVMConstFieldVariant CurrentField = WidgetHelper->GetSelectedField();
-		if (CurrentField != SelectedField)
+		if (FMVVMBlueprintViewBinding* ViewBinding = GetThisViewBinding())
 		{
-			OnPropertySelectionChanged(SelectedField, WidgetHelper.Get(), true);
-
-			if (ViewModelFieldSelector.IsValid())
-			{
-				ViewModelFieldSelector->Refresh();
-			}
+			return UE::MVVM::FBindingSource::CreateForWidget(WidgetBlueprint, ViewBinding->WidgetPath.GetWidgetName());
 		}
+
+		return UE::MVVM::FBindingSource();
 	}
 
-	void OnPropertySelectionChanged(UE::MVVM::FMVVMConstFieldVariant SelectedField, UE::MVVM::IFieldPathHelper* PathHelper, bool bIsWidget)
+	void OnWidgetSelectionChanged(UE::MVVM::FBindingSource Source)
 	{
-		FScopedTransaction Transaction(LOCTEXT("SetBindingProperty", "Set Binding Property"));
-
-		FName ChangedProperty;
-		if (bIsWidget)
+		if (FMVVMBlueprintViewBinding* ViewBinding = GetThisViewBinding())
 		{
-			ChangedProperty = GET_MEMBER_NAME_CHECKED(FMVVMBlueprintViewBinding, WidgetPath);
+			GEditor->GetEditorSubsystem<UMVVMEditorSubsystem>()->SetWidgetForBinding(WidgetBlueprint, *ViewBinding, Source.Name);
 		}
-		else
+
+		if (WidgetFieldSelector.IsValid())
 		{
-			ChangedProperty = GET_MEMBER_NAME_CHECKED(FMVVMBlueprintViewBinding, ViewModelPath);
+			WidgetFieldSelector->Refresh();
 		}
-
-		OnPreEditChange(ChangedProperty);
-
-		PathHelper->SetBindingReference(SelectedField);
-
-		OnPostEditChange(ChangedProperty);
 	}
 
-	void OnUpdateModeSelectionChanged(int32 Value, ESelectInfo::Type)
+	void OnViewModelPropertySelectionChanged(FMVVMBlueprintPropertyPath SelectedField)
 	{
-		if (FMVVMBlueprintViewBinding* ViewModelBinding = GetThisViewBinding())
+		if (FMVVMBlueprintViewBinding* ViewBinding = GetThisViewBinding())
 		{
-			if (ViewModelBinding->UpdateMode != (EMVVMViewBindingUpdateMode) Value)
+			if (ViewBinding->ViewModelPath != SelectedField)
 			{
-				FScopedTransaction Transaction(LOCTEXT("SetUpdateMode", "Set Update Mode"));
-
-				OnPreEditChange(GET_MEMBER_NAME_CHECKED(FMVVMBlueprintViewBinding, UpdateMode));
+				UMVVMEditorSubsystem* Subsystem = GEditor->GetEditorSubsystem<UMVVMEditorSubsystem>();
+				Subsystem->SetWidgetPropertyForBinding(WidgetBlueprint, *ViewBinding, SelectedField);
 
-				ViewModelBinding->UpdateMode = (EMVVMViewBindingUpdateMode) Value;
-
-				OnPostEditChange(GET_MEMBER_NAME_CHECKED(FMVVMBlueprintViewBinding, UpdateMode));
+				if (WidgetFieldSelector.IsValid())
+				{
+					WidgetFieldSelector->Refresh();
+				}
 			}
 		}
 	}
 
-	void OnConversionFunctionChanged(const UFunction* Function, bool bSourceToDestination)
+	void OnWidgetPropertySelectionChanged(FMVVMBlueprintPropertyPath SelectedField)
 	{
-		if (FMVVMBlueprintViewBinding* ViewModelBinding = GetThisViewBinding())
+		if (FMVVMBlueprintViewBinding* ViewBinding = GetThisViewBinding())
 		{
-			FScopedTransaction Transaction(LOCTEXT("SetConversionFunction", "Set Conversion Function"));
-
-			FMemberReference NewReference;
-			if (Function != nullptr)
+			if (ViewBinding->WidgetPath != SelectedField)
 			{
-				NewReference.SetFromField<UFunction>(Function, WidgetBlueprint->SkeletonGeneratedClass);
-			}
-				
-			OnPreEditChange(GET_MEMBER_NAME_CHECKED(FMVVMBlueprintViewBinding, Conversion));
+				UMVVMEditorSubsystem* Subsystem = GEditor->GetEditorSubsystem<UMVVMEditorSubsystem>();
+				Subsystem->SetViewModelPropertyForBinding(WidgetBlueprint, *ViewBinding, SelectedField);
 
-			if (bSourceToDestination)
-			{
-				ViewModelBinding->Conversion.SourceToDestinationFunction = NewReference;
-			}
-			else
-			{
-				ViewModelBinding->Conversion.DestinationToSourceFunction = NewReference;
+				if (ViewModelFieldSelector.IsValid())
+				{
+					ViewModelFieldSelector->Refresh();
+				}
 			}
-
-			OnPostEditChange(GET_MEMBER_NAME_CHECKED(FMVVMBlueprintViewBinding, Conversion));
 		}
 	}
 
+	void OnUpdateModeSelectionChanged(int32 Value, ESelectInfo::Type)
+	{
+		if (FMVVMBlueprintViewBinding* ViewBinding = GetThisViewBinding())
+		{
+			UMVVMEditorSubsystem* Subsystem = GEditor->GetEditorSubsystem<UMVVMEditorSubsystem>();
+			Subsystem->SetUpdateModeForBinding(WidgetBlueprint, *ViewBinding, (EMVVMViewBindingUpdateMode) Value);
+		}			
+	}
 
 	int32 GetUpdateModeValue() const
 	{
@@ -557,17 +623,10 @@ private:
 			return;
 		}
 
-		if (FMVVMBlueprintViewBinding* ViewModelBinding = GetThisViewBinding())
+		if (FMVVMBlueprintViewBinding* ViewBinding = GetThisViewBinding())
 		{
-			bool bNewEnabled = (NewState == ECheckBoxState::Checked);
-			if (ViewModelBinding->bEnabled != bNewEnabled)
-			{
-				OnPreEditChange(GET_MEMBER_NAME_CHECKED(FMVVMBlueprintViewBinding, bEnabled));
-
-				ViewModelBinding->bEnabled = bNewEnabled;
-
-				OnPostEditChange(GET_MEMBER_NAME_CHECKED(FMVVMBlueprintViewBinding, bEnabled));
-			}
+			UMVVMEditorSubsystem* Subsystem = GEditor->GetEditorSubsystem<UMVVMEditorSubsystem>();
+			Subsystem->SetEnabledForBinding(WidgetBlueprint, *ViewBinding, NewState == ECheckBoxState::Checked);
 		}
 	}
 
@@ -578,20 +637,10 @@ private:
 			return;
 		}
 
-		if (FMVVMBlueprintViewBinding* ViewModelBinding = GetThisViewBinding())
+		if (FMVVMBlueprintViewBinding* ViewBinding = GetThisViewBinding())
 		{
-			bool bNewCompile = (NewState == ECheckBoxState::Checked);
-			if (ViewModelBinding->bCompile != bNewCompile)
-			{
-				if (UMVVMBlueprintView* BlueprintViewPtr = BlueprintView.Get())
-				{
-					BlueprintViewPtr->PreEditChange(UMVVMBlueprintView::StaticClass()->FindPropertyByName("Bindings"));
-
-					ViewModelBinding->bCompile = bNewCompile;
-
-					BlueprintViewPtr->PostEditChange();
-				}
-			}
+			UMVVMEditorSubsystem* Subsystem = GEditor->GetEditorSubsystem<UMVVMEditorSubsystem>();
+			Subsystem->SetCompileForBinding(WidgetBlueprint, *ViewBinding, NewState == ECheckBoxState::Checked);
 		}
 	}
 
@@ -675,20 +724,21 @@ private:
 
 	void OnModeSelectionChanged(FName ValueName, ESelectInfo::Type)
 	{
-		if (FMVVMBlueprintViewBinding* ViewModelBinding = GetThisViewBinding())
+		if (FMVVMBlueprintViewBinding* ViewBinding = GetThisViewBinding())
 		{
 			const UEnum* ModeEnum = StaticEnum<EMVVMBindingMode>();
 			EMVVMBindingMode NewMode = (EMVVMBindingMode) ModeEnum->GetValueByName(ValueName);
 
-			if (ViewModelBinding->BindingType != NewMode)
-			{
-				OnPreEditChange(GET_MEMBER_NAME_CHECKED(FMVVMBlueprintViewBinding, BindingType));
-
-				ViewModelBinding->BindingType = NewMode;
-
-				OnPostEditChange(GET_MEMBER_NAME_CHECKED(FMVVMBlueprintViewBinding, BindingType));
+			UMVVMEditorSubsystem* Subsystem = GEditor->GetEditorSubsystem<UMVVMEditorSubsystem>();
+			Subsystem->SetBindingTypeForBinding(WidgetBlueprint, *ViewBinding, NewMode);
 
+			if (ViewModelFieldSelector.IsValid())
+			{
 				ViewModelFieldSelector->Refresh();
+			}
+
+			if (WidgetFieldSelector.IsValid())
+			{
 				WidgetFieldSelector->Refresh();
 			}
 		}
@@ -716,16 +766,6 @@ private:
 		return FReply::Handled();
 	}
 
-	TArray<UE::MVVM::IFieldPathHelper*> GetWidgetHelpers() const
-	{
-		return TArray<UE::MVVM::IFieldPathHelper*> { WidgetHelper.Get() };
-	}
-
-	TArray<UE::MVVM::IFieldPathHelper*> GetViewModelHelpers() const
-	{
-		return TArray<UE::MVVM::IFieldPathHelper*> { ViewModelHelper.Get() };
-	}
-
 	void HandleBlueprintChanged(UBlueprint* Blueprint)
 	{
 		ViewModelSourceSelector->Refresh();
@@ -738,8 +778,6 @@ private:
 	FMVVMViewBindingListEntryPtr Entry;
 	TWeakObjectPtr<UMVVMBlueprintView> BlueprintView;
 	UWidgetBlueprint* WidgetBlueprint = nullptr;
-	TUniquePtr<UE::MVVM::FViewModelFieldPathHelper> ViewModelHelper;
-	TUniquePtr<UE::MVVM::FWidgetFieldPathHelper> WidgetHelper;
 	TSharedPtr<SMVVMSourceSelector> ViewModelSourceSelector;
 	TSharedPtr<SMVVMFieldSelector> ViewModelFieldSelector;
 	TSharedPtr<SMVVMSourceSelector> WidgetSourceSelector;
diff --git a/Engine/Plugins/Runtime/ModelViewViewModel/Source/ModelViewViewModelEditor/Private/Widgets/SMVVMViewModelContextListWidget.cpp b/Engine/Plugins/Runtime/ModelViewViewModel/Source/ModelViewViewModelEditor/Private/Widgets/SMVVMViewModelContextListWidget.cpp
--- a/Engine/Plugins/Runtime/ModelViewViewModel/Source/ModelViewViewModelEditor/Private/Widgets/SMVVMViewModelContextListWidget.cpp
+++ b/Engine/Plugins/Runtime/ModelViewViewModel/Source/ModelViewViewModelEditor/Private/Widgets/SMVVMViewModelContextListWidget.cpp
@@ -385,8 +385,11 @@ namespace UE::MVVM::Private
 
 			auto MakeFieldWidget = [](FMVVMConstFieldVariant Field)
 			{
+				FMVVMBlueprintPropertyPath Path;
+				Path.SetBasePropertyPath(Field);
+
 				return SNew(SMVVMFieldEntry)
-					.Field(Field);
+					.Field(Path);
 			};
 
 			UClass* InBindableClass = Cast<UClass>(InOwnerStruct);
diff --git a/Engine/Plugins/Runtime/ModelViewViewModel/Source/ModelViewViewModelEditor/Public/MVVMEditorSubsystem.h b/Engine/Plugins/Runtime/ModelViewViewModel/Source/ModelViewViewModelEditor/Public/MVVMEditorSubsystem.h
--- a/Engine/Plugins/Runtime/ModelViewViewModel/Source/ModelViewViewModelEditor/Public/MVVMEditorSubsystem.h
+++ b/Engine/Plugins/Runtime/ModelViewViewModel/Source/ModelViewViewModelEditor/Public/MVVMEditorSubsystem.h
@@ -6,10 +6,13 @@
 
 #include "MVVMBlueprintView.h"
 #include "Templates/SubclassOf.h"
+#include "Types/MVVMBindingMode.h"
+#include "Types/MVVMBindingSource.h"
 #include "Types/MVVMFieldVariant.h"
 
 #include "MVVMEditorSubsystem.generated.h"
 
+class UEdGraph;
 class UMVVMViewModelBase;
 class UWidgetBlueprint;
 
@@ -24,7 +27,7 @@ public:
 	UMVVMBlueprintView* RequestView(UWidgetBlueprint* WidgetBlueprint) const;
 
 	UFUNCTION(BlueprintCallable, Category = "MVVM")
-	UMVVMBlueprintView* GetView(UWidgetBlueprint* WidgetBlueprint) const;
+	UMVVMBlueprintView* GetView(const UWidgetBlueprint* WidgetBlueprint) const;
 	
 	UFUNCTION(BlueprintCallable, Category = "MVVM")
 	void RemoveViewModel(UWidgetBlueprint* WidgetBlueprint, FName ViewModel);
@@ -41,7 +44,36 @@ public:
 	UFUNCTION(BlueprintCallable, Category = "MVVM")
 	void RemoveBinding(UWidgetBlueprint* WidgetBlueprint, const FMVVMBlueprintViewBinding& Binding);
 
-	TArray<UE::MVVM::FMVVMConstFieldVariant> GetChildViewModels(UClass* Class);
+	void SetSourceToDestinationConversionFunction(UWidgetBlueprint* WidgetBlueprint, FMVVMBlueprintViewBinding& Binding, const UFunction* ConversionFunction);
+	void SetDestinationToSourceConversionFunction(UWidgetBlueprint* WidgetBlueprint, FMVVMBlueprintViewBinding& Binding, const UFunction* ConversionFunction);
+	void SetWidgetForBinding(UWidgetBlueprint* WidgetBlueprint, FMVVMBlueprintViewBinding& Binding, FName Widget);
+	void SetWidgetPropertyForBinding(UWidgetBlueprint* WidgetBlueprint, FMVVMBlueprintViewBinding& Binding, FMVVMBlueprintPropertyPath Field);
+	void SetViewModelForBinding(UWidgetBlueprint* WidgetBlueprint, FMVVMBlueprintViewBinding& Binding, FMVVMBlueprintViewModelContext ViewModel);
+	void SetViewModelPropertyForBinding(UWidgetBlueprint* WidgetBlueprint, FMVVMBlueprintViewBinding& Binding, FMVVMBlueprintPropertyPath Field);
+	void SetUpdateModeForBinding(UWidgetBlueprint* WidgetBlueprint, FMVVMBlueprintViewBinding& Binding, EMVVMViewBindingUpdateMode Mode);
+	void SetBindingTypeForBinding(UWidgetBlueprint* WidgetBlueprint, FMVVMBlueprintViewBinding& Binding, EMVVMBindingMode Type);
+	void SetEnabledForBinding(UWidgetBlueprint* WidgetBlueprint, FMVVMBlueprintViewBinding& Binding, bool bEnabled);
+	void SetCompileForBinding(UWidgetBlueprint* WidgetBlueprint, FMVVMBlueprintViewBinding& Binding, bool bCompile);
 
-	TArray<const UFunction*> GetAvailableConversionFunctions(const UE::MVVM::FMVVMConstFieldVariant& Source, const UE::MVVM::FMVVMConstFieldVariant& Dest) const;
+	UFUNCTION(BlueprintCallable, Category = "MVVM")
+	bool IsValidConversionFunction(const UFunction* Function, const FMVVMBlueprintPropertyPath& Source, const FMVVMBlueprintPropertyPath& Destination) const;
+
+	UFUNCTION(BlueprintCallable, Category = "MVVM")
+	UEdGraph* GetConversionFunctionGraph(const UWidgetBlueprint* WidgetBlueprint, const FMVVMBlueprintViewBinding& Binding, bool bSourceToDestination) const;
+
+	UFUNCTION(BlueprintCallable, Category = "MVVM")
+	TArray<UFunction*> GetAvailableConversionFunctions(const UWidgetBlueprint* WidgetBlueprint, const FMVVMBlueprintPropertyPath& Source, const FMVVMBlueprintPropertyPath& Destination) const;
+
+	FMVVMBlueprintPropertyPath GetPathForConversionFunctionArgument(const UWidgetBlueprint* WidgetBlueprint, const FMVVMBlueprintViewBinding& Binding, FName ArgumentName, bool bSourceToDestination) const;
+	void SetPathForConversionFunctionArgument(UWidgetBlueprint* WidgetBlueprint, FMVVMBlueprintViewBinding& Binding, FName ArgumentName, const FMVVMBlueprintPropertyPath& Path, bool bSourceToDestination) const;
+
+	TArray<UE::MVVM::FBindingSource> GetBindableWidgets(const UWidgetBlueprint* WidgetBlueprint) const;
+	
+	TArray<UE::MVVM::FBindingSource> GetAllViewModels(const UWidgetBlueprint* WidgetBlueprint) const;
+
+	TArray<UE::MVVM::FMVVMConstFieldVariant> GetChildViewModels(const UClass* Class);
+
+private:
+	FName GetConversionFunctionWrapperName(const UWidgetBlueprint* WidgetBlueprint, const FMVVMBlueprintViewBinding& Binding, bool bSourceToDestination) const;
+	UEdGraph* CreateConversionFunctionWrapperGraph(UWidgetBlueprint* WidgetBlueprint, const FMVVMBlueprintViewBinding& Binding, const UFunction* ConversionFunction, bool bSourceToDestination);
 };
