Fixing build failure caused by missed cpp that somehow passed preflight.

#rnx

diff --git a/Engine/Plugins/Runtime/ModelViewViewModel/Source/ModelViewViewModelBlueprint/Private/MVVMBlueprintViewBinding.cpp b/Engine/Plugins/Runtime/ModelViewViewModel/Source/ModelViewViewModelBlueprint/Private/MVVMBlueprintViewBinding.cpp
new file mode 100644
--- /dev/null
+++ b/Engine/Plugins/Runtime/ModelViewViewModel/Source/ModelViewViewModelBlueprint/Private/MVVMBlueprintViewBinding.cpp
@@ -0,0 +1,118 @@
+// Copyright Epic Games, Inc. All Rights Reserved.
+
+#include "MVVMBlueprintViewBinding.h"
+#include "MVVMBlueprintView.h"
+#include "MVVMWidgetBlueprintExtension_View.h"
+#include "WidgetBlueprintExtension.h"
+
+FName FMVVMBlueprintViewBinding::GetFName(const UMVVMBlueprintView* View) const
+{
+	TStringBuilder<256> BindingName;
+
+	const FMVVMBlueprintViewModelContext* ViewModel = nullptr;
+	if (View != nullptr)
+	{
+		ViewModel = View->FindViewModel(ViewModelPath.GetViewModelId());
+	}
+
+	if (ViewModel != nullptr)
+	{
+		BindingName << ViewModel->GetViewModelName();
+	}
+	else
+	{
+		BindingName << TEXT("Invalid");
+	}
+
+	BindingName << TEXT("_");
+	BindingName << ViewModelPath.GetBasePropertyPath();
+
+	if (BindingType == EMVVMBindingMode::TwoWay)
+	{
+		BindingName << TEXT("_Both_");
+	}
+	else if (UE::MVVM::IsForwardBinding(BindingType))
+	{
+		BindingName << TEXT("_To_");
+	}
+	else if (UE::MVVM::IsBackwardBinding(BindingType))
+	{
+		BindingName << TEXT("_From_");
+	}
+
+	if (View->GetOuterUMVVMWidgetBlueprintExtension_View()->GetWidgetBlueprint()->GetFName() != WidgetPath.GetWidgetName())
+	{
+		if (WidgetPath.GetWidgetName().IsNone())
+		{
+			BindingName << TEXT("Invalid");
+		}
+		else
+		{
+			BindingName << WidgetPath.GetWidgetName();
+		}
+		BindingName << TEXT("_");
+	}
+
+	BindingName << WidgetPath.GetBasePropertyPath();
+
+	return BindingName.ToString();
+}
+
+FString FMVVMBlueprintViewBinding::GetDisplayNameString(const UMVVMBlueprintView* View) const
+{
+	const FMVVMBlueprintViewModelContext* ViewModel = nullptr;
+	if (View != nullptr)
+	{
+		ViewModel = View->FindViewModel(ViewModelPath.GetViewModelId());
+	}
+
+	TStringBuilder<512> BindingName;
+	if (ViewModel != nullptr)
+	{
+		BindingName << ViewModel->GetViewModelName();
+	}
+	else
+	{
+		BindingName << TEXT("<none>");
+	}
+
+	BindingName << TEXT(".");
+	BindingName << ViewModelPath.GetBasePropertyPath();
+
+	if (BindingType == EMVVMBindingMode::TwoWay)
+	{
+		BindingName << TEXT(" <-> ");
+	}
+	else if (UE::MVVM::IsForwardBinding(BindingType))
+	{
+		BindingName << TEXT(" -> ");
+	}
+	else if (UE::MVVM::IsBackwardBinding(BindingType))
+	{
+		BindingName << TEXT(" <- ");
+	}
+	else
+	{
+		BindingName << TEXT(" ??? "); // shouldn't happen
+	}
+
+	if (View->GetOuterUMVVMWidgetBlueprintExtension_View()->GetWidgetBlueprint()->GetFName() == WidgetPath.GetWidgetName())
+	{
+		BindingName << WidgetPath.GetBasePropertyPath();
+	}
+	else
+	{
+		if (WidgetPath.GetWidgetName().IsNone())
+		{
+			BindingName << TEXT("<none>");
+		}
+		else
+		{
+			BindingName << WidgetPath.GetWidgetName();
+		}
+		BindingName << TEXT(".");
+		BindingName << WidgetPath.GetBasePropertyPath();
+	}
+
+	return BindingName.ToString();
+}
