MVVM: Remove FMVVMBlueprintViewModelContext::OverrideDisplayName and replace it by the ViewModelName as a FName.

diff --git a/Engine/Plugins/Runtime/ModelViewViewModel/Source/ModelViewViewModelBlueprint/Private/MVVMBlueprintView.cpp b/Engine/Plugins/Runtime/ModelViewViewModel/Source/ModelViewViewModelBlueprint/Private/MVVMBlueprintView.cpp
--- a/Engine/Plugins/Runtime/ModelViewViewModel/Source/ModelViewViewModelBlueprint/Private/MVVMBlueprintView.cpp
+++ b/Engine/Plugins/Runtime/ModelViewViewModel/Source/ModelViewViewModelBlueprint/Private/MVVMBlueprintView.cpp
@@ -79,7 +79,7 @@ bool UMVVMBlueprintView::RenameViewModel(FName OldViewModelName, FName NewViewMo
 			});
 	if (ViewModelContext)
 	{
-		ViewModelContext->OverrideDisplayName = FText::FromName(NewViewModelName);
+		ViewModelContext->ViewModelName = NewViewModelName;
 
 		FBlueprintEditorUtils::ReplaceVariableReferences(GetOuterUMVVMWidgetBlueprintExtension_View()->GetWidgetBlueprint(), OldViewModelName, NewViewModelName);
 		FBlueprintEditorUtils::ValidateBlueprintChildVariables(GetOuterUMVVMWidgetBlueprintExtension_View()->GetWidgetBlueprint(), NewViewModelName);
diff --git a/Engine/Plugins/Runtime/ModelViewViewModel/Source/ModelViewViewModelBlueprint/Private/MVVMBlueprintViewModelContext.cpp b/Engine/Plugins/Runtime/ModelViewViewModel/Source/ModelViewViewModelBlueprint/Private/MVVMBlueprintViewModelContext.cpp
--- a/Engine/Plugins/Runtime/ModelViewViewModel/Source/ModelViewViewModelBlueprint/Private/MVVMBlueprintViewModelContext.cpp
+++ b/Engine/Plugins/Runtime/ModelViewViewModel/Source/ModelViewViewModelBlueprint/Private/MVVMBlueprintViewModelContext.cpp
@@ -2,22 +2,16 @@
 
 #include "MVVMBlueprintViewModelContext.h"
 
-FMVVMBlueprintViewModelContext::FMVVMBlueprintViewModelContext(TSubclassOf<UMVVMViewModelBase> InClass, FGuid InGuid) :
-	ViewModelContextId(InGuid),
-	ViewModelClass(InClass)
+FMVVMBlueprintViewModelContext::FMVVMBlueprintViewModelContext(TSubclassOf<UMVVMViewModelBase> InClass, FName InViewModelName)
+	: ViewModelContextId(FGuid::NewGuid())
+	, ViewModelClass(InClass)
+	, ViewModelName(InViewModelName)
 {
 
 }
 
+
 FText FMVVMBlueprintViewModelContext::GetDisplayName() const
 {
-	if (!OverrideDisplayName.IsEmpty())
-	{
-		return OverrideDisplayName;
-	}
-	if (ViewModelClass.Get())
-	{
-		return ViewModelClass->GetDisplayNameText();
-	}
-	return FText::GetEmpty();
+	return FText::FromName(ViewModelName);
 }
diff --git a/Engine/Plugins/Runtime/ModelViewViewModel/Source/ModelViewViewModelBlueprint/Public/MVVMBlueprintViewModelContext.h b/Engine/Plugins/Runtime/ModelViewViewModel/Source/ModelViewViewModelBlueprint/Public/MVVMBlueprintViewModelContext.h
--- a/Engine/Plugins/Runtime/ModelViewViewModel/Source/ModelViewViewModelBlueprint/Public/MVVMBlueprintViewModelContext.h
+++ b/Engine/Plugins/Runtime/ModelViewViewModel/Source/ModelViewViewModelBlueprint/Public/MVVMBlueprintViewModelContext.h
@@ -30,8 +30,8 @@ struct MODELVIEWVIEWMODELBLUEPRINT_API FMVVMBlueprintViewModelContext
 	GENERATED_BODY()
 
 public:
-	FMVVMBlueprintViewModelContext() {}
-	FMVVMBlueprintViewModelContext(TSubclassOf<UMVVMViewModelBase> InClass, FGuid InGuid);
+	FMVVMBlueprintViewModelContext() = default;
+	FMVVMBlueprintViewModelContext(TSubclassOf<UMVVMViewModelBase> InClass, FName InViewModelName);
 
 	FGuid GetViewModelId() const
 	{
@@ -40,8 +40,7 @@ public:
 
 	FName GetViewModelName() const
 	{
-		//todo change that to the fname
-		return *GetDisplayName().ToString();
+		return ViewModelName;
 	}
 
 	FText GetDisplayName() const;
@@ -51,6 +50,21 @@ public:
 		return ViewModelClass;
 	}
 
+	void PostSerialize(const FArchive& Ar)
+	{
+		if (Ar.IsLoading())
+		{
+			if (ViewModelName.IsNone())
+			{
+				ViewModelName = *OverrideDisplayName_DEPRECATED.ToString();
+			}
+			if (ViewModelName.IsNone() )
+			{
+				ViewModelName = *ViewModelContextId.ToString();
+			}
+		}
+	}
+
 private:
 	/** When the view is spawn, create an instance of the viewmodel. */
 	UPROPERTY(EditAnywhere, Category = "MVVM")
@@ -59,9 +73,13 @@ private:
 	UPROPERTY(EditAnywhere, Category = "MVVM")
 	TSubclassOf<UMVVMViewModelBase> ViewModelClass;
 
+	UPROPERTY()
+	FText OverrideDisplayName_DEPRECATED;
+
 public:
+	/** Property name that will be generated. */
 	UPROPERTY(EditAnywhere, Category = "MVVM")
-	FText OverrideDisplayName;
+	FName ViewModelName;
 
 	/** When the view is spawn, create an instance of the viewmodel. */
 	UPROPERTY(EditAnywhere, Category = "MVVM")
@@ -79,3 +97,12 @@ public:
 	UPROPERTY(EditAnywhere, Category = "MVVM", AdvancedDisplay)
 	bool bCreateSetterFunction = false;
 };
+
+template<>
+struct TStructOpsTypeTraits<FMVVMBlueprintViewModelContext> : public TStructOpsTypeTraitsBase2<FMVVMBlueprintViewModelContext>
+{
+	enum
+	{
+		WithPostSerialize = true,
+	};
+};
\ No newline at end of file
diff --git a/Engine/Plugins/Runtime/ModelViewViewModel/Source/ModelViewViewModelEditor/Private/Widgets/SMVVMViewModelContextListWidget.cpp b/Engine/Plugins/Runtime/ModelViewViewModel/Source/ModelViewViewModelEditor/Private/Widgets/SMVVMViewModelContextListWidget.cpp
--- a/Engine/Plugins/Runtime/ModelViewViewModel/Source/ModelViewViewModelEditor/Private/Widgets/SMVVMViewModelContextListWidget.cpp
+++ b/Engine/Plugins/Runtime/ModelViewViewModel/Source/ModelViewViewModelEditor/Private/Widgets/SMVVMViewModelContextListWidget.cpp
@@ -187,7 +187,7 @@ namespace UE::MVVM::Private
 						.OnTextCommitted_Lambda(
 							[TempEntry](const FText& NewText, ETextCommit::Type CommitType)
 							{
-								TempEntry->OverrideDisplayName = NewText;
+								TempEntry->ViewModelName = *NewText.ToString();
 							})
 						.OnVerifyTextChanged_Lambda(
 							[Self](const FText& InText, FText& OutErrorMessage) -> bool
@@ -744,7 +744,7 @@ void SMVVMViewModelContextListWidget::AddViewModelContext(TSubclassOf<UMVVMViewM
 		ExistingContextNames.Add(Context->GetDisplayName().ToString());
 	}
 
-	TSharedPtr<FMVVMBlueprintViewModelContext> NewContext = MakeShared<FMVVMBlueprintViewModelContext>(ViewModelClass, FGuid::NewGuid());
+	TSharedPtr<FMVVMBlueprintViewModelContext> NewContext = MakeShared<FMVVMBlueprintViewModelContext>(ViewModelClass, ViewModelClass->GetFName());
 	FString ClassName = ViewModelClass->ClassGeneratedBy != nullptr ? ViewModelClass->ClassGeneratedBy->GetName() : ViewModelClass->GetAuthoredName();
 	FString TempNewName = ClassName;
 	FKismetNameValidator NameValidator(WidgetBlueprint);
@@ -758,7 +758,7 @@ void SMVVMViewModelContextListWidget::AddViewModelContext(TSubclassOf<UMVVMViewM
 		++Index;
 	}
 
-	NewContext->OverrideDisplayName = FText::FromString(TempNewName);
+	NewContext->ViewModelName = *TempNewName;
 	ContextListSource.Add(NewContext);
 	ContextListWidget->RequestListRefresh();
 }
@@ -781,7 +781,7 @@ bool SMVVMViewModelContextListWidget::IsContextNameAvailable(FGuid Guid, FText C
 			continue;
 		}
 
-		if (Context->OverrideDisplayName.EqualTo(ContextName))
+		if (Context->ViewModelName == *ContextName.ToString())
 		{
 			return false;
 		}
