MVVM: To support property without field notification (that will be one way only), add the bHasNotify to the FMVVMAvailableBinding.

diff --git a/Engine/Plugins/Runtime/ModelViewViewModel/Source/ModelViewViewModel/Private/MVVMSubsystem.cpp b/Engine/Plugins/Runtime/ModelViewViewModel/Source/ModelViewViewModel/Private/MVVMSubsystem.cpp
--- a/Engine/Plugins/Runtime/ModelViewViewModel/Source/ModelViewViewModel/Private/MVVMSubsystem.cpp
+++ b/Engine/Plugins/Runtime/ModelViewViewModel/Source/ModelViewViewModel/Private/MVVMSubsystem.cpp
@@ -116,14 +116,17 @@ bool UMVVMSubsystem::DoesWidgetTreeContainedWidget(const UWidgetTree* WidgetTree
 
 namespace UE::MVVM::Private
 {
-	TArray<FMVVMAvailableBinding> GetAvailableBindings(const UClass* Class, const UE::FieldNotification::IClassDescriptor& ClassDescriptor)
+	TArray<FMVVMAvailableBinding> GetAvailableBindings(const UClass* Class, const UE::FieldNotification::IClassDescriptor* ClassDescriptor)
 	{
 		TSet<FName> FieldDescriptors;
-		ClassDescriptor.ForEachField(Class, [&FieldDescriptors](UE::FieldNotification::FFieldId FieldId)
+		if (ClassDescriptor)
+		{
+			ClassDescriptor->ForEachField(Class, [&FieldDescriptors](UE::FieldNotification::FFieldId FieldId)
 			{
 				FieldDescriptors.Add(FieldId.GetName());
 				return true;
 			});
+		}
 
 		TArray<FMVVMAvailableBinding> Result;
 		Result.Reserve(FieldDescriptors.Num());
@@ -132,11 +135,15 @@ namespace UE::MVVM::Private
 		{
 			const UFunction* Function = *FunctionIt;
 			check(Function);
-			const bool bIsReadable = FieldDescriptors.Contains(Function->GetFName()) && BindingHelper::IsValidForSourceBinding(Function);
-			const bool bIsWritable = BindingHelper::IsValidForDestinationBinding(Function);
-			if (bIsReadable || bIsWritable)
+			//if (!Function->HasAnyFunctionFlags(FUNC_Private | FUNC_Protected)) // todo fix for property and BlueprintPrivate in IsValidForSourceBindnig
 			{
-				Result.Add(FMVVMAvailableBinding(FMVVMBindingName(Function->GetFName()), bIsReadable, bIsWritable));
+				const bool bIsReadable = BindingHelper::IsValidForSourceBinding(Function);
+				const bool bIsWritable = BindingHelper::IsValidForDestinationBinding(Function);
+				const bool bHasNotify = FieldDescriptors.Contains(Function->GetFName()) && bIsReadable;
+				if (bIsReadable || bIsWritable || bHasNotify)
+				{
+					Result.Add(FMVVMAvailableBinding(FMVVMBindingName(Function->GetFName()), bIsReadable, bIsWritable, bHasNotify));
+				}
 			}
 		}
 
@@ -149,11 +156,13 @@ namespace UE::MVVM::Private
 		for (TFieldIterator<const FProperty> PropertyIt(Class, EFieldIteratorFlags::IncludeSuper); PropertyIt; ++PropertyIt)
 		{
 			const FProperty* Property = *PropertyIt;
-			const bool bIsReadable = FieldDescriptors.Contains(Property->GetFName()) && BindingHelper::IsValidForSourceBinding(Property);
+			const bool bIsReadable = BindingHelper::IsValidForSourceBinding(Property);
 			const bool bIsWritable = BindingHelper::IsValidForDestinationBinding(Property);
-			if (bIsReadable || bIsWritable)
+			const bool bHasNotify = FieldDescriptors.Contains(Property->GetFName()) && bIsReadable;
+			if (bIsReadable || bIsWritable || bHasNotify)
 			{
 #if WITH_EDITOR
+				// Remove any BlueprintGetter & BlueprintSetter and use the Property instead.
 				if (bIsReadable)
 				{
 					const FString& PropertyGetter = Property->GetMetaData(NAME_BlueprintGetter);
@@ -181,7 +190,7 @@ namespace UE::MVVM::Private
 					}
 				}
 #endif
-				Result.Add(FMVVMAvailableBinding(FMVVMBindingName(Property->GetFName()), bIsReadable, bIsWritable));
+				Result.Add(FMVVMAvailableBinding(FMVVMBindingName(Property->GetFName()), bIsReadable, bIsWritable, bHasNotify));
 			}
 		}
 
@@ -194,10 +203,12 @@ namespace UE::MVVM::Private
 		{
 			TScriptInterface<INotifyFieldValueChanged> DefaultObject = InSubClass.GetDefaultObject();
 			const UE::FieldNotification::IClassDescriptor& ClassDescriptor = DefaultObject->GetFieldNotificationDescriptor();
-			return GetAvailableBindings(InSubClass.Get(), ClassDescriptor);
+			return GetAvailableBindings(InSubClass.Get(), &ClassDescriptor);
+		}
+		else
+		{
+			return GetAvailableBindings(InSubClass.Get(), nullptr);
 		}
-
-		return TArray<FMVVMAvailableBinding>();
 	}
 } //namespace
 
diff --git a/Engine/Plugins/Runtime/ModelViewViewModel/Source/ModelViewViewModel/Public/Types/MVVMAvailableBinding.h b/Engine/Plugins/Runtime/ModelViewViewModel/Source/ModelViewViewModel/Public/Types/MVVMAvailableBinding.h
--- a/Engine/Plugins/Runtime/ModelViewViewModel/Source/ModelViewViewModel/Public/Types/MVVMAvailableBinding.h
+++ b/Engine/Plugins/Runtime/ModelViewViewModel/Source/ModelViewViewModel/Public/Types/MVVMAvailableBinding.h
@@ -17,10 +17,11 @@ struct FMVVMAvailableBinding
 public:
 	FMVVMAvailableBinding() = default;
 
-	explicit FMVVMAvailableBinding(FMVVMBindingName InBindingName, bool bReadable, bool bWritable)
+	explicit FMVVMAvailableBinding(FMVVMBindingName InBindingName, bool bInReadable, bool bInWritable, bool bInHasNotify)
 		: BindingName(InBindingName)
-		, bIsReadable(bReadable)
-		, bIsWritable(bWritable)
+		, bIsReadable(bInReadable)
+		, bIsWritable(bInWritable)
+		, bHasNotify(bInHasNotify)
 	{
 
 	}
@@ -31,6 +32,11 @@ public:
 		return BindingName;
 	}
 
+	bool HasNotify() const
+	{
+		return bHasNotify;
+	}
+
 	bool IsReadable() const
 	{
 		return bIsReadable;
@@ -48,12 +54,15 @@ public:
 
 	bool operator==(const FMVVMAvailableBinding& Other) const 
 	{
-		return BindingName == Other.BindingName && bIsReadable == Other.bIsReadable && bIsWritable == Other.bIsWritable;
+		return BindingName == Other.BindingName
+			&& bIsReadable == Other.bIsReadable
+			&& bIsWritable == Other.bIsWritable
+			&& bHasNotify == Other.bHasNotify;
 	}
 
 	bool operator!=(const FMVVMAvailableBinding& Other) const
 	{
-		return BindingName == Other.BindingName && bIsReadable == Other.bIsReadable && bIsWritable == Other.bIsWritable;
+		return !(this->operator==(Other));
 	}
 
 	friend int32 GetTypeHash(const FMVVMAvailableBinding& Value)
@@ -70,4 +79,7 @@ private:
 
 	UPROPERTY()
 	bool bIsWritable = false;
+
+	UPROPERTY()
+	bool bHasNotify = false;
 };
diff --git a/Engine/Plugins/Runtime/ModelViewViewModel/Source/ModelViewViewModelEditor/ModelViewViewModelEditor.Build.cs b/Engine/Plugins/Runtime/ModelViewViewModel/Source/ModelViewViewModelEditor/ModelViewViewModelEditor.Build.cs
--- a/Engine/Plugins/Runtime/ModelViewViewModel/Source/ModelViewViewModelEditor/ModelViewViewModelEditor.Build.cs
+++ b/Engine/Plugins/Runtime/ModelViewViewModel/Source/ModelViewViewModelEditor/ModelViewViewModelEditor.Build.cs
@@ -23,6 +23,7 @@ public class ModelViewViewModelEditor : ModuleRules
 		PrivateDependencyModuleNames.AddRange(
 			new string[]
 			{
+				"AdvancedWidgets",
 				"AssetTools",
 				"EditorWidgets",
 				"BlueprintGraph",
diff --git a/Engine/Plugins/Runtime/ModelViewViewModel/Source/ModelViewViewModelEditor/Private/MVVMEditorSubsystem.cpp b/Engine/Plugins/Runtime/ModelViewViewModel/Source/ModelViewViewModelEditor/Private/MVVMEditorSubsystem.cpp
--- a/Engine/Plugins/Runtime/ModelViewViewModel/Source/ModelViewViewModelEditor/Private/MVVMEditorSubsystem.cpp
+++ b/Engine/Plugins/Runtime/ModelViewViewModel/Source/ModelViewViewModelEditor/Private/MVVMEditorSubsystem.cpp
@@ -166,7 +166,7 @@ namespace UE::MVVM::Private
 			WidgetBlueprint->FunctionGraphs.Remove(Graph);
 		}
 	}
-}
+} //namespace
 
 FName UMVVMEditorSubsystem::GetConversionFunctionWrapperName(const UWidgetBlueprint* WidgetBlueprint, const FMVVMBlueprintViewBinding& Binding, bool bSourceToDestination) const
 {
@@ -749,7 +749,7 @@ void UMVVMEditorSubsystem::SetPathForConversionFunctionArgument(UWidgetBlueprint
 			++ArgumentIndex;
 		}
 	}
-	
+
 	if (!ensure(ArgumentPin != nullptr))
 	{
 		return;
@@ -1023,49 +1023,35 @@ TArray<UE::MVVM::FBindingSource> UMVVMEditorSubsystem::GetAllViewModels(const UW
 	return Sources;
 }
 
-
-TArray<UE::MVVM::FMVVMConstFieldVariant> UMVVMEditorSubsystem::GetChildViewModels(const UClass* Class)
+TArray<FMVVMAvailableBinding> UMVVMEditorSubsystem::GetChildViewModels(TSubclassOf<UObject> Class)
 {
-	auto IsValidObjectProperty = [](const FProperty* Property)
-	{
-		if (const FObjectPropertyBase* ObjectProperty = CastField<FObjectPropertyBase>(Property))
-		{
-			return ObjectProperty->PropertyClass->ImplementsInterface(UNotifyFieldValueChanged::StaticClass());
-		}
-		return false;
-	};
-
-
-	TArray<UE::MVVM::FMVVMConstFieldVariant> Result;
-	for (TFieldIterator<const FProperty> PropertyIt(Class, EFieldIteratorFlags::IncludeSuper); PropertyIt; ++PropertyIt)
+	if (Class.Get() == nullptr)
 	{
-		const FProperty* Property = *PropertyIt;
-		if (IsValidObjectProperty(Property) && UE::MVVM::BindingHelper::IsValidForSourceBinding(Property))
-		{
-			Result.Add(UE::MVVM::FMVVMConstFieldVariant(Property));
-		}
+		return TArray<FMVVMAvailableBinding>();
 	}
 
-	for (TFieldIterator<const UFunction> FunctionItt(Class, EFieldIteratorFlags::IncludeSuper); FunctionItt; ++FunctionItt)
-	{
-		const UFunction* Function = *FunctionItt;
-		check(Function);
-		if (!Function->HasAnyFunctionFlags(FUNC_Private | FUNC_Protected))
+	TArray<FMVVMAvailableBinding> ViewModelAvailableBindingsList = GEngine->GetEngineSubsystem<UMVVMSubsystem>()->GetAvailableBindings(Class);
+	ViewModelAvailableBindingsList.RemoveAllSwap([Class](const FMVVMAvailableBinding& Value)
 		{
-			if (UE::MVVM::BindingHelper::IsValidForSourceBinding(Function))
+			UE::MVVM::FMVVMFieldVariant Variant = UE::MVVM::BindingHelper::FindFieldByName(Class.Get(), Value.GetBindingName());
+			const FProperty* Property = nullptr;
+			if (Variant.IsProperty())
 			{
-				if (const FProperty* ReturnValue = UE::MVVM::BindingHelper::GetReturnProperty(Function))
-				{
-					if (IsValidObjectProperty(ReturnValue))
-					{
-						Result.Add(UE::MVVM::FMVVMConstFieldVariant(Function));
-					}
-				}
+				Property = Variant.GetProperty();
 			}
-		}
-	}
+			else if (Variant.IsFunction() && Variant.GetFunction())
+			{
+				Property = UE::MVVM::BindingHelper::GetReturnProperty(Variant.GetFunction());
+			}
+
+			if (const FObjectPropertyBase* ObjectProperty = CastField<FObjectPropertyBase>(Property))
+			{
+				return !ObjectProperty->PropertyClass->ImplementsInterface(UNotifyFieldValueChanged::StaticClass());
+			}
+			return true;
+		});
 
-	return Result;
+	return ViewModelAvailableBindingsList;
 }
 
 #undef LOCTEXT_NAMESPACE
diff --git a/Engine/Plugins/Runtime/ModelViewViewModel/Source/ModelViewViewModelEditor/Private/Widgets/SMVVMViewModelBindingListWidget.cpp b/Engine/Plugins/Runtime/ModelViewViewModel/Source/ModelViewViewModelEditor/Private/Widgets/SMVVMViewModelBindingListWidget.cpp
--- a/Engine/Plugins/Runtime/ModelViewViewModel/Source/ModelViewViewModelEditor/Private/Widgets/SMVVMViewModelBindingListWidget.cpp
+++ b/Engine/Plugins/Runtime/ModelViewViewModel/Source/ModelViewViewModelEditor/Private/Widgets/SMVVMViewModelBindingListWidget.cpp
@@ -5,558 +5,145 @@
 #include "Bindings/MVVMBindingHelper.h"
 #include "Editor.h"
 #include "Engine/Engine.h"
+#include "FieldNotification/IFieldValueChanged.h"
 #include "MVVMEditorSubsystem.h"
 #include "MVVMSubsystem.h"
 #include "MVVMViewModelBase.h"
-#include "Textures/SlateIcon.h"
 #include "Types/MVVMAvailableBinding.h"
 #include "Types/MVVMFieldVariant.h"
 
-#include "Framework/Commands/GenericCommands.h"
-#include "Framework/Commands/UICommandList.h"
-#include "Framework/MultiBox/MultiBoxBuilder.h"
-#include "Styling/SlateIconFinder.h"
-
-#include "Widgets/Images/SImage.h"
-#include "Widgets/SMVVMFieldIcon.h"
-#include "Widgets/Text/SInlineEditableTextBlock.h"
-#include "Widgets/Views/STreeView.h"
-
+#include "Widgets/PropertyViewer/SPropertyViewer.h"
 
 #define LOCTEXT_NAMESPACE "SViewModelBindingListWidget"
 
 namespace UE::MVVM
 {
 
-namespace Private
-{
-
-static TArray<FMVVMAvailableBinding> GetBindingsList(TSubclassOf<UObject> Class)
+/** */
+TArray<FFieldVariant> FFieldIterator_ViewModel::GetFields(const UStruct* Struct) const
 {
-	TArray<FMVVMAvailableBinding> ViewModelAvailableBindingsList = GEngine->GetEngineSubsystem<UMVVMSubsystem>()->GetAvailableBindings(Class);
-	const TArray<FMVVMConstFieldVariant> ChildViewModels = GEditor->GetEditorSubsystem<UMVVMEditorSubsystem>()->GetChildViewModels(Class);
-	for (const FMVVMConstFieldVariant& Field : ChildViewModels)
+	TArray<FFieldVariant> Result;
+	if (const UClass* Class = Cast<const UClass>(Struct))
 	{
-		if (!ViewModelAvailableBindingsList.ContainsByPredicate([&Field](const FMVVMAvailableBinding& Other){ return Other.GetBindingName().ToName() == Field.GetName(); }))
+		TArray<FMVVMAvailableBinding> ViewModelAvailableBindingsList = GEngine->GetEngineSubsystem<UMVVMSubsystem>()->GetAvailableBindings(const_cast<UClass*>(Class));
+		Result.Reserve(ViewModelAvailableBindingsList.Num());
+		for (const FMVVMAvailableBinding& Value : ViewModelAvailableBindingsList)
 		{
-			ViewModelAvailableBindingsList.Add(FMVVMAvailableBinding(FMVVMBindingName(Field.GetName()), true, false));
-		}
-	}
-
-	return ViewModelAvailableBindingsList;
-}
-
-
-/** 
- * 
- */
-class FAvailableBindingNode : public TSharedFromThis<FAvailableBindingNode>
-{
-public:
-	FAvailableBindingNode(SViewModelBindingListWidget* InOwner, const FMVVMAvailableBinding& InBinding, TSubclassOf<UObject> InViewModelClass)
-		: Owner(InOwner)
-		, bIsRoot(false)
-		, Binding(InBinding)
-		, ViewModelOwnerClass(InViewModelClass)
-		, ViewModelClassProperty(GetCacheViewModelClassProperty())
-	{
-	}
-
-
-	FAvailableBindingNode(SViewModelBindingListWidget* InOwner, const SViewModelBindingListWidget::FViewModel& InViewModel)
-		: Owner(InOwner)
-		, bIsRoot(true)
-		, ViewModel(InViewModel)
-	{
-		check(InViewModel.Class && InViewModel.Class->ImplementsInterface(UNotifyFieldValueChanged::StaticClass()));
-	}
-
-private:
-	SViewModelBindingListWidget* Owner = nullptr;
-
-	const bool bIsRoot = false;
-	const SViewModelBindingListWidget::FViewModel ViewModel;
-
-	const FMVVMAvailableBinding Binding;
-	const TSubclassOf<UObject> ViewModelOwnerClass;
-	const TSubclassOf<UObject> ViewModelClassProperty; // if the property is also a viewmodel
-
-	bool bChildGenerated = false;
-	TArray<TSharedPtr<FAvailableBindingNode>, TInlineAllocator<1>> ChildNodes;
-	TWeakPtr<SInlineEditableTextBlock> ViewModelEditTextBlock;
-
-public:
-	bool ShowContextMenu() const
-	{
-		return bIsRoot;
-	}
-
-	bool CanRename() const
-	{
-		return bIsRoot && ViewModelEditTextBlock.IsValid();
-	}
-
-
-	void EnterRename() const
-	{
-		if (bIsRoot)
-		{
-			if (TSharedPtr<SInlineEditableTextBlock> SafeEditBox = ViewModelEditTextBlock.Pin())
+			FMVVMFieldVariant FieldVariant = BindingHelper::FindFieldByName(Struct, Value.GetBindingName());
+			if (FieldVariant.IsFunction())
 			{
-				SafeEditBox->EnterEditingMode();
+				Result.Add(FFieldVariant(FieldVariant.GetFunction()));
 			}
-		}
-	}
-
-
-	TArray<TSharedPtr<FAvailableBindingNode>> GetChildNodes()
-	{
-		if (!bChildGenerated && ViewModelClassProperty.Get() != nullptr)
-		{
-			// Do not build when filtering (only search in what it's already been built)
-			if (!Owner->FilterHandler->GetIsEnabled())
+			else if (FieldVariant.IsProperty())
 			{
-				BuildChildNodesRecursive(ViewModelClassProperty, 2);
+				Result.Add(FFieldVariant(FieldVariant.GetProperty()));
 			}
-		}
-		return TArray<TSharedPtr<FAvailableBindingNode>>(ChildNodes);
-	}
-
-
-	void BuildChildNodes()
-	{
-		BuildChildNodesRecursive(ViewModel.Class, 3);
+		}		
 	}
-
-
-	FText GetDisplayName()
+	else if (Cast<const UScriptStruct>(Struct))
 	{
-		if (bIsRoot)
+		for (TFieldIterator<FProperty> PropertyIt(Struct, EFieldIteratorFlags::IncludeSuper); PropertyIt; ++PropertyIt)
 		{
-			return !ViewModel.Name.IsNone() ? FText::FromName(ViewModel.Name) : ViewModel.Class->GetDisplayNameText();
+			FProperty* Property = *PropertyIt;
+			if (BindingHelper::IsValidForSourceBinding(Property) || BindingHelper::IsValidForDestinationBinding(Property))
+			{
+				Result.Add(FFieldVariant(Property));
+			}
 		}
-
-		return FText::FromName(Binding.GetBindingName().ToName());
 	}
 
-
-	void GetFilterStrings(TArray<FString>& OutStrings) const
-	{
-		auto AddClass = [&OutStrings](UClass* Class)
-		{
-			OutStrings.Add(Class->GetName());
-			OutStrings.Add(Class->GetDisplayNameText().ToString());
-		};
-
-		if (bIsRoot)
+	Result.Sort([](const FFieldVariant& A, const FFieldVariant& B)
 		{
-			if (!ViewModel.Name.IsNone())
+			bool bIsAViewModel = A.Get<FObjectPropertyBase>() && A.Get<FObjectPropertyBase>()->PropertyClass->ImplementsInterface(UNotifyFieldValueChanged::StaticClass());
+			bool bIsBViewModel = B.Get<FObjectPropertyBase>() && B.Get<FObjectPropertyBase>()->PropertyClass->ImplementsInterface(UNotifyFieldValueChanged::StaticClass());
+			if (A.IsUObject() && B.IsUObject())
 			{
-				OutStrings.Add(ViewModel.Name.ToString());
+				return A.GetFName().LexicalLess(B.GetFName());
 			}
-			AddClass(ViewModel.Class);
-		}
-		else
-		{
-			OutStrings.Add(Binding.GetBindingName().ToString());
-			if (ViewModelOwnerClass.Get())
+			else if (bIsAViewModel && bIsBViewModel)
 			{
-				AddClass(ViewModelOwnerClass.Get());
+				return A.GetFName().LexicalLess(B.GetFName());
 			}
-			if (ViewModelClassProperty.Get())
+			else if (bIsAViewModel)
 			{
-				AddClass(ViewModelClassProperty.Get());
+				return true;
 			}
-		}
-	}
-
-
-	FText GetToolTip()
-	{
-		if (bIsRoot)
-		{
-			return ViewModel.Class->GetDisplayNameText();
-		}
-
-		FMVVMFieldVariant FieldVariant = BindingHelper::FindFieldByName(ViewModelOwnerClass, Binding.GetBindingName());
-		if (FieldVariant.IsProperty())
-		{
-			FProperty* Property = FieldVariant.GetProperty();
-			check(Property);
-			return FText::FormatOrdered(INVTEXT("{0}\n{1}"), FText::FromString(Property->GetCPPType()), Property->GetMetaDataText("ToolTip"));
-		}
-
-		if (FieldVariant.IsFunction())
-		{
-			UFunction* Function = FieldVariant.GetFunction();
-			check(Function);
-			FProperty* Property = Function->GetReturnProperty();
-			if (!Property)
+			else if (bIsBViewModel)
 			{
-				return FText::FormatOrdered(INVTEXT("void\n{0}"), Function->GetMetaDataText("ToolTip"));
+				return false;
 			}
-			return FText::FormatOrdered(INVTEXT("{0}\n{1}"), FText::FromString(Property->GetCPPType()), Function->GetMetaDataText("ToolTip"));
-		}
-
-		return FText::GetEmpty();
-	}
-
-
-	TSharedRef<SWidget> CreateIconWidget()
-	{
-		if (bIsRoot)
-		{
-			return SNew(SImage).Image(FSlateIconFinder::FindIconBrushForClass(ViewModel.Class));
-		}
-
-		FMVVMFieldVariant FieldVariant = BindingHelper::FindFieldByName(ViewModelOwnerClass, Binding.GetBindingName());
-			
-		return SNew(SMVVMFieldIcon)
-			.Field(FieldVariant);
-	}
-
-
-	TSharedRef<SWidget> CreateDetailWidget(TAttribute<FText> GetFilterText, const SViewModelBindingListWidget::FOnViewModelRenamed& OnRenamed)
-	{
-		ViewModelEditTextBlock.Reset();
-		if (OnRenamed.IsBound() && bIsRoot)
-		{
-			return SAssignNew(ViewModelEditTextBlock, SInlineEditableTextBlock)
-				.Text(GetDisplayName())
-				.HighlightText(GetFilterText)
-				.OnVerifyTextChanged(this, &FAvailableBindingNode::HandleVerifyNameTextChanged, OnRenamed)
-				.OnTextCommitted(this, &FAvailableBindingNode::HandleNameTextCommited, OnRenamed);
-		}
-
-		return SNew(STextBlock)
-			.Text(GetDisplayName())
-			.HighlightText(GetFilterText);
-	}
-	
-private:
-	bool HandleVerifyNameTextChanged(const FText& InText, FText& OutErrorMessage, SViewModelBindingListWidget::FOnViewModelRenamed OnRenamed)
-	{
-		return OnRenamed.Execute(ViewModel, InText, false, OutErrorMessage);
-	}
-	
-	
-	void HandleNameTextCommited(const FText& InText, ETextCommit::Type CommitInfo, SViewModelBindingListWidget::FOnViewModelRenamed OnRenamed)
-	{
-		if (CommitInfo == ETextCommit::OnEnter)
-		{
-			FText OutErrorMessage;
-			OnRenamed.Execute(ViewModel, InText, true, OutErrorMessage);
-		}
-	}
-
-
-	static bool Sort(const TSharedPtr<FAvailableBindingNode>& Node, const TSharedPtr<FAvailableBindingNode>& OtherNode)
-	{
-		TSubclassOf<UObject> NodeClass = Node->ViewModelClassProperty;
-		TSubclassOf<UObject> OtherNodeClass = OtherNode->ViewModelClassProperty;
-		const FName NodeStr = Node->Binding.GetBindingName().ToName();
-		const FName OtherNodeStr = OtherNode->Binding.GetBindingName().ToName();
-
-		if (NodeClass && OtherNodeClass)
-		{
-			return NodeStr.LexicalLess(OtherNodeStr);
-		}
-
-		if (NodeClass)
-		{
-			return true;
-		}
-
-		if (OtherNodeClass)
-		{
-			return false;
-		}
-
-		return NodeStr.LexicalLess(OtherNodeStr);
-	}
-
-private:
-	TSubclassOf<UObject> GetCacheViewModelClassProperty() const
-	{
-		if (bIsRoot)
-		{
-			return nullptr;
-		}
-
-		FMVVMFieldVariant FieldVariant = BindingHelper::FindFieldByName(ViewModelOwnerClass, Binding.GetBindingName());
-		if (FieldVariant.IsProperty())
-		{
-			FProperty* Property = FieldVariant.GetProperty();
-
-			if (FObjectProperty* ObjectProperty = CastField<FObjectProperty>(Property))
+			else if(A.IsUObject())
 			{
-				if (ObjectProperty->PropertyClass->ImplementsInterface(UNotifyFieldValueChanged::StaticClass()))
-				{
-					return ObjectProperty->PropertyClass;
-				}
+				return true;
 			}
-		}
-
-		return nullptr;
-	}
-
-
-	void BuildChildNodesRecursive(UClass* ViewModelClass, int32 RecursiveCount)
-	{
-		if (RecursiveCount <= 0)
-		{
-			return;
-		}
-		--RecursiveCount;
-
-		ChildNodes.Reset();
-		for (const FMVVMAvailableBinding& ChildBinding : Private::GetBindingsList(ViewModelClass))
-		{
-			TSharedRef<Private::FAvailableBindingNode> ChildNode = MakeShared<Private::FAvailableBindingNode>(Owner, ChildBinding, ViewModelClass);
-			ChildNodes.Add(ChildNode);
-			if (ChildNode->ViewModelClassProperty.Get() != nullptr)
+			else if (B.IsUObject())
 			{
-				ChildNode->BuildChildNodesRecursive(ChildNode->ViewModelClassProperty, RecursiveCount);
+				return false;
 			}
-		}
-		ChildNodes.Sort(Private::FAvailableBindingNode::Sort);
-		bChildGenerated = true;
-	}
-};
+			return A.GetFName().LexicalLess(B.GetFName());
+		});
+
 
-} //namespace Private
+	return Result;
+}
 
 
+/** */
 void SViewModelBindingListWidget::Construct(const FArguments& InArgs)
 {
-	OnRenamed = InArgs._OnRenamed;
-	OnDeleted = InArgs._OnDeleted;
-	OnDuplicated = InArgs._OnDuplicated;
+	PropertyViewer = SNew(UE::PropertyViewer::SPropertyViewer)
+		.PropertyVisibility(UE::PropertyViewer::SPropertyViewer::EPropertyVisibility::Hidden)
+		.bShowFieldIcon(true)
+		.bSanitizeName(true)
+		.FieldIterator(&ViewModelFieldIterator);
 
-	SearchFilter = MakeShared<ViewModelTextFilter>(ViewModelTextFilter::FItemToStringArray::CreateSP(this, &SViewModelBindingListWidget::HandleGetFilterStrings));
-
-	FilterHandler = MakeShared<ViewModelTreeFilter>();
-	FilterHandler->SetFilter(SearchFilter.Get());
-	FilterHandler->SetRootItems(&TreeSource, &FilteredTreeSource);
-	FilterHandler->SetGetChildrenDelegate(ViewModelTreeFilter::FOnGetChildren::CreateRaw(this, &SViewModelBindingListWidget::HandleGetChildren));
-
-	SetViewModels(InArgs._ViewModels);
-	CreateCommandList();
-
-	SAssignNew(TreeWidget, STreeView<TSharedPtr<Private::FAvailableBindingNode>>)
-		.ItemHeight(1.0f)
-		.TreeItemsSource(&FilteredTreeSource)
-		.SelectionMode(ESelectionMode::Single)
-		.OnGetChildren(FilterHandler.ToSharedRef(), &ViewModelTreeFilter::OnGetFilteredChildren)
-		.OnGenerateRow(this, &SViewModelBindingListWidget::HandleGenerateRow)
-		//	.OnSelectionChanged(this, &SMVVMViewModelPanel::HandleSelectionChanged)
-		.OnContextMenuOpening(this, &SViewModelBindingListWidget::HandleContextMenuOpening);
-
-	FilterHandler->SetTreeView(TreeWidget.Get());
+	if (InArgs._ViewModel.Class.Get())
+	{
+		SetViewModels(MakeArrayView(&InArgs._ViewModel, 1));
+	}
 
 	ChildSlot
 	[
-		TreeWidget.ToSharedRef()
+		PropertyViewer.ToSharedRef()
 	];
 }
 
 
-void SViewModelBindingListWidget::SetViewModel(TSubclassOf<UObject> ViewModelClass, FName InViewModelName)
+void SViewModelBindingListWidget::SetViewModel(UClass* Class, FName Name, FGuid Guid)
 {
 	FViewModel ViewModel;
-	ViewModel.Class = ViewModelClass;
-	ViewModel.Name = InViewModelName;
-
+	ViewModel.Class = Class;
+	ViewModel.ViewModelId = Guid;
+	ViewModel.Name = Name;
 	SetViewModels(MakeArrayView(&ViewModel, 1));
 }
 
 
-void SViewModelBindingListWidget::SetViewModels(TArrayView<const FViewModel> ViewModels)
-{
-	TreeSource.Reset();
-	for (const FViewModel& ViewModel : ViewModels)
-	{
-		if (ViewModel.Class && ViewModel.Class->ImplementsInterface(UNotifyFieldValueChanged::StaticClass()))
-		{
-			TSharedRef<Private::FAvailableBindingNode> RootNode = MakeShared<Private::FAvailableBindingNode>(this, ViewModel);
-			RootNode->BuildChildNodes();
-			TreeSource.Add(RootNode);
-
-			if (TreeWidget)
-			{
-				TreeWidget->SetItemExpansion(RootNode, true);
-			}
-		}
-	}
-
-	if (FilterHandler)
-	{
-		FilterHandler->RefreshAndFilterTree();
-	}
-}
-
-
-FText SViewModelBindingListWidget::SetRawFilterText(const FText& InFilterText)
-{
-	const bool bNewFilteringEnabled = !InFilterText.IsEmpty();
-	FilterHandler->SetIsEnabled(bNewFilteringEnabled);
-	SearchFilter->SetRawFilterText(InFilterText);
-	FilterHandler->RefreshAndFilterTree();
-	return SearchFilter->GetFilterErrorText();
-}
-
-
-FText SViewModelBindingListWidget::GetRawFilterText() const
-{
-	return SearchFilter->GetRawFilterText();
-}
-
-FReply SViewModelBindingListWidget::OnKeyDown(const FGeometry& MyGeometry, const FKeyEvent& InKeyEvent)
+void SViewModelBindingListWidget::SetViewModels(TArrayView<const FViewModel> InViewModels)
 {
-	if (CommandList->ProcessCommandBindings(InKeyEvent))
+	ViewModels = InViewModels;
+	if (PropertyViewer)
 	{
-		return FReply::Handled();
-	}
-
-	return FReply::Unhandled();
-}
-
-
-void SViewModelBindingListWidget::HandleGetFilterStrings(TSharedPtr<Private::FAvailableBindingNode> Item, TArray<FString>& OutStrings)
-{
-	Item->GetFilterStrings(OutStrings);
-}
-
-
-TSharedRef<ITableRow> SViewModelBindingListWidget::HandleGenerateRow(TSharedPtr<Private::FAvailableBindingNode> Item, const TSharedRef<STableViewBase>& OwnerTable)
-{
-	using Private::FAvailableBindingNode;
-	typedef STableRow<TSharedPtr<FAvailableBindingNode>> RowType;
-
-
-	return SNew(RowType, OwnerTable)
-		//.OnDragDetected(this, &SHierarchyViewItem::HandleDragDetected)
-		//.OnDragEnter(this, &SHierarchyViewItem::HandleDragEnter)
-		//.OnDragLeave(this, &SHierarchyViewItem::HandleDragLeave)
-		.Padding(0.0f)
-		.Content()
-		[
-			SNew(SHorizontalBox)
-			.ToolTipText(Item->GetToolTip())
-			+ SHorizontalBox::Slot()
-			.AutoWidth()
-			.HAlign(HAlign_Right)
-			.VAlign(VAlign_Center)
-			[
-				Item->CreateIconWidget()
-			]
-			+ SHorizontalBox::Slot()
-			.Padding(4.0f)
-			[
-				Item->CreateDetailWidget(MakeAttributeSP(this, &SViewModelBindingListWidget::GetRawFilterText), OnRenamed)
-			]
-		];
-}
-
-
-void SViewModelBindingListWidget::HandleGetChildren(TSharedPtr<Private::FAvailableBindingNode> InParent, TArray<TSharedPtr<Private::FAvailableBindingNode>>& OutChildren)
-{
-	OutChildren = InParent->GetChildNodes();
-}
-
-
-void SViewModelBindingListWidget::CreateCommandList()
-{
-	CommandList = MakeShared<FUICommandList>();
-
-	CommandList->MapAction(
-		FGenericCommands::Get().Duplicate,
-		FExecuteAction::CreateSP(this, &SViewModelBindingListWidget::HandleDuplicateViewModel),
-		FCanExecuteAction::CreateSP(this, &SViewModelBindingListWidget::HandleCanDuplicateViewModel)
-	);
-	CommandList->MapAction(
-		FGenericCommands::Get().Delete,
-		FExecuteAction::CreateSP(this, &SViewModelBindingListWidget::HandleDeleteViewModel),
-		FCanExecuteAction::CreateSP(this, &SViewModelBindingListWidget::HandleCanDeleteViewModel)
-	);
-	CommandList->MapAction(
-		FGenericCommands::Get().Rename,
-		FExecuteAction::CreateSP(this, &SViewModelBindingListWidget::HandleRenameViewModel),
-		FCanExecuteAction::CreateSP(this, &SViewModelBindingListWidget::HandleCanRenameViewModel)
-	);
-}
-
-
-TSharedPtr<SWidget> SViewModelBindingListWidget::HandleContextMenuOpening()
-{
-	TArray<TSharedPtr<Private::FAvailableBindingNode>> Items = TreeWidget->GetSelectedItems();
-	if (Items.Num() == 1 && Items[0] && Items[0]->ShowContextMenu())
-	{
-		if (OnRenamed.IsBound() || OnDeleted.IsBound() || OnDuplicated.IsBound())
+		PropertyViewer->RemoveAll();
+		for (const FViewModel& ViewModel : ViewModels)
 		{
-			FMenuBuilder MenuBuilder(true, CommandList);
-
-			MenuBuilder.BeginSection("Edit", LOCTEXT("Edit", "Edit"));
+			if (ViewModel.Class && ViewModel.Class->ImplementsInterface(UNotifyFieldValueChanged::StaticClass()))
 			{
-				if (OnDuplicated.IsBound())
-				{
-					MenuBuilder.AddMenuEntry(FGenericCommands::Get().Duplicate);
-				}
-				if (OnDeleted.IsBound())
-				{
-					MenuBuilder.AddMenuEntry(FGenericCommands::Get().Delete);
-				}
-				if (OnRenamed.IsBound())
-				{
-					MenuBuilder.AddMenuEntry(FGenericCommands::Get().Rename);
-				}
+				PropertyViewer->AddContainer(ViewModel.Class.Get());
 			}
-			MenuBuilder.EndSection();
-
-			return MenuBuilder.MakeWidget();
 		}
 	}
-	return TSharedPtr<SWidget>();
-}
-
-
-void SViewModelBindingListWidget::HandleDuplicateViewModel()
-{}
-
-
-bool SViewModelBindingListWidget::HandleCanDuplicateViewModel()
-{
-	return true; //todo unless it's an autogen viewmodel
-}
-
-
-void SViewModelBindingListWidget::HandleDeleteViewModel()
-{}
-
-
-bool SViewModelBindingListWidget::HandleCanDeleteViewModel()
-{
-	return true; //todo unless it's an autogen viewmodel
 }
 
 
-void SViewModelBindingListWidget::HandleRenameViewModel()
+void SViewModelBindingListWidget::SetRawFilterText(const FText& InFilterText)
 {
-	TArray<TSharedPtr<Private::FAvailableBindingNode>> Items = TreeWidget->GetSelectedItems();
-	if (Items.Num() == 1 && Items[0])
+	if (PropertyViewer)
 	{
-		Items[0]->EnterRename();
+		//PropertyViewer->SetRawFilterText(InFilterText);
 	}
 }
 
-
-bool SViewModelBindingListWidget::HandleCanRenameViewModel()
-{
-	TArray<TSharedPtr<Private::FAvailableBindingNode>> Items = TreeWidget->GetSelectedItems();
-	return Items.Num() == 1 && Items[0] && Items[0]->CanRename();
-}
-
-
 } // namespace UE::MVVM
 
 #undef LOCTEXT_NAMESPACE
\ No newline at end of file
diff --git a/Engine/Plugins/Runtime/ModelViewViewModel/Source/ModelViewViewModelEditor/Private/Widgets/SMVVMViewModelBindingListWidget.h b/Engine/Plugins/Runtime/ModelViewViewModel/Source/ModelViewViewModelEditor/Private/Widgets/SMVVMViewModelBindingListWidget.h
--- a/Engine/Plugins/Runtime/ModelViewViewModel/Source/ModelViewViewModelEditor/Private/Widgets/SMVVMViewModelBindingListWidget.h
+++ b/Engine/Plugins/Runtime/ModelViewViewModel/Source/ModelViewViewModelEditor/Private/Widgets/SMVVMViewModelBindingListWidget.h
@@ -3,35 +3,32 @@
 #pragma once
 
 #include "CoreMinimal.h"
-#include "Framework/SlateDelegates.h"
-#include "Misc/IFilter.h"
-#include "Misc/TextFilter.h"
-#include "Framework/Views/TreeFilterHandler.h"
+#include "Framework/PropertyViewer/IFieldIterator.h"
 #include "Templates/SubclassOf.h"
-#include "Types/MVVMFieldVariant.h"
 #include "Widgets/SCompoundWidget.h"
 
-struct FMVVMAvailableBinding;
-class FUICommandList;
-class ITableRow;
-template<typename T>
-class STreeView;
-class STableViewBase;
-class SInlineEditableTextBlock;
-class UMVVMViewModelBase;
-
+namespace UE::PropertyViewer
+{
+class SPropertyViewer;
+}
 
 namespace UE::MVVM
 {
-namespace Private
+
+/**
+ * 
+ */
+class FFieldIterator_ViewModel : public UE::PropertyViewer::FFieldIterator_BlueprintVisible
 {
-class FAvailableBindingNode;
-}//namespace Private
+	virtual TArray<FFieldVariant> GetFields(const UStruct*) const override;
+};
 
 
+/** 
+ * 
+ */
 class SViewModelBindingListWidget : public SCompoundWidget
 {
-	friend Private::FAvailableBindingNode;
 public:
 	struct FViewModel
 	{
@@ -39,55 +36,20 @@ public:
 		FGuid ViewModelId;
 		FName Name;
 	};
-
-	DECLARE_DELEGATE_RetVal_FourParams(bool, FOnViewModelRenamed, const FViewModel&, const FText&, bool, FText&);
-	DECLARE_DELEGATE_OneParam(FOnViewModelDeleted, const FViewModel&);
-	DECLARE_DELEGATE_RetVal_OneParam(FViewModel, FOnViewModelDuplicated, const FViewModel&);
-
 	SLATE_BEGIN_ARGS(SViewModelBindingListWidget) {}
-		SLATE_ARGUMENT(TArrayView<const FViewModel>, ViewModels)
-		SLATE_EVENT(FOnViewModelRenamed, OnRenamed)
-		SLATE_EVENT(FOnViewModelDeleted, OnDeleted)
-		SLATE_EVENT(FOnViewModelDuplicated, OnDuplicated)
+		SLATE_ARGUMENT(FViewModel, ViewModel)
 	SLATE_END_ARGS()
 
 	void Construct(const FArguments& InArgs);
-	void SetViewModel(TSubclassOf<UObject> ViewModelClass, FName ViewModelName);
+	void SetViewModel(UClass* Class, FName Name, FGuid Guid);
 	void SetViewModels(TArrayView<const FViewModel> ViewModels);
-	FText SetRawFilterText(const FText& InFilterText);
-
-private:
-	virtual FReply OnKeyDown(const FGeometry& MyGeometry, const FKeyEvent& InKeyEvent) override;
-
-	FText GetRawFilterText() const;
-
-	void HandleGetFilterStrings(TSharedPtr<Private::FAvailableBindingNode> Item, TArray<FString>& OutStrings);
-	TSharedRef<ITableRow> HandleGenerateRow(TSharedPtr<Private::FAvailableBindingNode> Item, const TSharedRef<STableViewBase>& OwnerTable);
-	void HandleGetChildren(TSharedPtr<Private::FAvailableBindingNode> InParent, TArray<TSharedPtr<Private::FAvailableBindingNode>>& OutChildren);
-
-	TSharedPtr<SWidget> HandleContextMenuOpening();
-	void HandleDuplicateViewModel();
-	bool HandleCanDuplicateViewModel();
-	void HandleDeleteViewModel();
-	bool HandleCanDeleteViewModel();
-	void HandleRenameViewModel();
-	bool HandleCanRenameViewModel();
 
-	void CreateCommandList();
+	void SetRawFilterText(const FText& InFilterText);
 
-	TSharedPtr<STreeView<TSharedPtr<Private::FAvailableBindingNode>>> TreeWidget;
-	TArray<TSharedPtr<Private::FAvailableBindingNode>> TreeSource;
-	TArray<TSharedPtr<Private::FAvailableBindingNode>> FilteredTreeSource;
-
-	TSharedPtr<FUICommandList> CommandList;
-	typedef TTextFilter<TSharedPtr<Private::FAvailableBindingNode>> ViewModelTextFilter;
-	TSharedPtr<ViewModelTextFilter> SearchFilter;
-	typedef TreeFilterHandler<TSharedPtr<Private::FAvailableBindingNode>> ViewModelTreeFilter;
-	TSharedPtr<ViewModelTreeFilter> FilterHandler;
-
-	FOnViewModelRenamed OnRenamed;
-	FOnViewModelDeleted OnDeleted;
-	FOnViewModelDuplicated OnDuplicated;
+private:
+	FFieldIterator_ViewModel ViewModelFieldIterator;
+	TArray<FViewModel> ViewModels;
+	TSharedPtr<UE::PropertyViewer::SPropertyViewer> PropertyViewer;
 };
 
 } //namespace UE::MVVM
\ No newline at end of file
diff --git a/Engine/Plugins/Runtime/ModelViewViewModel/Source/ModelViewViewModelEditor/Private/Widgets/SMVVMViewModelPanel.cpp b/Engine/Plugins/Runtime/ModelViewViewModel/Source/ModelViewViewModelEditor/Private/Widgets/SMVVMViewModelPanel.cpp
--- a/Engine/Plugins/Runtime/ModelViewViewModel/Source/ModelViewViewModelEditor/Private/Widgets/SMVVMViewModelPanel.cpp
+++ b/Engine/Plugins/Runtime/ModelViewViewModel/Source/ModelViewViewModelEditor/Private/Widgets/SMVVMViewModelPanel.cpp
@@ -2,12 +2,14 @@
 
 #include "Widgets/SMVVMViewModelPanel.h"
 
+#include "Bindings/MVVMBindingHelper.h"
 #include "Editor.h"
 #include "Editor/EditorEngine.h"
 #include "Kismet2/Kismet2NameValidators.h"
 #include "MVVMBlueprintView.h"
 #include "MVVMBlueprintViewModelContext.h"
 #include "MVVMEditorSubsystem.h"
+#include "MVVMSubsystem.h"
 #include "MVVMViewModelBase.h"
 #include "WidgetBlueprint.h"
 #include "WidgetBlueprintEditor.h"
@@ -19,8 +21,11 @@
 #include "Widgets/SBoxPanel.h"
 #include "Widgets/Input/SSearchBox.h"
 #include "Widgets/Layout/SBorder.h"
+#include "Widgets/PropertyViewer/SPropertyViewer.h"
 #include "Widgets/SMVVMViewModelBindingListWidget.h"
 
+#include "SPositiveActionButton.h"
+
 #define LOCTEXT_NAMESPACE "ViewModelPanel"
 
 namespace UE::MVVM
@@ -35,43 +40,37 @@ void SMVVMViewModelPanel::Construct(const FArguments& InArgs, TSharedPtr<FWidget
 	WeakBlueprintEditor = WidgetBlueprintEditor;
 	WeakBlueprintView = CurrentBlueprintView;
 
-	WidgetBlueprint->OnExtensionAdded.AddSP(this, &SMVVMViewModelPanel::HandleViewUpdated);
-	WidgetBlueprint->OnExtensionRemoved.AddSP(this, &SMVVMViewModelPanel::HandleViewUpdated);
-
 	if (CurrentBlueprintView)
 	{
 		// Listen to when the viewmodel are modified
-		CurrentBlueprintView->OnViewModelsUpdated.AddSP(this, &SMVVMViewModelPanel::HandleViewModelsUpdated);
+		ViewModelsUpdatedHandle = CurrentBlueprintView->OnViewModelsUpdated.AddSP(this, &SMVVMViewModelPanel::HandleViewModelsUpdated);
 	}
 
-	GEditor->OnBlueprintCompiled().AddSP(this, &SMVVMViewModelPanel::HandleBlueprintCompiled);
-
-	ViewModelTreeView = SNew(SViewModelBindingListWidget)
-		.OnRenamed(this, &SMVVMViewModelPanel::HandleViewModelRenamed);
+	ViewModelTreeView = SNew(UE::PropertyViewer::SPropertyViewer)
+		.PropertyVisibility(UE::PropertyViewer::SPropertyViewer::EPropertyVisibility::Visible)
+		.bShowSearchBox(true)
+		.bShowFieldIcon(true)
+		.bSanitizeName(true)
+		.FieldIterator(&ViewModelFieldIterator)
+		.SearchBoxPreSlot()
+		[
+			SNew(SPositiveActionButton)
+			.OnGetMenuContent(this, &SMVVMViewModelPanel::MakeAddMenu)
+			.Text(LOCTEXT("Viewmodel", "Viewmodel"))
+			.IsEnabled(this, &SMVVMViewModelPanel::HandleCanEditViewmodelList)
+		]
+		;
+		//.OnRenamed(this, &SMVVMViewModelPanel::HandleViewModelRenamed);
 
-	GenerateViewModelTreeView();
+	FillViewModel();
 
 	ChildSlot
 	[
-		SNew(SVerticalBox)
-
-		+ SVerticalBox::Slot()
-		.Padding(4)
-		.AutoHeight()
-		[
-			SAssignNew(SearchBoxPtr, SSearchBox)
-			.HintText(LOCTEXT("SearchViewmodel", "Search"))
-			.OnTextChanged(this, &SMVVMViewModelPanel::HandleSearchChanged)
-		]
-		+ SVerticalBox::Slot()
-		.FillHeight(1.0f)
+		SNew(SBorder)
+		.BorderImage(FAppStyle::Get().GetBrush("Brushes.Recessed"))
+		.Padding(0)
 		[
-			SNew(SBorder)
-			.BorderImage(FAppStyle::Get().GetBrush("Brushes.Recessed"))
-			.Padding(0)
-			[
-				ViewModelTreeView.ToSharedRef()
-			]
+			ViewModelTreeView.ToSharedRef()
 		]
 	];
 }
@@ -79,11 +78,6 @@ void SMVVMViewModelPanel::Construct(const FArguments& InArgs, TSharedPtr<FWidget
 
 SMVVMViewModelPanel::~SMVVMViewModelPanel()
 {
-	if (GEditor)
-	{
-		GEditor->OnBlueprintCompiled().RemoveAll(this);
-	}
-
 	if (TSharedPtr<FWidgetBlueprintEditor> WidgetBlueprintEditor = WeakBlueprintEditor.Pin())
 	{
 		if (UWidgetBlueprint* WidgetBlueprint = WidgetBlueprintEditor->GetWidgetBlueprintObj())
@@ -96,29 +90,30 @@ SMVVMViewModelPanel::~SMVVMViewModelPanel()
 	if (UMVVMBlueprintView* CurrentBlueprintView = WeakBlueprintView.Get())
 	{
 		// bind to check if the view is enabled
-		CurrentBlueprintView->OnViewModelsUpdated.RemoveAll(this);
+		CurrentBlueprintView->OnViewModelsUpdated.Remove(ViewModelsUpdatedHandle);
 	}
 }
 
 
-void SMVVMViewModelPanel::GenerateViewModelTreeView()
+void SMVVMViewModelPanel::FillViewModel()
 {
 	if (ViewModelTreeView)
 	{
-		TArray<SViewModelBindingListWidget::FViewModel, TInlineAllocator<16>> ViewModelList;
+		ViewModelTreeView->RemoveAll();
+		ViewModelHandles.Reset();
+
 		if (UMVVMBlueprintView* View = WeakBlueprintView.Get())
 		{
 			for (const FMVVMBlueprintViewModelContext& ViewModelContext : View->GetViewModels())
 			{
-				SViewModelBindingListWidget::FViewModel ViewModel;
-				ViewModel.Class = ViewModelContext.GetViewModelClass();
-				ViewModel.ViewModelId = ViewModelContext.GetViewModelId();
-				ViewModel.Name = ViewModelContext.GetViewModelName();
-				ViewModelList.Add(ViewModel);
+				UObject* ViewModelInstance = ViewModelContext.GetViewModelClass().GetDefaultObject();
+				
+				FViewModelHandle ViewModelHandle;
+				ViewModelHandle.Key = ViewModelContext.GetViewModelId();
+				ViewModelHandle.Value = ViewModelTreeView->AddInstance(ViewModelInstance);
+				ViewModelHandles.Add(ViewModelHandle);
 			}
 		}
-
-		ViewModelTreeView->SetViewModels(ViewModelList);
 	}
 }
 
@@ -126,82 +121,105 @@ void SMVVMViewModelPanel::GenerateViewModelTreeView()
 void SMVVMViewModelPanel::HandleViewUpdated(UBlueprintExtension*)
 {
 	bool bViewUpdated = false;
-	if (TSharedPtr<FWidgetBlueprintEditor> WidgetBlueprintEditor = WeakBlueprintEditor.Pin())
+
+	if (!ViewModelsUpdatedHandle.IsValid())
 	{
-		if (UWidgetBlueprint* WidgetBlueprint = WidgetBlueprintEditor->GetWidgetBlueprintObj())
+		if (TSharedPtr<FWidgetBlueprintEditor> WidgetBlueprintEditor = WeakBlueprintEditor.Pin())
 		{
-			UMVVMBlueprintView* CurrentBlueprintView = GEditor->GetEditorSubsystem<UMVVMEditorSubsystem>()->GetView(WidgetBlueprint);
-			WeakBlueprintView = CurrentBlueprintView;
-
-			if (CurrentBlueprintView)
+			if (UWidgetBlueprint* WidgetBlueprint = WidgetBlueprintEditor->GetWidgetBlueprintObj())
 			{
-				CurrentBlueprintView->OnViewModelsUpdated.AddSP(this, &SMVVMViewModelPanel::HandleViewModelsUpdated);
+				UMVVMBlueprintView* CurrentBlueprintView = GEditor->GetEditorSubsystem<UMVVMEditorSubsystem>()->GetView(WidgetBlueprint);
+				WeakBlueprintView = CurrentBlueprintView;
+
+				if (CurrentBlueprintView)
+				{
+					ViewModelsUpdatedHandle = CurrentBlueprintView->OnViewModelsUpdated.AddSP(this, &SMVVMViewModelPanel::HandleViewModelsUpdated);
+					bViewUpdated = true;
+				}
 			}
 		}
 	}
-	GenerateViewModelTreeView();
-}
 
-
-void SMVVMViewModelPanel::HandleViewModelsUpdated()
-{
-	GenerateViewModelTreeView();
+	if (bViewUpdated)
+	{
+		FillViewModel();
+	}
 }
 
 
-void SMVVMViewModelPanel::HandleSearchChanged(const FText& InFilterText)
+void SMVVMViewModelPanel::HandleViewModelsUpdated()
 {
-	SearchBoxPtr->SetError(ViewModelTreeView->SetRawFilterText(InFilterText));
+	FillViewModel();
 }
 
 
-bool SMVVMViewModelPanel::HandleViewModelRenamed(const SViewModelBindingListWidget::FViewModel& ViewModel, const FText& RenameTo, bool bCommit, FText& OutErrorMessage)
+TSharedRef<SWidget> SMVVMViewModelPanel::MakeAddMenu()
 {
-	if (RenameTo.IsEmptyOrWhitespace())
-	{
-		OutErrorMessage = LOCTEXT("EmptyWidgetName", "Empty Widget Name");
-		return false;
-	}
+	//using namespace UE::Sequencer;
 
-	const FString& NewNameString = RenameTo.ToString();
-	if (NewNameString.Len() >= NAME_SIZE)
-	{
-		OutErrorMessage = LOCTEXT("WidgetNameTooLong", "Widget Name is Too Long");
-		return false;
-	}
+	//TSharedPtr<FExtender> Extender = FExtender::Combine(AddMenuExtenders);
+	//FMenuBuilder MenuBuilder(true, nullptr, Extender);
+	//{
+	//	TSharedPtr<FSequencer> Sequencer = SequencerPtr.Pin();
+	//	Sequencer->GetViewModel()->GetOutliner()->CastThisChecked<FSequencerOutlinerViewModel>()->BuildContextMenu(MenuBuilder);
+	//}
 
-	FString GeneratedName = SlugStringForValidName(NewNameString);
-	if (GeneratedName.IsEmpty())
-	{
-		OutErrorMessage = LOCTEXT("EmptyWidgetName", "Empty Widget Name");
-		return false;
-	}
+	//return MenuBuilder.MakeWidget();
+	return SNullWidget::NullWidget;
+}
 
-	const FName GeneratedFName(*GeneratedName);
-	check(GeneratedFName.IsValidXName(INVALID_OBJECTNAME_CHARACTERS));
 
+bool SMVVMViewModelPanel::HandleCanEditViewmodelList() const
+{
 	if (TSharedPtr<FWidgetBlueprintEditor> WidgetBlueprintEditor = WeakBlueprintEditor.Pin())
 	{
-		if (UWidgetBlueprint* WidgetBP = WidgetBlueprintEditor->GetWidgetBlueprintObj())
-		{
-			if (bCommit)
-			{
-				return GEditor->GetEditorSubsystem<UMVVMEditorSubsystem>()->RenameViewModel(WidgetBP, ViewModel.Name, *NewNameString, OutErrorMessage);
-			}
-			else
-			{
-				return GEditor->GetEditorSubsystem<UMVVMEditorSubsystem>()->VerifyViewModelRename(WidgetBP, ViewModel.Name, *NewNameString, OutErrorMessage);
-			}
-		}
+		return WidgetBlueprintEditor->InEditingMode();
 	}
 	return false;
 }
 
 
-void SMVVMViewModelPanel::HandleBlueprintCompiled()
-{
-	GenerateViewModelTreeView();
-}
+//bool SMVVMViewModelPanel::HandleViewModelRenamed(const SViewModelBindingListWidget::FViewModel& ViewModel, const FText& RenameTo, bool bCommit, FText& OutErrorMessage)
+//{
+//	if (RenameTo.IsEmptyOrWhitespace())
+//	{
+//		OutErrorMessage = LOCTEXT("EmptyViewmodelName", "Empty viewmodel name");
+//		return false;
+//	}
+//
+//	const FString& NewNameString = RenameTo.ToString();
+//	if (NewNameString.Len() >= NAME_SIZE)
+//	{
+//		OutErrorMessage = LOCTEXT("ViewmodelNameTooLong", "Viewmodel name is too long");
+//		return false;
+//	}
+//
+//	FString GeneratedName = SlugStringForValidName(NewNameString);
+//	if (GeneratedName.IsEmpty())
+//	{
+//		OutErrorMessage = LOCTEXT("EmptyViewmodelName", "Empty viewmodel name");
+//		return false;
+//	}
+//
+//	const FName GeneratedFName(*GeneratedName);
+//	check(GeneratedFName.IsValidXName(INVALID_OBJECTNAME_CHARACTERS));
+//
+//	if (TSharedPtr<FWidgetBlueprintEditor> WidgetBlueprintEditor = WeakBlueprintEditor.Pin())
+//	{
+//		if (UWidgetBlueprint* WidgetBP = WidgetBlueprintEditor->GetWidgetBlueprintObj())
+//		{
+//			if (bCommit)
+//			{
+//				return GEditor->GetEditorSubsystem<UMVVMEditorSubsystem>()->RenameViewModel(WidgetBP, ViewModel.Name, *NewNameString, OutErrorMessage);
+//			}
+//			else
+//			{
+//				return GEditor->GetEditorSubsystem<UMVVMEditorSubsystem>()->VerifyViewModelRename(WidgetBP, ViewModel.Name, *NewNameString, OutErrorMessage);
+//			}
+//		}
+//	}
+//	return false;
+//}
 
 } // namespace
 
diff --git a/Engine/Plugins/Runtime/ModelViewViewModel/Source/ModelViewViewModelEditor/Private/Widgets/SMVVMViewModelPanel.h b/Engine/Plugins/Runtime/ModelViewViewModel/Source/ModelViewViewModelEditor/Private/Widgets/SMVVMViewModelPanel.h
--- a/Engine/Plugins/Runtime/ModelViewViewModel/Source/ModelViewViewModelEditor/Private/Widgets/SMVVMViewModelPanel.h
+++ b/Engine/Plugins/Runtime/ModelViewViewModel/Source/ModelViewViewModelEditor/Private/Widgets/SMVVMViewModelPanel.h
@@ -4,14 +4,11 @@
 
 #include "CoreMinimal.h"
 
+#include "Widgets/PropertyViewer/SPropertyViewer.h"
 #include "Widgets/SMVVMViewModelBindingListWidget.h"
 #include "Widgets/SCompoundWidget.h"
 
 class FWidgetBlueprintEditor;
-class SSearchBox;
-class STableViewBase;
-template<typename ItemType>
-class SListView;
 class UBlueprintExtension;
 class UMVVMBlueprintView;
 
@@ -31,20 +28,24 @@ public:
 	virtual ~SMVVMViewModelPanel();
 
 private:
-	void HandleViewUpdated(UBlueprintExtension*);
+	void HandleViewUpdated(UBlueprintExtension* Extension);
 	void HandleViewModelsUpdated();
-	void HandleSearchChanged(const FText& InFilterText);
-	bool HandleViewModelRenamed(const SViewModelBindingListWidget::FViewModel& ViewModel, const FText& RenameTo, bool bCommit, FText& OutErrorMessage);
-	void HandleBlueprintCompiled();
+	//bool HandleViewModelRenamed(const SViewModelBindingListWidget::FViewModel& ViewModel, const FText& RenameTo, bool bCommit, FText& OutErrorMessage);
+	TSharedRef<SWidget> MakeAddMenu();
+	bool HandleCanEditViewmodelList() const;
 
-	void GenerateViewModelTreeView();
+	void FillViewModel();
 
 private:
-	TSharedPtr<SSearchBox> SearchBoxPtr;
-	TSharedPtr<SViewModelBindingListWidget> ViewModelTreeView;
+	TSharedPtr<UE::PropertyViewer::SPropertyViewer> ViewModelTreeView;
+	FFieldIterator_ViewModel ViewModelFieldIterator;
+
+	using FViewModelHandle = TPair<FGuid, UE::PropertyViewer::SPropertyViewer::FHandle>;
+	TArray<FViewModelHandle> ViewModelHandles;
 
 	TWeakPtr<FWidgetBlueprintEditor> WeakBlueprintEditor;
 	TWeakObjectPtr<UMVVMBlueprintView> WeakBlueprintView;
+	FDelegateHandle ViewModelsUpdatedHandle;
 };
 
 } // namespace
diff --git a/Engine/Plugins/Runtime/ModelViewViewModel/Source/ModelViewViewModelEditor/Public/MVVMEditorSubsystem.h b/Engine/Plugins/Runtime/ModelViewViewModel/Source/ModelViewViewModelEditor/Public/MVVMEditorSubsystem.h
--- a/Engine/Plugins/Runtime/ModelViewViewModel/Source/ModelViewViewModelEditor/Public/MVVMEditorSubsystem.h
+++ b/Engine/Plugins/Runtime/ModelViewViewModel/Source/ModelViewViewModelEditor/Public/MVVMEditorSubsystem.h
@@ -44,6 +44,9 @@ public:
 	UFUNCTION(BlueprintCallable, Category = "MVVM")
 	void RemoveBinding(UWidgetBlueprint* WidgetBlueprint, const FMVVMBlueprintViewBinding& Binding);
 
+	UFUNCTION(BlueprintCallable, Category = "MVVM")
+	TArray<FMVVMAvailableBinding> GetChildViewModels(TSubclassOf<UObject> Class);
+
 	void SetSourceToDestinationConversionFunction(UWidgetBlueprint* WidgetBlueprint, FMVVMBlueprintViewBinding& Binding, const UFunction* ConversionFunction);
 	void SetDestinationToSourceConversionFunction(UWidgetBlueprint* WidgetBlueprint, FMVVMBlueprintViewBinding& Binding, const UFunction* ConversionFunction);
 	void SetWidgetForBinding(UWidgetBlueprint* WidgetBlueprint, FMVVMBlueprintViewBinding& Binding, FName Widget);
@@ -71,8 +74,6 @@ public:
 	
 	TArray<UE::MVVM::FBindingSource> GetAllViewModels(const UWidgetBlueprint* WidgetBlueprint) const;
 
-	TArray<UE::MVVM::FMVVMConstFieldVariant> GetChildViewModels(const UClass* Class);
-
 private:
 	FName GetConversionFunctionWrapperName(const UWidgetBlueprint* WidgetBlueprint, const FMVVMBlueprintViewBinding& Binding, bool bSourceToDestination) const;
 	UEdGraph* CreateConversionFunctionWrapperGraph(UWidgetBlueprint* WidgetBlueprint, const FMVVMBlueprintViewBinding& Binding, const UFunction* ConversionFunction, bool bSourceToDestination);
