Missing file with CL 20857739
#rnx

diff --git a/Engine/Plugins/Runtime/ModelViewViewModel/Source/ModelViewViewModelEditor/Private/Widgets/SMVVMManageViewModelsWidget.cpp b/Engine/Plugins/Runtime/ModelViewViewModel/Source/ModelViewViewModelEditor/Private/Widgets/SMVVMManageViewModelsWidget.cpp
--- a/Engine/Plugins/Runtime/ModelViewViewModel/Source/ModelViewViewModelEditor/Private/Widgets/SMVVMManageViewModelsWidget.cpp
+++ b/Engine/Plugins/Runtime/ModelViewViewModel/Source/ModelViewViewModelEditor/Private/Widgets/SMVVMManageViewModelsWidget.cpp
@@ -769,7 +769,7 @@ void SMVVMManageViewModelsWidget::HandleTreeViewSelectionChanged(TSharedPtr<FMVV
 	{
 		InClass = GetClassFromNode(Item.ToSharedRef());
 	}
-	ViewModelBindingListWidget->SetViewModel(InClass, FName());
+	ViewModelBindingListWidget->SetViewModel(InClass, FName(), FGuid());
 }
 
 void SMVVMManageViewModelsWidget::HandleTreeViewExpansionChanged(TSharedPtr<FMVVMViewModelTreeNode> Item, bool bExpanded)
