Fixed static analysis warning about potentially null UBlueprintView*.

#rnx

diff --git a/Engine/Plugins/Runtime/ModelViewViewModel/Source/ModelViewViewModelBlueprint/Private/MVVMBlueprintViewBinding.cpp b/Engine/Plugins/Runtime/ModelViewViewModel/Source/ModelViewViewModelBlueprint/Private/MVVMBlueprintViewBinding.cpp
--- a/Engine/Plugins/Runtime/ModelViewViewModel/Source/ModelViewViewModelBlueprint/Private/MVVMBlueprintViewBinding.cpp
+++ b/Engine/Plugins/Runtime/ModelViewViewModel/Source/ModelViewViewModelBlueprint/Private/MVVMBlueprintViewBinding.cpp
@@ -40,16 +40,13 @@ FName FMVVMBlueprintViewBinding::GetFName(const UMVVMBlueprintView* View) const
 		BindingName << TEXT("_From_");
 	}
 
-	if (View->GetOuterUMVVMWidgetBlueprintExtension_View()->GetWidgetBlueprint()->GetFName() != WidgetPath.GetWidgetName())
-	{
-		if (WidgetPath.GetWidgetName().IsNone())
-		{
-			BindingName << TEXT("Invalid");
-		}
-		else
-		{
-			BindingName << WidgetPath.GetWidgetName();
-		}
+	if (View == nullptr || WidgetPath.GetWidgetName().IsNone())
+	{
+		BindingName << TEXT("Invalid_");
+	}
+	else if (View->GetOuterUMVVMWidgetBlueprintExtension_View()->GetWidgetBlueprint()->GetFName() != WidgetPath.GetWidgetName())
+	{
+		BindingName << WidgetPath.GetWidgetName();
 		BindingName << TEXT("_");
 	}
 
@@ -96,23 +93,17 @@ FString FMVVMBlueprintViewBinding::GetDisplayNameString(const UMVVMBlueprintView
 		BindingName << TEXT(" ??? "); // shouldn't happen
 	}
 
-	if (View->GetOuterUMVVMWidgetBlueprintExtension_View()->GetWidgetBlueprint()->GetFName() == WidgetPath.GetWidgetName())
+	if (View == nullptr || WidgetPath.GetWidgetName().IsNone())
 	{
-		BindingName << WidgetPath.GetBasePropertyPath();
+		BindingName << TEXT("<none>.");
 	}
-	else
+	else if (View->GetOuterUMVVMWidgetBlueprintExtension_View()->GetWidgetBlueprint()->GetFName() != WidgetPath.GetWidgetName())
 	{
-		if (WidgetPath.GetWidgetName().IsNone())
-		{
-			BindingName << TEXT("<none>");
-		}
-		else
-		{
-			BindingName << WidgetPath.GetWidgetName();
-		}
+		BindingName << WidgetPath.GetWidgetName();
 		BindingName << TEXT(".");
-		BindingName << WidgetPath.GetBasePropertyPath();
 	}
-
+	
+	BindingName << WidgetPath.GetBasePropertyPath();
+	
 	return BindingName.ToString();
 }
