Fixed static analysis warning about flipped null check.

#rnx

diff --git a/Engine/Plugins/Runtime/ModelViewViewModel/Source/ModelViewViewModelEditor/Private/MVVMEditorSubsystem.cpp b/Engine/Plugins/Runtime/ModelViewViewModel/Source/ModelViewViewModelEditor/Private/MVVMEditorSubsystem.cpp
--- a/Engine/Plugins/Runtime/ModelViewViewModel/Source/ModelViewViewModelEditor/Private/MVVMEditorSubsystem.cpp
+++ b/Engine/Plugins/Runtime/ModelViewViewModel/Source/ModelViewViewModelEditor/Private/MVVMEditorSubsystem.cpp
@@ -959,7 +959,7 @@ TArray<UE::MVVM::FBindingSource> UMVVMEditorSubsystem::GetBindableWidgets(const
 	TArray<UE::MVVM::FBindingSource> Sources;
 
 	const UWidgetTree* WidgetTree = WidgetBlueprint->WidgetTree;
-	if (WidgetTree != nullptr)
+	if (WidgetTree == nullptr)
 	{
 		return Sources;
 	}
