MVVM: Use the viewmodel name instead of the class name. Create the menu to add a viewmodel. Use a UClass of any class at implement the NotifyFieldValue interface in BlueprintViewModelContext instead a ViewModelBase class.

diff --git a/Engine/Plugins/Runtime/ModelViewViewModel/Source/ModelViewViewModelBlueprint/Private/MVVMBlueprintViewModelContext.cpp b/Engine/Plugins/Runtime/ModelViewViewModel/Source/ModelViewViewModelBlueprint/Private/MVVMBlueprintViewModelContext.cpp
--- a/Engine/Plugins/Runtime/ModelViewViewModel/Source/ModelViewViewModelBlueprint/Private/MVVMBlueprintViewModelContext.cpp
+++ b/Engine/Plugins/Runtime/ModelViewViewModel/Source/ModelViewViewModelBlueprint/Private/MVVMBlueprintViewModelContext.cpp
@@ -2,12 +2,16 @@
 
 #include "MVVMBlueprintViewModelContext.h"
 
-FMVVMBlueprintViewModelContext::FMVVMBlueprintViewModelContext(TSubclassOf<UMVVMViewModelBase> InClass, FName InViewModelName)
-	: ViewModelContextId(FGuid::NewGuid())
-	, ViewModelClass(InClass)
-	, ViewModelName(InViewModelName)
-{
+#include "FieldNotification/IFieldValueChanged.h"
 
+FMVVMBlueprintViewModelContext::FMVVMBlueprintViewModelContext(const UClass* InClass, FName InViewModelName)
+{
+	if (InClass && InClass->ImplementsInterface(UNotifyFieldValueChanged::StaticClass()))
+	{
+		ViewModelContextId = FGuid::NewGuid();
+		NotifyFieldValueClass = const_cast<UClass*>(InClass);
+		ViewModelName = InViewModelName;
+	}
 }
 
 
diff --git a/Engine/Plugins/Runtime/ModelViewViewModel/Source/ModelViewViewModelBlueprint/Private/MVVMViewBlueprintCompiler.cpp b/Engine/Plugins/Runtime/ModelViewViewModel/Source/ModelViewViewModelBlueprint/Private/MVVMViewBlueprintCompiler.cpp
--- a/Engine/Plugins/Runtime/ModelViewViewModel/Source/ModelViewViewModelBlueprint/Private/MVVMViewBlueprintCompiler.cpp
+++ b/Engine/Plugins/Runtime/ModelViewViewModel/Source/ModelViewViewModelBlueprint/Private/MVVMViewBlueprintCompiler.cpp
@@ -270,7 +270,7 @@ void FMVVMViewBlueprintCompiler::CreateSourceLists(const FWidgetBlueprintCompile
 
 		ViewModelsGuid.Add(ViewModelContext.GetViewModelId());
 
-		if (ViewModelContext.GetViewModelClass() == nullptr)
+		if (ViewModelContext.GetViewModelClass() == nullptr || !ViewModelContext.IsValid())
 		{
 			WidgetBlueprintCompilerContext.MessageLog.Error(*FString::Printf(TEXT("The viewmodel '%s' has an invalid class.")
 				, *ViewModelContext.GetViewModelId().ToString()));
@@ -303,7 +303,7 @@ void FMVVMViewBlueprintCompiler::CreateSourceLists(const FWidgetBlueprintCompile
 		checkf(FoundSourceContextIndex != INDEX_NONE, TEXT("The viewmodel was added after the skeleton was created?"));
 
 		FCompilerSourceContext SourceVariable;
-		SourceVariable.Class = ViewModelContext.GetViewModelClass().Get();
+		SourceVariable.Class = ViewModelContext.GetViewModelClass();
 		SourceVariable.PropertyName = ViewModelContext.GetViewModelName();
 		SourceVariable.DisplayName = ViewModelContext.GetDisplayName().ToString();
 		SourceVariable.CategoryName = TEXT("Viewmodel");
@@ -405,7 +405,7 @@ void FMVVMViewBlueprintCompiler::CreateFunctionsDeclaration(const FWidgetBluepri
 				WidgetBlueprintCompilerContext.MessageLog.Error(*FString::Printf(TEXT("The setter name %s already exist and could not be autogenerated."), *SourceCreator.SetterFunctionName));
 			}
 
-			UE::MVVM::FunctionGraphHelper::AddFunctionArgument(SourceCreator.SetterGraph, SourceCreator.ViewModelContext.GetViewModelClass().Get(), "Viewmodel");
+			UE::MVVM::FunctionGraphHelper::AddFunctionArgument(SourceCreator.SetterGraph, SourceCreator.ViewModelContext.GetViewModelClass(), "Viewmodel");
 		}
 	}
 }
diff --git a/Engine/Plugins/Runtime/ModelViewViewModel/Source/ModelViewViewModelBlueprint/Private/MVVMWidgetBlueprintExtension_View.cpp b/Engine/Plugins/Runtime/ModelViewViewModel/Source/ModelViewViewModelBlueprint/Private/MVVMWidgetBlueprintExtension_View.cpp
--- a/Engine/Plugins/Runtime/ModelViewViewModel/Source/ModelViewViewModelBlueprint/Private/MVVMWidgetBlueprintExtension_View.cpp
+++ b/Engine/Plugins/Runtime/ModelViewViewModel/Source/ModelViewViewModelBlueprint/Private/MVVMWidgetBlueprintExtension_View.cpp
@@ -48,7 +48,7 @@ void UMVVMWidgetBlueprintExtension_View::HandlePreloadObjectsForCompilation(UBlu
 		{
 			if (AvailableViewModel.GetViewModelClass())
 			{
-				UBlueprint::ForceLoad(AvailableViewModel.GetViewModelClass().Get());
+				UBlueprint::ForceLoad(AvailableViewModel.GetViewModelClass());
 			}
 			//UBlueprint::ForceLoad(ViewModelPropertyPath);
 		}
diff --git a/Engine/Plugins/Runtime/ModelViewViewModel/Source/ModelViewViewModelBlueprint/Public/MVVMBlueprintViewModelContext.h b/Engine/Plugins/Runtime/ModelViewViewModel/Source/ModelViewViewModelBlueprint/Public/MVVMBlueprintViewModelContext.h
--- a/Engine/Plugins/Runtime/ModelViewViewModel/Source/ModelViewViewModelBlueprint/Public/MVVMBlueprintViewModelContext.h
+++ b/Engine/Plugins/Runtime/ModelViewViewModel/Source/ModelViewViewModelBlueprint/Public/MVVMBlueprintViewModelContext.h
@@ -31,7 +31,7 @@ struct MODELVIEWVIEWMODELBLUEPRINT_API FMVVMBlueprintViewModelContext
 
 public:
 	FMVVMBlueprintViewModelContext() = default;
-	FMVVMBlueprintViewModelContext(TSubclassOf<UMVVMViewModelBase> InClass, FName InViewModelName);
+	FMVVMBlueprintViewModelContext(const UClass* InClass, FName InViewModelName);
 
 	FGuid GetViewModelId() const
 	{
@@ -45,9 +45,9 @@ public:
 
 	FText GetDisplayName() const;
 
-	TSubclassOf<UMVVMViewModelBase> GetViewModelClass() const
+	UClass* GetViewModelClass() const
 	{
-		return ViewModelClass;
+		return NotifyFieldValueClass;
 	}
 
 	void PostSerialize(const FArchive& Ar)
@@ -62,16 +62,28 @@ public:
 			{
 				ViewModelName = *ViewModelContextId.ToString();
 			}
+			if (ViewModelClass_DEPRECATED.Get())
+			{
+				NotifyFieldValueClass = ViewModelClass_DEPRECATED.Get();
+			}
 		}
 	}
 
+	bool IsValid() const
+	{
+		return NotifyFieldValueClass != nullptr;
+	}
+
 private:
 	/** When the view is spawn, create an instance of the viewmodel. */
-	UPROPERTY(EditAnywhere, Category = "MVVM")
+	UPROPERTY(VisibleAnywhere, Category = "MVVM")
 	FGuid ViewModelContextId;
 
 	UPROPERTY(EditAnywhere, Category = "MVVM")
-	TSubclassOf<UMVVMViewModelBase> ViewModelClass;
+	UClass* NotifyFieldValueClass = nullptr;
+
+	UPROPERTY()
+	TSubclassOf<UMVVMViewModelBase> ViewModelClass_DEPRECATED;
 
 	UPROPERTY()
 	FText OverrideDisplayName_DEPRECATED;
diff --git a/Engine/Plugins/Runtime/ModelViewViewModel/Source/ModelViewViewModelEditor/Private/Customizations/MVVMConversionPathCustomization.cpp b/Engine/Plugins/Runtime/ModelViewViewModel/Source/ModelViewViewModelEditor/Private/Customizations/MVVMConversionPathCustomization.cpp
--- a/Engine/Plugins/Runtime/ModelViewViewModel/Source/ModelViewViewModelEditor/Private/Customizations/MVVMConversionPathCustomization.cpp
+++ b/Engine/Plugins/Runtime/ModelViewViewModel/Source/ModelViewViewModelEditor/Private/Customizations/MVVMConversionPathCustomization.cpp
@@ -474,7 +474,7 @@ namespace UE::MVVM
 		}
 
 		UMVVMSubsystem* Subsystem = GEngine->GetEngineSubsystem<UMVVMSubsystem>();
-		TArray<FMVVMAvailableBinding> AvailableBindings = Source.ViewModelId.IsValid() ? Subsystem->GetViewModelAvailableBindings(SourceClass) : Subsystem->GetWidgetAvailableBindings(SourceClass);
+		TArray<FMVVMAvailableBinding> AvailableBindings = Subsystem->GetAvailableBindings(SourceClass);
 
 		const UEdGraphSchema_K2* Schema = GetDefault<UEdGraphSchema_K2>();
 		FEdGraphPinType ArgumentType = GetArgumentPinType(ArgumentName, bSourceToDestination);
diff --git a/Engine/Plugins/Runtime/ModelViewViewModel/Source/ModelViewViewModelEditor/Private/Customizations/MVVMPropertyBindingExtension.cpp b/Engine/Plugins/Runtime/ModelViewViewModel/Source/ModelViewViewModelEditor/Private/Customizations/MVVMPropertyBindingExtension.cpp
--- a/Engine/Plugins/Runtime/ModelViewViewModel/Source/ModelViewViewModelEditor/Private/Customizations/MVVMPropertyBindingExtension.cpp
+++ b/Engine/Plugins/Runtime/ModelViewViewModel/Source/ModelViewViewModelEditor/Private/Customizations/MVVMPropertyBindingExtension.cpp
@@ -100,7 +100,7 @@ static void ExtendBindingsMenu(FMenuBuilder& MenuBuilder, const UWidgetBlueprint
 		MenuBuilder.AddSubMenu(ViewModel.GetDisplayName(), ViewModel.GetDisplayName(),
 			FNewMenuDelegate::CreateLambda([ViewModel, Widget, Property, Schema, CreatePropertyWidget, CreateBinding](FMenuBuilder& MenuBuilder)
 				{
-					const UClass* ViewModelClass = ViewModel.GetViewModelClass().Get();
+					const UClass* ViewModelClass = ViewModel.GetViewModelClass();
 
 					MenuBuilder.BeginSection("ValidDataTypes", LOCTEXT("ValidDataTypes", "Valid Data Types"));
 
diff --git a/Engine/Plugins/Runtime/ModelViewViewModel/Source/ModelViewViewModelEditor/Private/MVVMEditorSubsystem.cpp b/Engine/Plugins/Runtime/ModelViewViewModel/Source/ModelViewViewModelEditor/Private/MVVMEditorSubsystem.cpp
--- a/Engine/Plugins/Runtime/ModelViewViewModel/Source/ModelViewViewModelEditor/Private/MVVMEditorSubsystem.cpp
+++ b/Engine/Plugins/Runtime/ModelViewViewModel/Source/ModelViewViewModelEditor/Private/MVVMEditorSubsystem.cpp
@@ -198,6 +198,43 @@ UMVVMBlueprintView* UMVVMEditorSubsystem::GetView(const UWidgetBlueprint* Widget
 	return nullptr;
 }
 
+FName UMVVMEditorSubsystem::AddViewModel(UWidgetBlueprint* WidgetBlueprint, const UClass* ViewModelClass)
+{
+	FName Result;
+	if (ViewModelClass)
+	{
+		if (UMVVMBlueprintView* View = GetView(WidgetBlueprint))
+		{
+			FString ClassName = ViewModelClass->ClassGeneratedBy != nullptr ? ViewModelClass->ClassGeneratedBy->GetName() : ViewModelClass->GetAuthoredName();
+			FString ViewModelName = ClassName;
+			FKismetNameValidator NameValidator(WidgetBlueprint);
+
+			int32 Index = 1;
+			while (NameValidator.IsValid(ViewModelName) != EValidatorResult::Ok)
+			{
+				ViewModelName = ClassName + "_";
+				ViewModelName.AppendInt(Index);
+
+				++Index;
+			}
+
+			Result = *ViewModelName;
+			FMVVMBlueprintViewModelContext Context = FMVVMBlueprintViewModelContext(ViewModelClass, Result);
+			if (Context.IsValid())
+			{
+				const FScopedTransaction Transaction(LOCTEXT("AddViewModel", "Add viewmodel"));
+				View->Modify();
+				View->AddViewModel(Context);
+			}
+			else
+			{
+				Result = FName();
+			}
+		}
+	}
+	return Result;
+}
+
 void UMVVMEditorSubsystem::RemoveViewModel(UWidgetBlueprint* WidgetBlueprint, FName ViewModel)
 {
 	if (UMVVMBlueprintView* View = GetView(WidgetBlueprint))
@@ -1015,7 +1052,7 @@ TArray<UE::MVVM::FBindingSource> UMVVMEditorSubsystem::GetAllViewModels(const UW
 			UE::MVVM::FBindingSource Source;
 			Source.ViewModelId = ViewModel.GetViewModelId();
 			Source.DisplayName = ViewModel.GetDisplayName();
-			Source.Class = ViewModel.GetViewModelClass().Get();
+			Source.Class = ViewModel.GetViewModelClass();
 			Sources.Add(Source);
 		}
 	}
diff --git a/Engine/Plugins/Runtime/ModelViewViewModel/Source/ModelViewViewModelEditor/Private/Types/MVVMBindingSource.cpp b/Engine/Plugins/Runtime/ModelViewViewModel/Source/ModelViewViewModelEditor/Private/Types/MVVMBindingSource.cpp
--- a/Engine/Plugins/Runtime/ModelViewViewModel/Source/ModelViewViewModelEditor/Private/Types/MVVMBindingSource.cpp
+++ b/Engine/Plugins/Runtime/ModelViewViewModel/Source/ModelViewViewModelEditor/Private/Types/MVVMBindingSource.cpp
@@ -38,7 +38,7 @@ UE::MVVM::FBindingSource UE::MVVM::FBindingSource::CreateForViewModel(const UWid
 		if (const FMVVMBlueprintViewModelContext* ViewModel = View->FindViewModel(Source.ViewModelId))
 		{
 			Source.DisplayName = ViewModel->GetDisplayName();
-			Source.Class = ViewModel->GetViewModelClass().Get();
+			Source.Class = ViewModel->GetViewModelClass();
 		}
 	}
 
@@ -56,7 +56,7 @@ UE::MVVM::FBindingSource UE::MVVM::FBindingSource::CreateForViewModel(const UWid
 		{
 			Source.ViewModelId = ViewModel->GetViewModelId();
 			Source.DisplayName = ViewModel->GetDisplayName();
-			Source.Class = ViewModel->GetViewModelClass().Get();
+			Source.Class = ViewModel->GetViewModelClass();
 		}
 	}
 
diff --git a/Engine/Plugins/Runtime/ModelViewViewModel/Source/ModelViewViewModelEditor/Private/Widgets/SMVVMSelectViewModel.cpp b/Engine/Plugins/Runtime/ModelViewViewModel/Source/ModelViewViewModelEditor/Private/Widgets/SMVVMSelectViewModel.cpp
new file mode 100644
--- /dev/null
+++ b/Engine/Plugins/Runtime/ModelViewViewModel/Source/ModelViewViewModelEditor/Private/Widgets/SMVVMSelectViewModel.cpp
@@ -0,0 +1,172 @@
+// Copyright Epic Games, Inc. All Rights Reserved.
+
+#include "Widgets/SMVVMSelectViewModel.h"
+
+#include "Bindings/MVVMBindingHelper.h"
+#include "Blueprint/UserWidget.h"
+#include "Editor.h"
+#include "Framework/MultiBox/MultiBoxBuilder.h"
+#include "Kismet2/KismetEditorUtilities.h"
+#include "MVVMBlueprintViewModelContext.h"
+#include "MVVMEditorSubsystem.h"
+#include "MVVMViewModelBase.h"
+#include "SMVVMViewModelContextListWidget.h"
+#include "SMVVMViewModelBindingListWidget.h"
+#include "SPrimaryButton.h"
+#include "SSimpleButton.h"
+#include "Styling/MVVMEditorStyle.h"
+#include "Styling/SlateIconFinder.h"
+#include "Textures/SlateIcon.h"
+#include "Types/MVVMFieldVariant.h"
+#include "UMGStyle.h"
+#include "UObject/UObjectIterator.h"
+#include "Widgets/Images/SImage.h"
+#include "Widgets/Input/SButton.h"
+#include "Widgets/Input/SComboButton.h"
+#include "Widgets/Input/SEditableTextBox.h"
+#include "Widgets/Input/SSearchBox.h"
+#include "Widgets/Layout/SSplitter.h"
+#include "Widgets/Layout/SUniformGridPanel.h"
+#include "Widgets/SWindow.h"
+#include "Widgets/Views/SListView.h"
+#include "Widgets/Views/STreeView.h"
+
+//ClassViewer
+#include "ClassViewerModule.h"
+#include "ClassViewerFilter.h"
+#include "SClassViewer.h"
+#include "Components/Widget.h"
+#include "Modules/ModuleManager.h"
+
+#define LOCTEXT_NAMESPACE "SMVVMSelectViewModel"
+
+namespace UE::MVVM
+{
+
+namespace Private
+{
+	class FSelectViewModelClassFilter : public IClassViewerFilter
+	{
+	public:
+		FSelectViewModelClassFilter() = default;
+		static const EClassFlags ClassFlag = CLASS_HideDropDown | CLASS_Hidden | CLASS_Deprecated;
+		virtual bool IsClassAllowed(const FClassViewerInitializationOptions& InInitOptions, const UClass* InClass, TSharedRef<FClassViewerFilterFuncs> InFilterFuncs) override
+		{
+			if (InClass->IsChildOf(UWidget::StaticClass()))
+			{
+				return false;
+			}
+			return InClass->ImplementsInterface(UNotifyFieldValueChanged::StaticClass()) && !InClass->HasAnyClassFlags(ClassFlag);
+		}
+
+		virtual bool IsUnloadedClassAllowed(const FClassViewerInitializationOptions& InInitOptions, const TSharedRef<const IUnloadedBlueprintData> InUnloadedClassData, TSharedRef<FClassViewerFilterFuncs> InFilterFuncs) override
+		{
+			if (InUnloadedClassData->IsChildOf(UWidget::StaticClass()))
+			{
+				return false;
+			}
+			return InUnloadedClassData->ImplementsInterface(UNotifyFieldValueChanged::StaticClass()) && !InUnloadedClassData->HasAnyClassFlags(ClassFlag);
+		}
+	};
+} //namespace
+
+void SMVVMSelectViewModel::Construct(const FArguments& InArgs)
+{
+	OnCancel = InArgs._OnCancel;
+	OnViewModelCommitted = InArgs._OnViewModelCommitted;
+
+	TSharedRef<SWidget> ButtonsPanelContent = SNew(SUniformGridPanel)
+		.SlotPadding(FAppStyle::Get().GetMargin("StandardDialog.SlotPadding"))
+		+ SUniformGridPanel::Slot(0, 0)
+		[
+			SNew(SPrimaryButton)
+			.Text(LOCTEXT("ViewModelAddButtonText", "Ok"))
+			.OnClicked(this, &SMVVMSelectViewModel::HandleAccepted)
+			.IsEnabled(this, &SMVVMSelectViewModel::HandelIsSelectionEnabled)
+		]
+		+ SUniformGridPanel::Slot(1, 0)
+		[
+			SNew(SButton)
+			.Text(LOCTEXT("ViewModelCancelButtonText", "Cancel"))
+			.HAlign(HAlign_Center)
+			.OnClicked(this, &SMVVMSelectViewModel::HandleCancel)
+		];
+
+	FClassViewerInitializationOptions ClassViewerOptions;
+	ClassViewerOptions.DisplayMode = EClassViewerDisplayMode::TreeView;
+	ClassViewerOptions.Mode = EClassViewerMode::ClassPicker;
+	ClassViewerOptions.ClassFilters.Add(MakeShared<Private::FSelectViewModelClassFilter>());
+
+	ClassViewer = StaticCastSharedRef<SClassViewer>(FModuleManager::LoadModuleChecked<FClassViewerModule>("ClassViewer").CreateClassViewer(ClassViewerOptions, FOnClassPicked::CreateSP(this, &SMVVMSelectViewModel::HandleClassPicked)));
+
+	ChildSlot
+	[
+		SNew(SBorder)
+		.BorderImage(FAppStyle::GetBrush("ToolPanel.GroupBorder"))
+		[
+			SNew(SVerticalBox)
+			+ SVerticalBox::Slot()
+			.FillHeight(1.0f)
+			[
+				SNew(SSplitter)
+				.PhysicalSplitterHandleSize(2.0f)
+				+ SSplitter::Slot()
+				.Value(0.6f)
+				[
+					SNew(SBorder)
+					.BorderImage(FStyleDefaults::GetNoBrush())
+					[
+						ClassViewer.ToSharedRef()
+					]
+				]
+				+ SSplitter::Slot()
+				.Value(0.4f)
+				[
+					SNew(SBorder)
+					.BorderImage(FStyleDefaults::GetNoBrush())
+					[
+						SAssignNew(ViewModelBindingListWidget, SViewModelBindingListWidget)
+					]
+				]
+			]
+			+ SVerticalBox::Slot()
+			.HAlign(HAlign_Right)
+			.VAlign(VAlign_Bottom)
+			.AutoHeight()
+			.Padding(8)
+			[
+				ButtonsPanelContent
+			]
+		]
+	];
+}
+
+void SMVVMSelectViewModel::HandleClassPicked(UClass* ClassPicked)
+{
+	SelectedClass = ClassPicked;
+	ViewModelBindingListWidget->SetViewModel(ClassPicked, FName(), FGuid());
+}
+
+
+FReply SMVVMSelectViewModel::HandleAccepted()
+{
+	OnViewModelCommitted.ExecuteIfBound(SelectedClass.Get());
+	return FReply::Handled();
+}
+
+
+FReply SMVVMSelectViewModel::HandleCancel()
+{
+	OnCancel.ExecuteIfBound();
+	return FReply::Handled();
+}
+
+
+bool SMVVMSelectViewModel::HandelIsSelectionEnabled() const
+{
+	return SelectedClass.Get() != nullptr;
+}
+
+} //namespace
+
+#undef LOCTEXT_NAMESPACE
diff --git a/Engine/Plugins/Runtime/ModelViewViewModel/Source/ModelViewViewModelEditor/Private/Widgets/SMVVMSelectViewModel.h b/Engine/Plugins/Runtime/ModelViewViewModel/Source/ModelViewViewModelEditor/Private/Widgets/SMVVMSelectViewModel.h
new file mode 100644
--- /dev/null
+++ b/Engine/Plugins/Runtime/ModelViewViewModel/Source/ModelViewViewModelEditor/Private/Widgets/SMVVMSelectViewModel.h
@@ -0,0 +1,42 @@
+// Copyright Epic Games, Inc. All Rights Reserved.
+
+#pragma once
+
+#include "CoreMinimal.h"
+#include "Widgets/SCompoundWidget.h"
+
+class SClassViewer;
+
+namespace UE::MVVM
+{
+class SViewModelBindingListWidget;
+
+class SMVVMSelectViewModel : public SCompoundWidget
+{
+public:
+	DECLARE_DELEGATE_OneParam(FOnValueChanged, const UClass*);
+
+	SLATE_BEGIN_ARGS(SMVVMSelectViewModel) {}
+		SLATE_EVENT(FSimpleDelegate, OnCancel)
+		SLATE_EVENT(FOnValueChanged, OnViewModelCommitted)
+	SLATE_END_ARGS()
+
+public:
+	void Construct(const FArguments& InArgs);
+
+private:
+	void HandleClassPicked(UClass* ClassPicked);
+	FReply HandleAccepted();
+	FReply HandleCancel();
+	bool HandelIsSelectionEnabled() const;
+
+private:
+	TSharedPtr<SClassViewer> ClassViewer;
+	TSharedPtr<SViewModelBindingListWidget> ViewModelBindingListWidget;
+	TWeakObjectPtr<const UClass> SelectedClass;
+
+	FSimpleDelegate OnCancel;
+	FOnValueChanged OnViewModelCommitted;
+};
+
+} //namespace
\ No newline at end of file
diff --git a/Engine/Plugins/Runtime/ModelViewViewModel/Source/ModelViewViewModelEditor/Private/Widgets/SMVVMViewBindingListView.cpp b/Engine/Plugins/Runtime/ModelViewViewModel/Source/ModelViewViewModelEditor/Private/Widgets/SMVVMViewBindingListView.cpp
--- a/Engine/Plugins/Runtime/ModelViewViewModel/Source/ModelViewViewModelEditor/Private/Widgets/SMVVMViewBindingListView.cpp
+++ b/Engine/Plugins/Runtime/ModelViewViewModel/Source/ModelViewViewModelEditor/Private/Widgets/SMVVMViewBindingListView.cpp
@@ -415,7 +415,7 @@ private:
 			if (const FMVVMBlueprintViewModelContext* ViewModel = BlueprintViewPtr->FindViewModel(ViewBinding->ViewModelPath.GetViewModelId()))
 			{
 				UMVVMSubsystem* Subsystem = GEngine->GetEngineSubsystem<UMVVMSubsystem>();
-				TArray<FMVVMAvailableBinding> AvailableBindings = Subsystem->GetViewModelAvailableBindings(ViewModel->GetViewModelClass());
+				TArray<FMVVMAvailableBinding> AvailableBindings = Subsystem->GetAvailableBindings(ViewModel->GetViewModelClass());
 
 				Algo::Transform(AvailableBindings, AvailablePaths, [ViewModel](const FMVVMAvailableBinding& Binding)
 					{
diff --git a/Engine/Plugins/Runtime/ModelViewViewModel/Source/ModelViewViewModelEditor/Private/Widgets/SMVVMViewModelBindingListWidget.cpp b/Engine/Plugins/Runtime/ModelViewViewModel/Source/ModelViewViewModelEditor/Private/Widgets/SMVVMViewModelBindingListWidget.cpp
--- a/Engine/Plugins/Runtime/ModelViewViewModel/Source/ModelViewViewModelEditor/Private/Widgets/SMVVMViewModelBindingListWidget.cpp
+++ b/Engine/Plugins/Runtime/ModelViewViewModel/Source/ModelViewViewModelEditor/Private/Widgets/SMVVMViewModelBindingListWidget.cpp
@@ -140,7 +140,7 @@ void SViewModelBindingListWidget::SetRawFilterText(const FText& InFilterText)
 {
 	if (PropertyViewer)
 	{
-		//PropertyViewer->SetRawFilterText(InFilterText);
+		PropertyViewer->SetRawFilterText(InFilterText);
 	}
 }
 
diff --git a/Engine/Plugins/Runtime/ModelViewViewModel/Source/ModelViewViewModelEditor/Private/Widgets/SMVVMViewModelContextListWidget.cpp b/Engine/Plugins/Runtime/ModelViewViewModel/Source/ModelViewViewModelEditor/Private/Widgets/SMVVMViewModelContextListWidget.cpp
--- a/Engine/Plugins/Runtime/ModelViewViewModel/Source/ModelViewViewModelEditor/Private/Widgets/SMVVMViewModelContextListWidget.cpp
+++ b/Engine/Plugins/Runtime/ModelViewViewModel/Source/ModelViewViewModelEditor/Private/Widgets/SMVVMViewModelContextListWidget.cpp
@@ -29,7 +29,7 @@
 
 namespace UE::MVVM::Private
 {
-	static bool IsPropertyTypeChildOf(const FProperty* Property, TSubclassOf<UMVVMViewModelBase> ParentClass)
+	static bool IsPropertyTypeChildOf(const FProperty* Property, const UClass* ParentClass)
 	{
 		if (const FObjectPropertyBase* ObjectProperty = CastField<FObjectPropertyBase>(Property))
 		{
@@ -42,7 +42,7 @@ namespace UE::MVVM::Private
 		return false;
 	}
 
-	static bool HasBindableViewModelFieldRecursive(UStruct* InStruct, TSubclassOf<UMVVMViewModelBase> InParentClass, TSet<UStruct*>& VisitedSet, uint32 RecursionDepth)
+	static bool HasBindableViewModelFieldRecursive(UStruct* InStruct, const UClass* InParentClass, TSet<UStruct*>& VisitedSet, uint32 RecursionDepth)
 	{
 		if (RecursionDepth > 10)
 		{
@@ -393,9 +393,9 @@ namespace UE::MVVM::Private
 			};
 
 			UClass* InBindableClass = Cast<UClass>(InOwnerStruct);
-			TSubclassOf<UMVVMViewModelBase> ViewModelClass = Entry->GetViewModelClass();
+			const UClass* ViewModelClass = Entry->GetViewModelClass();
 
-			check(ViewModelClass.Get());
+			check(ViewModelClass);
 
 			bool bHasFunctionEntry = false;
 			MenuBuilder.BeginSection("Functions", LOCTEXT("Functions", "Functions"));
diff --git a/Engine/Plugins/Runtime/ModelViewViewModel/Source/ModelViewViewModelEditor/Private/Widgets/SMVVMViewModelPanel.cpp b/Engine/Plugins/Runtime/ModelViewViewModel/Source/ModelViewViewModelEditor/Private/Widgets/SMVVMViewModelPanel.cpp
--- a/Engine/Plugins/Runtime/ModelViewViewModel/Source/ModelViewViewModelEditor/Private/Widgets/SMVVMViewModelPanel.cpp
+++ b/Engine/Plugins/Runtime/ModelViewViewModel/Source/ModelViewViewModelEditor/Private/Widgets/SMVVMViewModelPanel.cpp
@@ -21,7 +21,9 @@
 #include "Widgets/SBoxPanel.h"
 #include "Widgets/Input/SSearchBox.h"
 #include "Widgets/Layout/SBorder.h"
+#include "Widgets/Layout/SBox.h"
 #include "Widgets/PropertyViewer/SPropertyViewer.h"
+#include "Widgets/SMVVMSelectViewModel.h"
 #include "Widgets/SMVVMViewModelBindingListWidget.h"
 
 #include "SPositiveActionButton.h"
@@ -54,7 +56,7 @@ void SMVVMViewModelPanel::Construct(const FArguments& InArgs, TSharedPtr<FWidget
 		.FieldIterator(&ViewModelFieldIterator)
 		.SearchBoxPreSlot()
 		[
-			SNew(SPositiveActionButton)
+			SAssignNew(AddMenuButton, SPositiveActionButton)
 			.OnGetMenuContent(this, &SMVVMViewModelPanel::MakeAddMenu)
 			.Text(LOCTEXT("Viewmodel", "Viewmodel"))
 			.IsEnabled(this, &SMVVMViewModelPanel::HandleCanEditViewmodelList)
@@ -106,12 +108,23 @@ void SMVVMViewModelPanel::FillViewModel()
 		{
 			for (const FMVVMBlueprintViewModelContext& ViewModelContext : View->GetViewModels())
 			{
-				UObject* ViewModelInstance = ViewModelContext.GetViewModelClass().GetDefaultObject();
-				
-				FViewModelHandle ViewModelHandle;
-				ViewModelHandle.Key = ViewModelContext.GetViewModelId();
-				ViewModelHandle.Value = ViewModelTreeView->AddInstance(ViewModelInstance);
-				ViewModelHandles.Add(ViewModelHandle);
+				if (ViewModelContext.GetViewModelClass())
+				{
+					UObject* ViewModelInstance = ViewModelContext.GetViewModelClass()->GetDefaultObject();
+
+					FViewModelHandle ViewModelHandle;
+					ViewModelHandle.Key = ViewModelContext.GetViewModelId();
+					ViewModelHandle.Value = ViewModelTreeView->AddInstance(ViewModelInstance, ViewModelContext.GetDisplayName());
+					ViewModelHandles.Add(ViewModelHandle);
+				}
+				else
+				{
+					FViewModelHandle ViewModelHandle;
+					ViewModelHandle.Key = ViewModelContext.GetViewModelId();
+					// Find a way to show context that are not valid anymore
+					//ViewModelHandle.Value = ViewModelTreeView->AddInstance(ViewModelInstance, ViewModelContext.GetDisplayName());
+					ViewModelHandles.Add(ViewModelHandle);
+				}
 			}
 		}
 	}
@@ -155,17 +168,44 @@ void SMVVMViewModelPanel::HandleViewModelsUpdated()
 
 TSharedRef<SWidget> SMVVMViewModelPanel::MakeAddMenu()
 {
-	//using namespace UE::Sequencer;
+	return SNew(SBox)
+		.WidthOverride(600)
+		.HeightOverride(500)
+		[
+			SNew(SMVVMSelectViewModel)
+			.OnCancel(this, &SMVVMViewModelPanel::HandleCancelAddMenu)
+			.OnViewModelCommitted(this, &SMVVMViewModelPanel::HandleAddMenuViewModel)
+		];
+}
+
+
+void SMVVMViewModelPanel::HandleCancelAddMenu()
+{
+	if (AddMenuButton)
+	{
+		AddMenuButton->SetIsMenuOpen(false, false);
+	}
+}
 
-	//TSharedPtr<FExtender> Extender = FExtender::Combine(AddMenuExtenders);
-	//FMenuBuilder MenuBuilder(true, nullptr, Extender);
-	//{
-	//	TSharedPtr<FSequencer> Sequencer = SequencerPtr.Pin();
-	//	Sequencer->GetViewModel()->GetOutliner()->CastThisChecked<FSequencerOutlinerViewModel>()->BuildContextMenu(MenuBuilder);
-	//}
 
-	//return MenuBuilder.MakeWidget();
-	return SNullWidget::NullWidget;
+void SMVVMViewModelPanel::HandleAddMenuViewModel(const UClass* SelectedClass)
+{
+	if (AddMenuButton)
+	{
+		AddMenuButton->SetIsMenuOpen(false, false);
+		if (SelectedClass)
+		{
+			if (TSharedPtr<FWidgetBlueprintEditor> WidgetBlueprintEditor = WeakBlueprintEditor.Pin())
+			{
+				if (UWidgetBlueprint* WidgetBlueprint = WidgetBlueprintEditor->GetWidgetBlueprintObj())
+				{
+					UMVVMEditorSubsystem* EditorSubsystem = GEditor->GetEditorSubsystem<UMVVMEditorSubsystem>();
+					check(EditorSubsystem);
+					EditorSubsystem->AddViewModel(WidgetBlueprint, SelectedClass);
+				}
+			}
+		}
+	}
 }
 
 
diff --git a/Engine/Plugins/Runtime/ModelViewViewModel/Source/ModelViewViewModelEditor/Private/Widgets/SMVVMViewModelPanel.h b/Engine/Plugins/Runtime/ModelViewViewModel/Source/ModelViewViewModelEditor/Private/Widgets/SMVVMViewModelPanel.h
--- a/Engine/Plugins/Runtime/ModelViewViewModel/Source/ModelViewViewModelEditor/Private/Widgets/SMVVMViewModelPanel.h
+++ b/Engine/Plugins/Runtime/ModelViewViewModel/Source/ModelViewViewModelEditor/Private/Widgets/SMVVMViewModelPanel.h
@@ -9,6 +9,7 @@
 #include "Widgets/SCompoundWidget.h"
 
 class FWidgetBlueprintEditor;
+class SPositiveActionButton;
 class UBlueprintExtension;
 class UMVVMBlueprintView;
 
@@ -32,11 +33,14 @@ private:
 	void HandleViewModelsUpdated();
 	//bool HandleViewModelRenamed(const SViewModelBindingListWidget::FViewModel& ViewModel, const FText& RenameTo, bool bCommit, FText& OutErrorMessage);
 	TSharedRef<SWidget> MakeAddMenu();
+	void HandleCancelAddMenu();
+	void HandleAddMenuViewModel(const UClass* SelectedClass);
 	bool HandleCanEditViewmodelList() const;
 
 	void FillViewModel();
 
 private:
+	TSharedPtr<SPositiveActionButton> AddMenuButton;
 	TSharedPtr<UE::PropertyViewer::SPropertyViewer> ViewModelTreeView;
 	FFieldIterator_ViewModel ViewModelFieldIterator;
 
diff --git a/Engine/Plugins/Runtime/ModelViewViewModel/Source/ModelViewViewModelEditor/Public/MVVMEditorSubsystem.h b/Engine/Plugins/Runtime/ModelViewViewModel/Source/ModelViewViewModelEditor/Public/MVVMEditorSubsystem.h
--- a/Engine/Plugins/Runtime/ModelViewViewModel/Source/ModelViewViewModelEditor/Public/MVVMEditorSubsystem.h
+++ b/Engine/Plugins/Runtime/ModelViewViewModel/Source/ModelViewViewModelEditor/Public/MVVMEditorSubsystem.h
@@ -28,6 +28,9 @@ public:
 
 	UFUNCTION(BlueprintCallable, Category = "MVVM")
 	UMVVMBlueprintView* GetView(const UWidgetBlueprint* WidgetBlueprint) const;
+
+	UFUNCTION(BlueprintCallable, Category = "MVVM")
+	FName AddViewModel(UWidgetBlueprint* WidgetBlueprint, const UClass* ViewModel);
 	
 	UFUNCTION(BlueprintCallable, Category = "MVVM")
 	void RemoveViewModel(UWidgetBlueprint* WidgetBlueprint, FName ViewModel);
