MVVM: Change how available binding property are gather to respect the field privacy. Remove the different override that are not necessary anymore. Clean the include in SSelectViewModel. Give the WidgetBlueprint to different widget to be use by GetAvailableBindings.

diff --git a/Engine/Plugins/Runtime/ModelViewViewModel/Source/ModelViewViewModel/Private/Bindings/MVVMBindingHelper.cpp b/Engine/Plugins/Runtime/ModelViewViewModel/Source/ModelViewViewModel/Private/Bindings/MVVMBindingHelper.cpp
--- a/Engine/Plugins/Runtime/ModelViewViewModel/Source/ModelViewViewModel/Private/Bindings/MVVMBindingHelper.cpp
+++ b/Engine/Plugins/Runtime/ModelViewViewModel/Source/ModelViewViewModel/Private/Bindings/MVVMBindingHelper.cpp
@@ -46,11 +46,7 @@ namespace UE::MVVM::BindingHelper
 		return InProperty != nullptr
 			&& InProperty->ArrayDim == 1
 			&& !InProperty->HasAnyPropertyFlags(CPF_Deprecated | CPF_EditorOnly)
-			&& InProperty->HasAnyPropertyFlags(CPF_BlueprintVisible | CPF_BlueprintAssignable)
-#if WITH_EDITOR
-			&& !InProperty->GetBoolMetaData(Private::NAME_BlueprintPrivate)
-#endif
-			;
+			&& InProperty->HasAnyPropertyFlags(CPF_BlueprintVisible | CPF_BlueprintAssignable);
 	}
 
 
@@ -59,11 +55,7 @@ namespace UE::MVVM::BindingHelper
 		return InProperty != nullptr
 			&& InProperty->ArrayDim == 1
 			&& !InProperty->HasAnyPropertyFlags(CPF_Deprecated | CPF_EditorOnly | CPF_BlueprintReadOnly)
-			&& InProperty->HasAnyPropertyFlags(CPF_BlueprintVisible | CPF_BlueprintAssignable)
-#if WITH_EDITOR
-			&& !InProperty->GetBoolMetaData(Private::NAME_BlueprintPrivate)
-#endif
-			;
+			&& InProperty->HasAnyPropertyFlags(CPF_BlueprintVisible | CPF_BlueprintAssignable);
 	}
 
 
diff --git a/Engine/Plugins/Runtime/ModelViewViewModel/Source/ModelViewViewModel/Private/MVVMSubsystem.cpp b/Engine/Plugins/Runtime/ModelViewViewModel/Source/ModelViewViewModel/Private/MVVMSubsystem.cpp
--- a/Engine/Plugins/Runtime/ModelViewViewModel/Source/ModelViewViewModel/Private/MVVMSubsystem.cpp
+++ b/Engine/Plugins/Runtime/ModelViewViewModel/Source/ModelViewViewModel/Private/MVVMSubsystem.cpp
@@ -29,78 +29,6 @@ UMVVMView* UMVVMSubsystem::GetViewFromUserWidget(const UUserWidget* UserWidget)
 }
 
 
-bool UMVVMSubsystem::IsViewModelValueValidForSourceBinding(const UMVVMViewModelBase* ViewModel, FMVVMBindingName ViewModelPropertyOrFunctionName) const
-{
-	if (ViewModel)
-	{
-		UE::MVVM::FMVVMFieldVariant Binding = UE::MVVM::BindingHelper::FindFieldByName(ViewModel->GetClass(), ViewModelPropertyOrFunctionName);
-		if (Binding.IsFunction())
-		{
-			return UE::MVVM::BindingHelper::IsValidForSourceBinding(Binding.GetFunction());
-		}
-		if (Binding.IsProperty())
-		{
-			return UE::MVVM::BindingHelper::IsValidForSourceBinding(Binding.GetProperty());
-		}
-	}
-	return false;
-}
-
-
-bool UMVVMSubsystem::IsViewModelValueValidForDestinationBinding(const UMVVMViewModelBase* ViewModel, FMVVMBindingName ViewModelPropertyOrFunctionName) const
-{
-	if (ViewModel)
-	{
-		UE::MVVM::FMVVMFieldVariant Binding = UE::MVVM::BindingHelper::FindFieldByName(ViewModel->GetClass(), ViewModelPropertyOrFunctionName);
-		if (Binding.IsFunction())
-		{
-			return UE::MVVM::BindingHelper::IsValidForDestinationBinding(Binding.GetFunction());
-		}
-		if (Binding.IsProperty())
-		{
-			return UE::MVVM::BindingHelper::IsValidForDestinationBinding(Binding.GetProperty());
-		}
-	}
-	return false;
-}
-
-
-bool UMVVMSubsystem::IsViewValueValidForSourceBinding(const UWidget* View, FMVVMBindingName ViewPropertyOrFunctionName) const
-{
-	if (View)
-	{
-		UE::MVVM::FMVVMFieldVariant Binding = UE::MVVM::BindingHelper::FindFieldByName(View->GetClass(), ViewPropertyOrFunctionName);
-		if (Binding.IsFunction())
-		{
-			return UE::MVVM::BindingHelper::IsValidForSourceBinding(Binding.GetFunction());
-		}
-		if (Binding.IsProperty())
-		{
-			return UE::MVVM::BindingHelper::IsValidForSourceBinding(Binding.GetProperty());
-		}
-	}
-	return false;
-}
-
-
-bool UMVVMSubsystem::IsViewValueValidForDestinationBinding(const UWidget* View, FMVVMBindingName ViewPropertyOrFunctionName) const
-{
-	if (View)
-	{
-		UE::MVVM::FMVVMFieldVariant Binding = UE::MVVM::BindingHelper::FindFieldByName(View->GetClass(), ViewPropertyOrFunctionName);
-		if (Binding.IsFunction())
-		{
-			return UE::MVVM::BindingHelper::IsValidForDestinationBinding(Binding.GetFunction());
-		}
-		if (Binding.IsProperty())
-		{
-			return UE::MVVM::BindingHelper::IsValidForDestinationBinding(Binding.GetProperty());
-		}
-	}
-	return false;
-}
-
-
 bool UMVVMSubsystem::DoesWidgetTreeContainedWidget(const UWidgetTree* WidgetTree, const UWidget* ViewWidget) const
 {
 	// Test if the View's Widget is valid.
@@ -116,12 +44,12 @@ bool UMVVMSubsystem::DoesWidgetTreeContainedWidget(const UWidgetTree* WidgetTree
 
 namespace UE::MVVM::Private
 {
-	TArray<FMVVMAvailableBinding> GetAvailableBindings(const UClass* Class, const UE::FieldNotification::IClassDescriptor* ClassDescriptor)
+	TArray<FMVVMAvailableBinding> GetAvailableBindings(const UStruct* Class, const UClass* AccessorType, const UE::FieldNotification::IClassDescriptor* ClassDescriptor)
 	{
 		TSet<FName> FieldDescriptors;
 		if (ClassDescriptor)
 		{
-			ClassDescriptor->ForEachField(Class, [&FieldDescriptors](UE::FieldNotification::FFieldId FieldId)
+			ClassDescriptor->ForEachField(CastChecked<UClass>(Class), [&FieldDescriptors](UE::FieldNotification::FFieldId FieldId)
 			{
 				FieldDescriptors.Add(FieldId.GetName());
 				return true;
@@ -131,12 +59,26 @@ namespace UE::MVVM::Private
 		TArray<FMVVMAvailableBinding> Result;
 		Result.Reserve(FieldDescriptors.Num());
 
-		for (TFieldIterator<const UFunction> FunctionIt(Class, EFieldIteratorFlags::IncludeSuper); FunctionIt; ++FunctionIt)
+		const bool bAccessorCanSeeProtectedMember = AccessorType ? AccessorType->IsChildOf(Class) : false;
+
+		if (Cast<UClass>(Class))
 		{
-			const UFunction* Function = *FunctionIt;
-			check(Function);
-			//if (!Function->HasAnyFunctionFlags(FUNC_Private | FUNC_Protected)) // todo fix for property and BlueprintPrivate in IsValidForSourceBindnig
+			for (TFieldIterator<const UFunction> FunctionIt(Class, EFieldIteratorFlags::IncludeSuper); FunctionIt; ++FunctionIt)
 			{
+				const UFunction* Function = *FunctionIt;
+				check(Function);
+
+				// N.B. A function can be private/protected and can still be use when it's a BlueprintGetter/BlueprintSetter.
+				//But we bind to the property not the function.
+				if (Function->HasAnyFunctionFlags(FUNC_Private) && Class != AccessorType)
+				{
+					continue;
+				}
+				if (Function->HasAnyFunctionFlags(FUNC_Protected) && !bAccessorCanSeeProtectedMember)
+				{
+					continue;
+				}
+
 				const bool bIsReadable = BindingHelper::IsValidForSourceBinding(Function);
 				const bool bIsWritable = BindingHelper::IsValidForDestinationBinding(Function);
 				const bool bHasNotify = FieldDescriptors.Contains(Function->GetFName()) && bIsReadable;
@@ -150,12 +92,22 @@ namespace UE::MVVM::Private
 #if WITH_EDITOR
 		FName NAME_BlueprintGetter = "BlueprintGetter";
 		FName NAME_BlueprintSetter = "BlueprintSetter";
+		FName NAME_BlueprintPrivate = "BlueprintPrivate";
 #endif
 
-		// Remove BP getters / setters and replace them with their property for bindings
 		for (TFieldIterator<const FProperty> PropertyIt(Class, EFieldIteratorFlags::IncludeSuper); PropertyIt; ++PropertyIt)
 		{
 			const FProperty* Property = *PropertyIt;
+
+			// N.B. Property can be private/protected in cpp but if they are visible in BP that is all that matter.
+			//Only property defined in BP can be BP visible and really private.
+#if WITH_EDITOR
+			if (Property->GetBoolMetaData(NAME_BlueprintPrivate) && Class != AccessorType)
+			{
+				continue;
+			}
+#endif
+
 			const bool bIsReadable = BindingHelper::IsValidForSourceBinding(Property);
 			const bool bIsWritable = BindingHelper::IsValidForDestinationBinding(Property);
 			const bool bHasNotify = FieldDescriptors.Contains(Property->GetFName()) && bIsReadable;
@@ -197,37 +149,31 @@ namespace UE::MVVM::Private
 		return Result;
 	}
 
-	TArray<FMVVMAvailableBinding> GetAvailableBindings(TSubclassOf<UObject> InSubClass)
+	TArray<FMVVMAvailableBinding> GetAvailableBindings(TSubclassOf<UObject> InSubClass, TSubclassOf<UObject> InAccessor)
 	{
 		if (InSubClass.Get() && InSubClass.Get()->ImplementsInterface(UNotifyFieldValueChanged::StaticClass()))
 		{
 			TScriptInterface<INotifyFieldValueChanged> DefaultObject = InSubClass.GetDefaultObject();
 			const UE::FieldNotification::IClassDescriptor& ClassDescriptor = DefaultObject->GetFieldNotificationDescriptor();
-			return GetAvailableBindings(InSubClass.Get(), &ClassDescriptor);
+			return GetAvailableBindings(InSubClass.Get(), InAccessor.Get(), &ClassDescriptor);
 		}
 		else
 		{
-			return GetAvailableBindings(InSubClass.Get(), nullptr);
+			return GetAvailableBindings(InSubClass.Get(), InAccessor.Get(), nullptr);
 		}
 	}
 } //namespace
 
 
-TArray<FMVVMAvailableBinding> UMVVMSubsystem::GetViewModelAvailableBindings(TSubclassOf<UMVVMViewModelBase> ViewModelClass) const
-{
-	return UE::MVVM::Private::GetAvailableBindings(ViewModelClass);
-}
-
-
-TArray<FMVVMAvailableBinding> UMVVMSubsystem::GetWidgetAvailableBindings(TSubclassOf<UWidget> WidgetClass) const
+TArray<FMVVMAvailableBinding> UMVVMSubsystem::GetAvailableBindings(TSubclassOf<UObject> Class, TSubclassOf<UObject> Accessor) const
 {
-	return UE::MVVM::Private::GetAvailableBindings(WidgetClass);
+	return UE::MVVM::Private::GetAvailableBindings(Class, Accessor);
 }
 
 
-TArray<FMVVMAvailableBinding> UMVVMSubsystem::GetAvailableBindings(TSubclassOf<UObject> Class) const
+TArray<FMVVMAvailableBinding> UMVVMSubsystem::GetAvailableBindingsForStruct(const UScriptStruct* Struct) const
 {
-	return UE::MVVM::Private::GetAvailableBindings(Class);
+	return UE::MVVM::Private::GetAvailableBindings(Struct, nullptr, nullptr);
 }
 
 
diff --git a/Engine/Plugins/Runtime/ModelViewViewModel/Source/ModelViewViewModel/Public/MVVMSubsystem.h b/Engine/Plugins/Runtime/ModelViewViewModel/Source/ModelViewViewModel/Public/MVVMSubsystem.h
--- a/Engine/Plugins/Runtime/ModelViewViewModel/Source/ModelViewViewModel/Public/MVVMSubsystem.h
+++ b/Engine/Plugins/Runtime/ModelViewViewModel/Source/ModelViewViewModel/Public/MVVMSubsystem.h
@@ -32,35 +32,21 @@ public:
 	UFUNCTION(BlueprintCallable, Category="MVVM")
 	UMVVMView* GetViewFromUserWidget(const UUserWidget* UserWidget) const;
 
-	UFUNCTION(BlueprintCallable, Category="MVVM")
-	bool IsViewModelValueValidForSourceBinding(const UMVVMViewModelBase* ViewModel, FMVVMBindingName ViewModelPropertyOrFunctionName) const;
-
-	UFUNCTION(BlueprintCallable, Category = "MVVM")
-	bool IsViewModelValueValidForDestinationBinding(const UMVVMViewModelBase* ViewModel, FMVVMBindingName ViewModelPropertyOrFunctionName) const;
-
-	UFUNCTION(BlueprintCallable, Category = "MVVM")
-	bool IsViewValueValidForSourceBinding(const UWidget* View, FMVVMBindingName ViewPropertyOrFunctionName) const;
-
-	UFUNCTION(BlueprintCallable, Category = "MVVM")
-	bool IsViewValueValidForDestinationBinding(const UWidget* View, FMVVMBindingName ViewPropertyOrFunctionName) const;
-
 	UFUNCTION(BlueprintCallable, Category = "MVVM")
 	bool DoesWidgetTreeContainedWidget(const UWidgetTree* WidgetTree, const UWidget* ViewWidget) const;
 
-	/** Returns the list of all the bindings that are available for the ViewModel. */
-	UFUNCTION(BlueprintCallable, Category = "MVVM")
-	TArray<FMVVMAvailableBinding> GetViewModelAvailableBindings(TSubclassOf<UMVVMViewModelBase> ViewModelClass) const;
-
-	/** Returns the list of all the bindings that are available for the Widget. */
-	UFUNCTION(BlueprintCallable, Category = "MVVM")
-	TArray<FMVVMAvailableBinding> GetWidgetAvailableBindings(TSubclassOf<UWidget> WidgetClass) const;
-
 	/**
 	 * Returns the list of all the bindings that are available for the Class.
 	 * The class must implement the UNotifyFieldValueChanged interface.
 	 */
 	UFUNCTION(BlueprintCallable, Category = "MVVM")
-	TArray<FMVVMAvailableBinding> GetAvailableBindings(TSubclassOf<UObject> Class) const;
+	TArray<FMVVMAvailableBinding> GetAvailableBindings(TSubclassOf<UObject> Class, TSubclassOf<UObject> Accessor) const;
+
+	/**
+	 * Returns the list of all the bindings that are available from the SriptStuct.
+	 * @note When FMVVMAvailableBinding::HasNotify is false, a notification can still be triggered by the owner of the struct. The struct changed but which property of the struct changed is unknown.
+	 */
+	TArray<FMVVMAvailableBinding> GetAvailableBindingsForStruct(const UScriptStruct* Struct) const;
 
 	UFUNCTION(BlueprintCallable, Category = "MVVM")
 	UMVVMViewModelCollectionObject* GetGlobalViewModelCollection() const
diff --git a/Engine/Plugins/Runtime/ModelViewViewModel/Source/ModelViewViewModelEditor/Private/Customizations/MVVMConversionPathCustomization.cpp b/Engine/Plugins/Runtime/ModelViewViewModel/Source/ModelViewViewModelEditor/Private/Customizations/MVVMConversionPathCustomization.cpp
--- a/Engine/Plugins/Runtime/ModelViewViewModel/Source/ModelViewViewModelEditor/Private/Customizations/MVVMConversionPathCustomization.cpp
+++ b/Engine/Plugins/Runtime/ModelViewViewModel/Source/ModelViewViewModelEditor/Private/Customizations/MVVMConversionPathCustomization.cpp
@@ -474,7 +474,7 @@ namespace UE::MVVM
 		}
 
 		UMVVMSubsystem* Subsystem = GEngine->GetEngineSubsystem<UMVVMSubsystem>();
-		TArray<FMVVMAvailableBinding> AvailableBindings = Subsystem->GetAvailableBindings(SourceClass);
+		TArray<FMVVMAvailableBinding> AvailableBindings = Subsystem->GetAvailableBindings(SourceClass, WidgetBlueprint->GeneratedClass);
 
 		const UEdGraphSchema_K2* Schema = GetDefault<UEdGraphSchema_K2>();
 		FEdGraphPinType ArgumentType = GetArgumentPinType(ArgumentName, bSourceToDestination);
diff --git a/Engine/Plugins/Runtime/ModelViewViewModel/Source/ModelViewViewModelEditor/Private/Customizations/MVVMPropertyPathCustomization.cpp b/Engine/Plugins/Runtime/ModelViewViewModel/Source/ModelViewViewModelEditor/Private/Customizations/MVVMPropertyPathCustomization.cpp
--- a/Engine/Plugins/Runtime/ModelViewViewModel/Source/ModelViewViewModelEditor/Private/Customizations/MVVMPropertyPathCustomization.cpp
+++ b/Engine/Plugins/Runtime/ModelViewViewModel/Source/ModelViewViewModelEditor/Private/Customizations/MVVMPropertyPathCustomization.cpp
@@ -4,6 +4,7 @@
 
 #include "Algo/RemoveIf.h"
 #include "Algo/Transform.h"
+#include "Bindings/MVVMBindingHelper.h"
 #include "BlueprintEditor.h"
 #include "Components/Widget.h"
 #include "DetailWidgetRow.h"
@@ -285,34 +286,14 @@ namespace UE::MVVM
 
 		UClass* Class = Source.Class;
 		UMVVMSubsystem* Subsystem = GEngine->GetEngineSubsystem<UMVVMSubsystem>();
-		if (bIsWidget)
-		{
-			AvailableBindings = Subsystem->GetWidgetAvailableBindings(Class);
-		}
-		else
-		{
-			AvailableBindings = Subsystem->GetViewModelAvailableBindings(Class);
-		}
+		AvailableBindings = Subsystem->GetAvailableBindings(Class, WidgetBlueprint->GeneratedClass);
 
 		TArray<FMVVMBlueprintPropertyPath> AvailablePaths;
 		Algo::Transform(AvailableBindings, AvailablePaths, 
 			[this, Class, Source](const FMVVMAvailableBinding& Binding)
 			{
-				FName BindingName = Binding.GetBindingName().ToName();
-
-				UE::MVVM::FMVVMConstFieldVariant Variant;
-
-				if (const UFunction* Function = Class->FindFunctionByName(BindingName))
-				{
-					Variant = UE::MVVM::FMVVMConstFieldVariant(Function);
-				}
-				else if (const FProperty* Property = Class->FindPropertyByName(BindingName))
-				{
-					Variant = UE::MVVM::FMVVMConstFieldVariant(Property);
-				}
-
+				UE::MVVM::FMVVMConstFieldVariant Variant = UE::MVVM::BindingHelper::FindFieldByName(Class, Binding.GetBindingName());
 				FMVVMBlueprintPropertyPath Path;
-
 				if (!Variant.IsEmpty())
 				{
 					if (bIsWidget)
diff --git a/Engine/Plugins/Runtime/ModelViewViewModel/Source/ModelViewViewModelEditor/Private/MVVMEditorSubsystem.cpp b/Engine/Plugins/Runtime/ModelViewViewModel/Source/ModelViewViewModelEditor/Private/MVVMEditorSubsystem.cpp
--- a/Engine/Plugins/Runtime/ModelViewViewModel/Source/ModelViewViewModelEditor/Private/MVVMEditorSubsystem.cpp
+++ b/Engine/Plugins/Runtime/ModelViewViewModel/Source/ModelViewViewModelEditor/Private/MVVMEditorSubsystem.cpp
@@ -1009,7 +1009,7 @@ TArray<UE::MVVM::FBindingSource> UMVVMEditorSubsystem::GetBindableWidgets(const
 	// Add current widget as a possible binding source
 	if (UClass* BPClass = WidgetBlueprint->GeneratedClass)
 	{
-		TArray<FMVVMAvailableBinding> Bindings = GEngine->GetEngineSubsystem<UMVVMSubsystem>()->GetWidgetAvailableBindings(BPClass);
+		TArray<FMVVMAvailableBinding> Bindings = GEngine->GetEngineSubsystem<UMVVMSubsystem>()->GetAvailableBindings(BPClass, WidgetBlueprint->GeneratedClass);
 		if (Bindings.Num() > 0)
 		{
 			// at least one valid property, add it to our list
@@ -1023,7 +1023,7 @@ TArray<UE::MVVM::FBindingSource> UMVVMEditorSubsystem::GetBindableWidgets(const
 
 	for (const UWidget* Widget : AllWidgets)
 	{
-		TArray<FMVVMAvailableBinding> Bindings = GEngine->GetEngineSubsystem<UMVVMSubsystem>()->GetWidgetAvailableBindings(Widget->GetClass());
+		TArray<FMVVMAvailableBinding> Bindings = GEngine->GetEngineSubsystem<UMVVMSubsystem>()->GetAvailableBindings(Widget->GetClass(), WidgetBlueprint->GeneratedClass);
 		if (Bindings.Num() > 0)
 		{
 			// at least one valid property, add it to our list
@@ -1060,14 +1060,14 @@ TArray<UE::MVVM::FBindingSource> UMVVMEditorSubsystem::GetAllViewModels(const UW
 	return Sources;
 }
 
-TArray<FMVVMAvailableBinding> UMVVMEditorSubsystem::GetChildViewModels(TSubclassOf<UObject> Class)
+TArray<FMVVMAvailableBinding> UMVVMEditorSubsystem::GetChildViewModels(TSubclassOf<UObject> Class, TSubclassOf<UObject> Accessor)
 {
 	if (Class.Get() == nullptr)
 	{
 		return TArray<FMVVMAvailableBinding>();
 	}
 
-	TArray<FMVVMAvailableBinding> ViewModelAvailableBindingsList = GEngine->GetEngineSubsystem<UMVVMSubsystem>()->GetAvailableBindings(Class);
+	TArray<FMVVMAvailableBinding> ViewModelAvailableBindingsList = GEngine->GetEngineSubsystem<UMVVMSubsystem>()->GetAvailableBindings(Class, Accessor);
 	ViewModelAvailableBindingsList.RemoveAllSwap([Class](const FMVVMAvailableBinding& Value)
 		{
 			UE::MVVM::FMVVMFieldVariant Variant = UE::MVVM::BindingHelper::FindFieldByName(Class.Get(), Value.GetBindingName());
diff --git a/Engine/Plugins/Runtime/ModelViewViewModel/Source/ModelViewViewModelEditor/Private/Widgets/SMVVMManageViewModelsWidget.cpp b/Engine/Plugins/Runtime/ModelViewViewModel/Source/ModelViewViewModelEditor/Private/Widgets/SMVVMManageViewModelsWidget.cpp
--- a/Engine/Plugins/Runtime/ModelViewViewModel/Source/ModelViewViewModelEditor/Private/Widgets/SMVVMManageViewModelsWidget.cpp
+++ b/Engine/Plugins/Runtime/ModelViewViewModel/Source/ModelViewViewModelEditor/Private/Widgets/SMVVMManageViewModelsWidget.cpp
@@ -285,7 +285,7 @@ void SMVVMManageViewModelsWidget::Construct(const FArguments& InArgs)
 						SNew(SBorder)
 						.BorderImage(FStyleDefaults::GetNoBrush())
 						[
-							SAssignNew(ViewModelBindingListWidget, UE::MVVM::SViewModelBindingListWidget)
+							SAssignNew(ViewModelBindingListWidget, UE::MVVM::SViewModelBindingListWidget, nullptr)
 						]
 					]
 				]
diff --git a/Engine/Plugins/Runtime/ModelViewViewModel/Source/ModelViewViewModelEditor/Private/Widgets/SMVVMSelectViewModel.cpp b/Engine/Plugins/Runtime/ModelViewViewModel/Source/ModelViewViewModelEditor/Private/Widgets/SMVVMSelectViewModel.cpp
--- a/Engine/Plugins/Runtime/ModelViewViewModel/Source/ModelViewViewModelEditor/Private/Widgets/SMVVMSelectViewModel.cpp
+++ b/Engine/Plugins/Runtime/ModelViewViewModel/Source/ModelViewViewModelEditor/Private/Widgets/SMVVMSelectViewModel.cpp
@@ -3,40 +3,24 @@
 #include "Widgets/SMVVMSelectViewModel.h"
 
 #include "Bindings/MVVMBindingHelper.h"
-#include "Blueprint/UserWidget.h"
+#include "Components/Widget.h"
 #include "Editor.h"
+#include "FieldNotification/IFieldValueChanged.h"
 #include "Framework/MultiBox/MultiBoxBuilder.h"
-#include "Kismet2/KismetEditorUtilities.h"
+#include "Modules/ModuleManager.h"
 #include "MVVMBlueprintViewModelContext.h"
 #include "MVVMEditorSubsystem.h"
 #include "MVVMViewModelBase.h"
-#include "SMVVMViewModelContextListWidget.h"
-#include "SMVVMViewModelBindingListWidget.h"
 #include "SPrimaryButton.h"
-#include "SSimpleButton.h"
 #include "Styling/MVVMEditorStyle.h"
-#include "Styling/SlateIconFinder.h"
-#include "Textures/SlateIcon.h"
-#include "Types/MVVMFieldVariant.h"
-#include "UMGStyle.h"
-#include "UObject/UObjectIterator.h"
-#include "Widgets/Images/SImage.h"
 #include "Widgets/Input/SButton.h"
-#include "Widgets/Input/SComboButton.h"
-#include "Widgets/Input/SEditableTextBox.h"
-#include "Widgets/Input/SSearchBox.h"
 #include "Widgets/Layout/SSplitter.h"
 #include "Widgets/Layout/SUniformGridPanel.h"
-#include "Widgets/SWindow.h"
-#include "Widgets/Views/SListView.h"
-#include "Widgets/Views/STreeView.h"
+#include "Widgets/SMVVMViewModelBindingListWidget.h"
 
-//ClassViewer
 #include "ClassViewerModule.h"
 #include "ClassViewerFilter.h"
 #include "SClassViewer.h"
-#include "Components/Widget.h"
-#include "Modules/ModuleManager.h"
 
 #define LOCTEXT_NAMESPACE "SMVVMSelectViewModel"
 
@@ -70,7 +54,7 @@ namespace Private
 	};
 } //namespace
 
-void SMVVMSelectViewModel::Construct(const FArguments& InArgs)
+void SMVVMSelectViewModel::Construct(const FArguments& InArgs, const UWidgetBlueprint* WidgetBlueprint)
 {
 	OnCancel = InArgs._OnCancel;
 	OnViewModelCommitted = InArgs._OnViewModelCommitted;
@@ -125,7 +109,7 @@ void SMVVMSelectViewModel::Construct(const FArguments& InArgs)
 					SNew(SBorder)
 					.BorderImage(FStyleDefaults::GetNoBrush())
 					[
-						SAssignNew(ViewModelBindingListWidget, SViewModelBindingListWidget)
+						SAssignNew(ViewModelBindingListWidget, SViewModelBindingListWidget, WidgetBlueprint)
 					]
 				]
 			]
diff --git a/Engine/Plugins/Runtime/ModelViewViewModel/Source/ModelViewViewModelEditor/Private/Widgets/SMVVMSelectViewModel.h b/Engine/Plugins/Runtime/ModelViewViewModel/Source/ModelViewViewModelEditor/Private/Widgets/SMVVMSelectViewModel.h
--- a/Engine/Plugins/Runtime/ModelViewViewModel/Source/ModelViewViewModelEditor/Private/Widgets/SMVVMSelectViewModel.h
+++ b/Engine/Plugins/Runtime/ModelViewViewModel/Source/ModelViewViewModelEditor/Private/Widgets/SMVVMSelectViewModel.h
@@ -6,6 +6,7 @@
 #include "Widgets/SCompoundWidget.h"
 
 class SClassViewer;
+class UWidgetBlueprint;
 
 namespace UE::MVVM
 {
@@ -22,7 +23,7 @@ public:
 	SLATE_END_ARGS()
 
 public:
-	void Construct(const FArguments& InArgs);
+	void Construct(const FArguments& InArgs, const UWidgetBlueprint* WidgetBlueprint);
 
 private:
 	void HandleClassPicked(UClass* ClassPicked);
diff --git a/Engine/Plugins/Runtime/ModelViewViewModel/Source/ModelViewViewModelEditor/Private/Widgets/SMVVMViewBindingListView.cpp b/Engine/Plugins/Runtime/ModelViewViewModel/Source/ModelViewViewModelEditor/Private/Widgets/SMVVMViewBindingListView.cpp
--- a/Engine/Plugins/Runtime/ModelViewViewModel/Source/ModelViewViewModelEditor/Private/Widgets/SMVVMViewBindingListView.cpp
+++ b/Engine/Plugins/Runtime/ModelViewViewModel/Source/ModelViewViewModelEditor/Private/Widgets/SMVVMViewBindingListView.cpp
@@ -415,7 +415,7 @@ private:
 			if (const FMVVMBlueprintViewModelContext* ViewModel = BlueprintViewPtr->FindViewModel(ViewBinding->ViewModelPath.GetViewModelId()))
 			{
 				UMVVMSubsystem* Subsystem = GEngine->GetEngineSubsystem<UMVVMSubsystem>();
-				TArray<FMVVMAvailableBinding> AvailableBindings = Subsystem->GetAvailableBindings(ViewModel->GetViewModelClass());
+				TArray<FMVVMAvailableBinding> AvailableBindings = Subsystem->GetAvailableBindings(ViewModel->GetViewModelClass(), WidgetBlueprint->GeneratedClass);
 
 				Algo::Transform(AvailableBindings, AvailablePaths, [ViewModel](const FMVVMAvailableBinding& Binding)
 					{
@@ -502,7 +502,7 @@ private:
 			if (Class != nullptr)
 			{
 				UMVVMSubsystem* Subsystem = GEngine->GetEngineSubsystem<UMVVMSubsystem>();
-				TArray<FMVVMAvailableBinding> AvailableBindings = Subsystem->GetWidgetAvailableBindings(Class);
+				TArray<FMVVMAvailableBinding> AvailableBindings = Subsystem->GetAvailableBindings(Class, WidgetBlueprint->GeneratedClass);
 
 				Algo::Transform(AvailableBindings, AvailablePaths, [Class, WidgetName](const FMVVMAvailableBinding& Binding)
 					{
diff --git a/Engine/Plugins/Runtime/ModelViewViewModel/Source/ModelViewViewModelEditor/Private/Widgets/SMVVMViewModelBindingListWidget.cpp b/Engine/Plugins/Runtime/ModelViewViewModel/Source/ModelViewViewModelEditor/Private/Widgets/SMVVMViewModelBindingListWidget.cpp
--- a/Engine/Plugins/Runtime/ModelViewViewModel/Source/ModelViewViewModelEditor/Private/Widgets/SMVVMViewModelBindingListWidget.cpp
+++ b/Engine/Plugins/Runtime/ModelViewViewModel/Source/ModelViewViewModelEditor/Private/Widgets/SMVVMViewModelBindingListWidget.cpp
@@ -11,7 +11,7 @@
 #include "MVVMViewModelBase.h"
 #include "Types/MVVMAvailableBinding.h"
 #include "Types/MVVMFieldVariant.h"
-
+#include "WidgetBlueprint.h"
 #include "Widgets/PropertyViewer/SPropertyViewer.h"
 
 #define LOCTEXT_NAMESPACE "SViewModelBindingListWidget"
@@ -20,14 +20,19 @@ namespace UE::MVVM
 {
 
 /** */
+FFieldIterator_ViewModel::FFieldIterator_ViewModel(const UWidgetBlueprint* InWidgetBlueprint)
+	: WidgetBlueprint(InWidgetBlueprint)
+{ }
+
+
 TArray<FFieldVariant> FFieldIterator_ViewModel::GetFields(const UStruct* Struct) const
 {
 	TArray<FFieldVariant> Result;
-	if (const UClass* Class = Cast<const UClass>(Struct))
+
+	auto AddResult = [&Result, Struct](const TArray<FMVVMAvailableBinding>& AvailableBindingsList)
 	{
-		TArray<FMVVMAvailableBinding> ViewModelAvailableBindingsList = GEngine->GetEngineSubsystem<UMVVMSubsystem>()->GetAvailableBindings(const_cast<UClass*>(Class));
-		Result.Reserve(ViewModelAvailableBindingsList.Num());
-		for (const FMVVMAvailableBinding& Value : ViewModelAvailableBindingsList)
+		Result.Reserve(AvailableBindingsList.Num());
+		for (const FMVVMAvailableBinding& Value : AvailableBindingsList)
 		{
 			FMVVMFieldVariant FieldVariant = BindingHelper::FindFieldByName(Struct, Value.GetBindingName());
 			if (FieldVariant.IsFunction())
@@ -38,18 +43,20 @@ TArray<FFieldVariant> FFieldIterator_ViewModel::GetFields(const UStruct* Struct)
 			{
 				Result.Add(FFieldVariant(FieldVariant.GetProperty()));
 			}
-		}		
+		}
+	};
+
+
+	if (const UClass* Class = Cast<const UClass>(Struct))
+	{
+		const UWidgetBlueprint* WidgetBlueprintPtr = WidgetBlueprint.Get();
+		TSubclassOf<UObject> AccessorClass = WidgetBlueprintPtr ? WidgetBlueprintPtr->GeneratedClass : nullptr;
+		AddResult(GEngine->GetEngineSubsystem<UMVVMSubsystem>()->GetAvailableBindings(const_cast<UClass*>(Class), AccessorClass));
+
 	}
-	else if (Cast<const UScriptStruct>(Struct))
+	else if (const UScriptStruct* ScriptStruct = Cast<const UScriptStruct>(Struct))
 	{
-		for (TFieldIterator<FProperty> PropertyIt(Struct, EFieldIteratorFlags::IncludeSuper); PropertyIt; ++PropertyIt)
-		{
-			FProperty* Property = *PropertyIt;
-			if (BindingHelper::IsValidForSourceBinding(Property) || BindingHelper::IsValidForDestinationBinding(Property))
-			{
-				Result.Add(FFieldVariant(Property));
-			}
-		}
+		AddResult(GEngine->GetEngineSubsystem<UMVVMSubsystem>()->GetAvailableBindingsForStruct(ScriptStruct));
 	}
 
 	Result.Sort([](const FFieldVariant& A, const FFieldVariant& B)
@@ -89,13 +96,15 @@ TArray<FFieldVariant> FFieldIterator_ViewModel::GetFields(const UStruct* Struct)
 
 
 /** */
-void SViewModelBindingListWidget::Construct(const FArguments& InArgs)
+void SViewModelBindingListWidget::Construct(const FArguments& InArgs, const UWidgetBlueprint* WidgetBlueprint)
 {
+	ViewModelFieldIterator = MakeUnique<FFieldIterator_ViewModel>(WidgetBlueprint);
+
 	PropertyViewer = SNew(UE::PropertyViewer::SPropertyViewer)
 		.PropertyVisibility(UE::PropertyViewer::SPropertyViewer::EPropertyVisibility::Hidden)
 		.bShowFieldIcon(true)
 		.bSanitizeName(true)
-		.FieldIterator(&ViewModelFieldIterator);
+		.FieldIterator(ViewModelFieldIterator.Get());
 
 	if (InArgs._ViewModel.Class.Get())
 	{
diff --git a/Engine/Plugins/Runtime/ModelViewViewModel/Source/ModelViewViewModelEditor/Private/Widgets/SMVVMViewModelBindingListWidget.h b/Engine/Plugins/Runtime/ModelViewViewModel/Source/ModelViewViewModelEditor/Private/Widgets/SMVVMViewModelBindingListWidget.h
--- a/Engine/Plugins/Runtime/ModelViewViewModel/Source/ModelViewViewModelEditor/Private/Widgets/SMVVMViewModelBindingListWidget.h
+++ b/Engine/Plugins/Runtime/ModelViewViewModel/Source/ModelViewViewModelEditor/Private/Widgets/SMVVMViewModelBindingListWidget.h
@@ -5,8 +5,10 @@
 #include "CoreMinimal.h"
 #include "Framework/PropertyViewer/IFieldIterator.h"
 #include "Templates/SubclassOf.h"
+#include "UObject/WeakObjectPtr.h"
 #include "Widgets/SCompoundWidget.h"
 
+class UWidgetBlueprint;
 namespace UE::PropertyViewer
 {
 class SPropertyViewer;
@@ -20,7 +22,12 @@ namespace UE::MVVM
  */
 class FFieldIterator_ViewModel : public UE::PropertyViewer::FFieldIterator_BlueprintVisible
 {
+public:
+	FFieldIterator_ViewModel(const UWidgetBlueprint* WidgetBlueprint);
 	virtual TArray<FFieldVariant> GetFields(const UStruct*) const override;
+
+private:
+	TWeakObjectPtr<const UWidgetBlueprint> WidgetBlueprint;
 };
 
 
@@ -40,14 +47,14 @@ public:
 		SLATE_ARGUMENT(FViewModel, ViewModel)
 	SLATE_END_ARGS()
 
-	void Construct(const FArguments& InArgs);
+	void Construct(const FArguments& InArgs, const UWidgetBlueprint* WidgetBlueprint);
 	void SetViewModel(UClass* Class, FName Name, FGuid Guid);
 	void SetViewModels(TArrayView<const FViewModel> ViewModels);
 
 	void SetRawFilterText(const FText& InFilterText);
 
 private:
-	FFieldIterator_ViewModel ViewModelFieldIterator;
+	TUniquePtr<FFieldIterator_ViewModel> ViewModelFieldIterator;
 	TArray<FViewModel> ViewModels;
 	TSharedPtr<UE::PropertyViewer::SPropertyViewer> PropertyViewer;
 };
diff --git a/Engine/Plugins/Runtime/ModelViewViewModel/Source/ModelViewViewModelEditor/Private/Widgets/SMVVMViewModelPanel.cpp b/Engine/Plugins/Runtime/ModelViewViewModel/Source/ModelViewViewModelEditor/Private/Widgets/SMVVMViewModelPanel.cpp
--- a/Engine/Plugins/Runtime/ModelViewViewModel/Source/ModelViewViewModelEditor/Private/Widgets/SMVVMViewModelPanel.cpp
+++ b/Engine/Plugins/Runtime/ModelViewViewModel/Source/ModelViewViewModelEditor/Private/Widgets/SMVVMViewModelPanel.cpp
@@ -41,6 +41,7 @@ void SMVVMViewModelPanel::Construct(const FArguments& InArgs, TSharedPtr<FWidget
 
 	WeakBlueprintEditor = WidgetBlueprintEditor;
 	WeakBlueprintView = CurrentBlueprintView;
+	ViewModelFieldIterator = MakeUnique<FFieldIterator_ViewModel>(WidgetBlueprint);
 
 	if (CurrentBlueprintView)
 	{
@@ -53,7 +54,7 @@ void SMVVMViewModelPanel::Construct(const FArguments& InArgs, TSharedPtr<FWidget
 		.bShowSearchBox(true)
 		.bShowFieldIcon(true)
 		.bSanitizeName(true)
-		.FieldIterator(&ViewModelFieldIterator)
+		.FieldIterator(ViewModelFieldIterator.Get())
 		.SearchBoxPreSlot()
 		[
 			SAssignNew(AddMenuButton, SPositiveActionButton)
@@ -168,11 +169,16 @@ void SMVVMViewModelPanel::HandleViewModelsUpdated()
 
 TSharedRef<SWidget> SMVVMViewModelPanel::MakeAddMenu()
 {
+	const UWidgetBlueprint* WidgetBlueprint = nullptr;
+	if (TSharedPtr<FWidgetBlueprintEditor> EditorPin = WeakBlueprintEditor.Pin())
+	{
+		WidgetBlueprint = EditorPin->GetWidgetBlueprintObj();
+	}
 	return SNew(SBox)
 		.WidthOverride(600)
 		.HeightOverride(500)
 		[
-			SNew(SMVVMSelectViewModel)
+			SNew(SMVVMSelectViewModel, WidgetBlueprint)
 			.OnCancel(this, &SMVVMViewModelPanel::HandleCancelAddMenu)
 			.OnViewModelCommitted(this, &SMVVMViewModelPanel::HandleAddMenuViewModel)
 		];
diff --git a/Engine/Plugins/Runtime/ModelViewViewModel/Source/ModelViewViewModelEditor/Private/Widgets/SMVVMViewModelPanel.h b/Engine/Plugins/Runtime/ModelViewViewModel/Source/ModelViewViewModelEditor/Private/Widgets/SMVVMViewModelPanel.h
--- a/Engine/Plugins/Runtime/ModelViewViewModel/Source/ModelViewViewModelEditor/Private/Widgets/SMVVMViewModelPanel.h
+++ b/Engine/Plugins/Runtime/ModelViewViewModel/Source/ModelViewViewModelEditor/Private/Widgets/SMVVMViewModelPanel.h
@@ -42,13 +42,13 @@ private:
 private:
 	TSharedPtr<SPositiveActionButton> AddMenuButton;
 	TSharedPtr<UE::PropertyViewer::SPropertyViewer> ViewModelTreeView;
-	FFieldIterator_ViewModel ViewModelFieldIterator;
 
 	using FViewModelHandle = TPair<FGuid, UE::PropertyViewer::SPropertyViewer::FHandle>;
 	TArray<FViewModelHandle> ViewModelHandles;
 
 	TWeakPtr<FWidgetBlueprintEditor> WeakBlueprintEditor;
 	TWeakObjectPtr<UMVVMBlueprintView> WeakBlueprintView;
+	TUniquePtr<FFieldIterator_ViewModel> ViewModelFieldIterator;
 	FDelegateHandle ViewModelsUpdatedHandle;
 };
 
diff --git a/Engine/Plugins/Runtime/ModelViewViewModel/Source/ModelViewViewModelEditor/Public/MVVMEditorSubsystem.h b/Engine/Plugins/Runtime/ModelViewViewModel/Source/ModelViewViewModelEditor/Public/MVVMEditorSubsystem.h
--- a/Engine/Plugins/Runtime/ModelViewViewModel/Source/ModelViewViewModelEditor/Public/MVVMEditorSubsystem.h
+++ b/Engine/Plugins/Runtime/ModelViewViewModel/Source/ModelViewViewModelEditor/Public/MVVMEditorSubsystem.h
@@ -48,7 +48,7 @@ public:
 	void RemoveBinding(UWidgetBlueprint* WidgetBlueprint, const FMVVMBlueprintViewBinding& Binding);
 
 	UFUNCTION(BlueprintCallable, Category = "MVVM")
-	TArray<FMVVMAvailableBinding> GetChildViewModels(TSubclassOf<UObject> Class);
+	TArray<FMVVMAvailableBinding> GetChildViewModels(TSubclassOf<UObject> Class, TSubclassOf<UObject> Accessor);
 
 	void SetSourceToDestinationConversionFunction(UWidgetBlueprint* WidgetBlueprint, FMVVMBlueprintViewBinding& Binding, const UFunction* ConversionFunction);
 	void SetDestinationToSourceConversionFunction(UWidgetBlueprint* WidgetBlueprint, FMVVMBlueprintViewBinding& Binding, const UFunction* ConversionFunction);
