FMVVMFieldVariant can now be initialized from FFieldVariant, though if the variant isn't an FProperty or UFunction then the variant will be empty.

diff --git a/Engine/Plugins/Runtime/ModelViewViewModel/Source/ModelViewViewModel/Public/Types/MVVMFieldVariant.h b/Engine/Plugins/Runtime/ModelViewViewModel/Source/ModelViewViewModel/Public/Types/MVVMFieldVariant.h
--- a/Engine/Plugins/Runtime/ModelViewViewModel/Source/ModelViewViewModel/Public/Types/MVVMFieldVariant.h
+++ b/Engine/Plugins/Runtime/ModelViewViewModel/Source/ModelViewViewModel/Public/Types/MVVMFieldVariant.h
@@ -28,6 +28,18 @@
 	public:
 		TMVVMFieldVariant() = default;
 
+		explicit TMVVMFieldVariant(FFieldVariant InVariant)
+		{
+			if (PropertyType* Property = InVariant.Get<PropertyType>())
+			{
+				Binding = VariantType(TInPlaceType<PropertyType*>(), Property);
+			}
+			else if (FunctionType* Function = InVariant.Get<FunctionType>())
+			{
+				Binding = VariantType(TInPlaceType<FunctionType*>(), Function);
+			}
+		}
+
 		explicit TMVVMFieldVariant(PropertyType* InValue)
 		{
 			Binding = VariantType(TInPlaceType<PropertyType*>(), InValue);
