Moved SMVVMFieldEntry to its own file.

#rnx

diff --git a/Engine/Plugins/Runtime/ModelViewViewModel/Source/ModelViewViewModelEditor/Private/Widgets/SMVVMFieldEntry.cpp b/Engine/Plugins/Runtime/ModelViewViewModel/Source/ModelViewViewModelEditor/Private/Widgets/SMVVMFieldEntry.cpp
new file mode 100644
--- /dev/null
+++ b/Engine/Plugins/Runtime/ModelViewViewModel/Source/ModelViewViewModelEditor/Private/Widgets/SMVVMFieldEntry.cpp
@@ -0,0 +1,102 @@
+
+#include "SMVVMFieldEntry.h"
+#include "SMVVMFieldIcon.h"
+#include "Widgets/SBoxPanel.h"
+
+#define LOCTEXT_NAMESPACE "MVVMFieldEntry"
+
+using namespace UE::MVVM;
+
+namespace UE::MVVM::Private
+{
+	FText GetFieldDisplayName(const FMVVMConstFieldVariant& Field)
+	{
+		if (Field.IsProperty())
+		{
+			return Field.GetProperty()->GetDisplayNameText();
+		}
+		else if (Field.IsFunction())
+		{
+			return Field.GetFunction()->GetDisplayNameText();
+		}
+		return LOCTEXT("None", "<None>");
+	}
+}
+
+void SMVVMFieldEntry::Construct(const FArguments& InArgs)
+{
+	Field = InArgs._Field;
+	OnValidateField = InArgs._OnValidateField;
+
+	ChildSlot
+	[
+		SNew(SHorizontalBox)
+		+ SHorizontalBox::Slot()
+		.HAlign(HAlign_Left)
+		.VAlign(VAlign_Center)
+		.AutoWidth()
+		[
+			SAssignNew(Icon, SMVVMFieldIcon)
+		]
+		+ SHorizontalBox::Slot()
+		.HAlign(HAlign_Fill)
+		.VAlign(VAlign_Center)
+		.Padding(4, 0, 0, 0)
+		[
+			SAssignNew(Label, STextBlock)
+			.TextStyle(InArgs._TextStyle)
+		]
+	];
+
+	Refresh();
+}
+
+void SMVVMFieldEntry::Refresh()
+{
+	FText ToolTipText = FText::GetEmpty();
+	bool bEnabled = true;
+	if (OnValidateField.IsBound())
+	{
+		TValueOrError<bool, FString> Result = OnValidateField.Execute(Field);
+		if (Result.HasError())
+		{
+			ToolTipText = FText::FromString(Result.GetError());
+			bEnabled = false;
+		}
+		else
+		{
+			bEnabled = true;
+		}
+	}
+
+	SetEnabled(bEnabled);
+
+	UE::MVVM::FMVVMConstFieldVariant Variant;
+
+	TArray<UE::MVVM::FMVVMConstFieldVariant> Fields = Field.GetFields();
+	if (Fields.Num() > 0)
+	{
+		Variant = Fields.Last();
+	}
+
+	if (ToolTipText.IsEmpty())
+	{
+		ToolTipText = Variant.IsFunction() ? Variant.GetFunction()->GetToolTipText() :
+			Variant.IsProperty() ? Variant.GetProperty()->GetToolTipText() :
+			FText::GetEmpty();
+	}
+
+	SetToolTipText(ToolTipText);
+
+	Icon->RefreshBinding(Variant);
+	Label->SetText(Private::GetFieldDisplayName(Variant));
+}
+
+void SMVVMFieldEntry::SetField(const FMVVMBlueprintPropertyPath& InField)
+{
+	Field = InField;
+
+	Refresh();
+}
+
+#undef LOCTEXT_NAMESPACE
diff --git a/Engine/Plugins/Runtime/ModelViewViewModel/Source/ModelViewViewModelEditor/Private/Widgets/SMVVMFieldEntry.h b/Engine/Plugins/Runtime/ModelViewViewModel/Source/ModelViewViewModelEditor/Private/Widgets/SMVVMFieldEntry.h
new file mode 100644
--- /dev/null
+++ b/Engine/Plugins/Runtime/ModelViewViewModel/Source/ModelViewViewModelEditor/Private/Widgets/SMVVMFieldEntry.h
@@ -0,0 +1,43 @@
+// Copyright Epic Games, Inc. All Rights Reserved.
+
+#pragma once
+
+#include "Delegates/Delegate.h"
+#include "MVVMPropertyPath.h"
+#include "Styling/CoreStyle.h"
+#include "Templates/ValueOrError.h"
+#include "UObject/UnrealType.h"
+#include "Widgets/DeclarativeSyntaxSupport.h"
+#include "Widgets/SCompoundWidget.h"
+#include "Widgets/Text/STextBlock.h"
+
+class SMVVMFieldIcon;
+class STextBlock;
+
+namespace UE::MVVM
+{
+	using FIsFieldValidResult = TValueOrError<bool, FString>;
+	DECLARE_DELEGATE_RetVal_OneParam(FIsFieldValidResult, FIsFieldValid, FMVVMBlueprintPropertyPath);
+}
+
+class SMVVMFieldEntry : public SCompoundWidget
+{
+public:
+	SLATE_BEGIN_ARGS(SMVVMFieldEntry) :
+		_TextStyle(&FCoreStyle::Get().GetWidgetStyle<FTextBlockStyle>("NormalText"))
+	{}
+		SLATE_STYLE_ARGUMENT(FTextBlockStyle, TextStyle)
+		SLATE_ARGUMENT(FMVVMBlueprintPropertyPath, Field)
+		SLATE_EVENT(UE::MVVM::FIsFieldValid, OnValidateField)
+	SLATE_END_ARGS()
+
+	void Construct(const FArguments& InArgs);
+	void Refresh();
+	void SetField(const FMVVMBlueprintPropertyPath& InField);
+
+private:
+	FMVVMBlueprintPropertyPath Field;
+	UE::MVVM::FIsFieldValid OnValidateField;
+	TSharedPtr<SMVVMFieldIcon> Icon;
+	TSharedPtr<STextBlock> Label;
+};
diff --git a/Engine/Plugins/Runtime/ModelViewViewModel/Source/ModelViewViewModelEditor/Private/Widgets/SMVVMFieldSelector.cpp b/Engine/Plugins/Runtime/ModelViewViewModel/Source/ModelViewViewModelEditor/Private/Widgets/SMVVMFieldSelector.cpp
--- a/Engine/Plugins/Runtime/ModelViewViewModel/Source/ModelViewViewModelEditor/Private/Widgets/SMVVMFieldSelector.cpp
+++ b/Engine/Plugins/Runtime/ModelViewViewModel/Source/ModelViewViewModelEditor/Private/Widgets/SMVVMFieldSelector.cpp
@@ -13,99 +13,6 @@
 
 using namespace UE::MVVM;
 
-namespace UE::MVVM::Private
-{
-	FText GetFieldDisplayName(const FMVVMConstFieldVariant& Field)
-	{
-		if (Field.IsProperty())
-		{
-			return Field.GetProperty()->GetDisplayNameText();
-		}
-		else if (Field.IsFunction())
-		{
-			return Field.GetFunction()->GetDisplayNameText();
-		}
-		return LOCTEXT("None", "<None>");
-	}
-}
-
-
-void SMVVMFieldEntry::Construct(const FArguments& InArgs)
-{
-	Field = InArgs._Field;
-	OnValidateField = InArgs._OnValidateField;
-
-	ChildSlot
-		[
-			SNew(SHorizontalBox)
-			+ SHorizontalBox::Slot()
-		.HAlign(HAlign_Left)
-		.VAlign(VAlign_Center)
-		.AutoWidth()
-		[
-			SAssignNew(Icon, SMVVMFieldIcon)
-		]
-	+ SHorizontalBox::Slot()
-		.HAlign(HAlign_Fill)
-		.VAlign(VAlign_Center)
-		.Padding(4, 0, 0, 0)
-		[
-			SAssignNew(Label, STextBlock)
-			.TextStyle(InArgs._TextStyle)
-		]
-		];
-
-	Refresh();
-}
-
-void SMVVMFieldEntry::Refresh()
-{
-	FText ToolTipText = FText::GetEmpty();
-	bool bEnabled = true;
-	if (OnValidateField.IsBound())
-	{
-		TValueOrError<bool, FString> Result = OnValidateField.Execute(Field);
-		if (Result.HasError())
-		{
-			ToolTipText = FText::FromString(Result.GetError());
-			bEnabled = false;
-		}
-		else
-		{
-			bEnabled = true;
-		}
-	}
-
-	SetEnabled(bEnabled);
-
-	UE::MVVM::FMVVMConstFieldVariant Variant;
-
-	TArray<UE::MVVM::FMVVMConstFieldVariant> Fields = Field.GetFields();
-	if (Fields.Num() > 0)
-	{
-		Variant = Fields.Last();
-	}
-
-	if (ToolTipText.IsEmpty())
-	{
-		ToolTipText = Variant.IsFunction() ? Variant.GetFunction()->GetToolTipText() :
-			Variant.IsProperty() ? Variant.GetProperty()->GetToolTipText() :
-			FText::GetEmpty();
-	}
-
-	SetToolTipText(ToolTipText);
-
-	Icon->RefreshBinding(Variant);
-	Label->SetText(Private::GetFieldDisplayName(Variant));
-}
-
-void SMVVMFieldEntry::SetField(const FMVVMBlueprintPropertyPath& InField)
-{
-	Field = InField;
-
-	Refresh();
-}
-
 void SMVVMFieldSelector::Construct(const FArguments& InArgs)
 {
 	OnValidateFieldDelegate = InArgs._OnValidateField;
diff --git a/Engine/Plugins/Runtime/ModelViewViewModel/Source/ModelViewViewModelEditor/Private/Widgets/SMVVMFieldSelector.h b/Engine/Plugins/Runtime/ModelViewViewModel/Source/ModelViewViewModelEditor/Private/Widgets/SMVVMFieldSelector.h
--- a/Engine/Plugins/Runtime/ModelViewViewModel/Source/ModelViewViewModelEditor/Private/Widgets/SMVVMFieldSelector.h
+++ b/Engine/Plugins/Runtime/ModelViewViewModel/Source/ModelViewViewModelEditor/Private/Widgets/SMVVMFieldSelector.h
@@ -10,8 +10,9 @@
 #include "Types/MVVMFieldVariant.h"
 #include "UObject/UnrealType.h"
 #include "Widgets/DeclarativeSyntaxSupport.h"
-#include "Widgets/SCompoundWidget.h"
 #include "Widgets/Input/SComboBox.h"
+#include "Widgets/SCompoundWidget.h"
+#include "Widgets/SMVVMFieldEntry.h"
 
 namespace UE::MVVM
 {
@@ -21,32 +22,6 @@ namespace UE::MVVM
 class SMVVMFieldIcon;
 class STextBlock; 
 
-using FIsFieldValidResult = TValueOrError<bool, FString>;
-DECLARE_DELEGATE_RetVal_OneParam(FIsFieldValidResult, FIsFieldValid, FMVVMBlueprintPropertyPath);
-
-class SMVVMFieldEntry : public SCompoundWidget
-{
-public:
-
-	SLATE_BEGIN_ARGS(SMVVMFieldEntry) :
-		_TextStyle( &FCoreStyle::Get().GetWidgetStyle<FTextBlockStyle>( "NormalText" ) )
-		{}
-		SLATE_STYLE_ARGUMENT(FTextBlockStyle, TextStyle)
-		SLATE_ARGUMENT(FMVVMBlueprintPropertyPath, Field)
-		SLATE_EVENT(FIsFieldValid, OnValidateField)
-	SLATE_END_ARGS()
-
-	void Construct(const FArguments& InArgs);
-	void Refresh(); 
-	void SetField(const FMVVMBlueprintPropertyPath& InField);
-
-private:
-	FMVVMBlueprintPropertyPath Field;
-	FIsFieldValid OnValidateField;
-	TSharedPtr<SMVVMFieldIcon> Icon;
-	TSharedPtr<STextBlock> Label;
-};
-
 class SMVVMFieldSelector : public SCompoundWidget
 {
 public:
@@ -69,7 +44,7 @@ public:
 		  * Would the given field be a valid entry for this combo box? 
 		  * The error string returned will be used as a tooltip.
 		  */
-		SLATE_EVENT(FIsFieldValid, OnValidateField)
+		SLATE_EVENT(UE::MVVM::FIsFieldValid, OnValidateField)
 	SLATE_END_ARGS()
 
 	void Construct(const FArguments& InArgs);
@@ -97,7 +72,7 @@ private:
 	TSharedPtr<SComboBox<FMVVMBlueprintPropertyPath>> FieldComboBox;
 
 	FSelectionChanged OnSelectionChangedDelegate;
-	FIsFieldValid OnValidateFieldDelegate;
+	UE::MVVM::FIsFieldValid OnValidateFieldDelegate;
 
 	TArray<FMVVMBlueprintPropertyPath> CachedAvailableFields;
 }; 
