MVVM - Field selector now uses SSourceBindingList to display the available fields.

SSourceBindingList can now display both widget and viewmodel fields.

Moved several widgets to the UE::MVVM namespace - SMVVMFieldSelector most notably.

#rnx

diff --git a/Engine/Plugins/Runtime/ModelViewViewModel/Source/ModelViewViewModelEditor/Private/Customizations/MVVMConversionPathCustomization.cpp b/Engine/Plugins/Runtime/ModelViewViewModel/Source/ModelViewViewModelEditor/Private/Customizations/MVVMConversionPathCustomization.cpp
--- a/Engine/Plugins/Runtime/ModelViewViewModel/Source/ModelViewViewModelEditor/Private/Customizations/MVVMConversionPathCustomization.cpp
+++ b/Engine/Plugins/Runtime/ModelViewViewModel/Source/ModelViewViewModelEditor/Private/Customizations/MVVMConversionPathCustomization.cpp
@@ -33,7 +33,6 @@
 #include "Widgets/Layout/SBox.h"
 #include "Widgets/SMVVMConversionPath.h"
 #include "Widgets/SMVVMFieldSelector.h"
-#include "Widgets/SMVVMSourceSelector.h"
 #include "Widgets/Text/STextBlock.h"
 
 #define LOCTEXT_NAMESPACE "MVVMConversionPath"
@@ -238,7 +237,7 @@ namespace UE::MVVM
 
 			ValueWidget->SetVisibility(TAttribute<EVisibility>::CreateSP(this, &FConversionPathCustomization::GetArgumentWidgetVisibility, ArgumentName, bSourceToDestination, true));
 
-			TSharedPtr<SMVVMFieldSelector> FieldSelector;
+			TSharedPtr<UE::MVVM::SFieldSelector> FieldSelector;
 
 			ChildBuilder.AddCustomRow(Pin->GetDisplayName())
 			.NameContent()
@@ -263,24 +262,11 @@ namespace UE::MVVM
 					]
 					+ SOverlay::Slot()
 					[
-						SNew(SHorizontalBox)
+						SAssignNew(FieldSelector, UE::MVVM::SFieldSelector, WidgetBlueprint, bSourceToDestination)
 						.Visibility(this, &FConversionPathCustomization::GetArgumentWidgetVisibility, ArgumentName, bSourceToDestination, false)
-						+ SHorizontalBox::Slot()
-						[
-							SNew(SMVVMSourceSelector)
-							.SelectedSource(this, &FConversionPathCustomization::OnGetSelectedSource, ArgumentName, bSourceToDestination)
-							.AvailableSources(this, &FConversionPathCustomization::OnGetAvailableSources, bSourceToDestination)
-							.OnSelectionChanged(this, &FConversionPathCustomization::OnSetSource, ArgumentName, bSourceToDestination)
-						]
-						+ SHorizontalBox::Slot()
-						[
-							SAssignNew(FieldSelector, SMVVMFieldSelector)
-							.IsSource(bSourceToDestination)
-							.SelectedSource(this, &FConversionPathCustomization::OnGetSelectedSource, ArgumentName, bSourceToDestination)
-							.AvailableFields(this, &FConversionPathCustomization::OnGetAvailableFields, ArgumentName, bSourceToDestination)
-							.SelectedField(this, &FConversionPathCustomization::OnGetSelectedField, ArgumentName, bSourceToDestination)
-							.OnSelectionChanged(this, &FConversionPathCustomization::OnSetProperty, ArgumentName, bSourceToDestination)
-						]
+						.SelectedField(this, &FConversionPathCustomization::OnGetSelectedField, ArgumentName, bSourceToDestination)
+						.BindingMode(this, &FConversionPathCustomization::GetBindingMode)
+						.OnSelectionChanged(this, &FConversionPathCustomization::OnSetProperty, ArgumentName, bSourceToDestination)
 					]
 				]
 				+ SHorizontalBox::Slot()
@@ -314,86 +300,6 @@ namespace UE::MVVM
 		}
 	}
 
-	TArray<FBindingSource> FConversionPathCustomization::OnGetAvailableSources(bool bSourceToDestination) const
-	{
-		UMVVMEditorSubsystem* Subsystem = GEditor->GetEditorSubsystem<UMVVMEditorSubsystem>();
-		if (bSourceToDestination)
-		{
-			return Subsystem->GetAllViewModels(WidgetBlueprint);
-		}
-		else
-		{
-			return Subsystem->GetBindableWidgets(WidgetBlueprint);
-		}
-	}
-
-	FBindingSource FConversionPathCustomization::OnGetSelectedSource(FName ArgumentName, bool bSourceToDestination) const
-	{
-		UMVVMEditorSubsystem* Subsystem = GEditor->GetEditorSubsystem<UMVVMEditorSubsystem>();
-
-		TArray<void*> RawBindings;
-		ParentHandle->AccessRawData(RawBindings);
-
-		FMVVMBlueprintPropertyPath Path;
-		bool bFirst = true;
-
-		for (void* RawBinding : RawBindings)
-		{
-			FMVVMBlueprintViewBinding* Binding = reinterpret_cast<FMVVMBlueprintViewBinding*>(RawBinding);
-			FMVVMBlueprintPropertyPath ThisPath = Subsystem->GetPathForConversionFunctionArgument(WidgetBlueprint, *Binding, ArgumentName, bSourceToDestination);
-
-			if (bFirst)
-			{
-				Path = ThisPath;
-				bFirst = false;
-			}
-			else if (Path != ThisPath)
-			{
-				Path = FMVVMBlueprintPropertyPath();
-			}
-		}
-		
-		if (Path.IsFromViewModel())
-		{
-			return FBindingSource::CreateForViewModel(WidgetBlueprint, Path.GetViewModelId());
-		}
-		else if (Path.IsFromWidget())
-		{
-			return FBindingSource::CreateForWidget(WidgetBlueprint, Path.GetWidgetName());
-		}
-
-		return FBindingSource();
-	}
-
-	void FConversionPathCustomization::OnSetSource(FBindingSource Source, FName ArgumentName, bool bSourceToDestination)
-	{
-		FMVVMBlueprintPropertyPath Path;
-		if (Source.ViewModelId.IsValid())
-		{
-			Path.SetViewModelId(Source.ViewModelId);
-		}
-		else
-		{
-			Path.SetWidgetName(Source.Name);
-		}
-
-		UMVVMEditorSubsystem* Subsystem = GEditor->GetEditorSubsystem<UMVVMEditorSubsystem>();
-
-		TArray<void*> RawBindings;
-		ParentHandle->AccessRawData(RawBindings);
-
-		for (void* RawBinding : RawBindings)
-		{
-			FMVVMBlueprintViewBinding* Binding = reinterpret_cast<FMVVMBlueprintViewBinding*>(RawBinding);
-			Subsystem->SetPathForConversionFunctionArgument(WidgetBlueprint, *Binding, ArgumentName, Path, bSourceToDestination);
-		}
-
-		for (const TSharedPtr<SMVVMFieldSelector>& FieldSelector : ArgumentFieldSelectors)
-		{
-			FieldSelector->Refresh();
-		}
-	}
-
 	void FConversionPathCustomization::OnSetProperty(FMVVMBlueprintPropertyPath NewSelection, FName ArgumentName, bool bSourceToDestination)
 	{
 		UMVVMEditorSubsystem* EditorSubsystem = GEditor->GetEditorSubsystem<UMVVMEditorSubsystem>();
@@ -437,108 +343,6 @@ namespace UE::MVVM
 		return Path;
 	}
 
-	TArray<FMVVMBlueprintPropertyPath> FConversionPathCustomization::OnGetAvailableFields(FName ArgumentName, bool bSourceToDestination) const
-	{
-		TArray<FMVVMBlueprintPropertyPath> AvailablePaths;
-
-		FBindingSource Source = OnGetSelectedSource(ArgumentName, bSourceToDestination);
-		UClass* SourceClass = nullptr;
-		if (!Source.Name.IsNone())
-		{
-			// widget
-			if (Source.Name.IsEqual(WidgetBlueprint->GetFName()))
-			{
-				SourceClass = WidgetBlueprint->GeneratedClass;
-			}
-			else if (const UWidget* Widget = WidgetBlueprint->WidgetTree->FindWidget(Source.Name))
-			{
-				SourceClass = Widget->GetClass();
-			}
-		}
-		else if (Source.ViewModelId.IsValid())
-		{
-			// viewmodel
-			UMVVMEditorSubsystem* EditorSubsystem = GEditor->GetEditorSubsystem<UMVVMEditorSubsystem>();
-			if (UMVVMBlueprintView* View = EditorSubsystem->GetView(WidgetBlueprint))
-			{
-				if (const FMVVMBlueprintViewModelContext* ViewModel = View->FindViewModel(Source.ViewModelId))
-				{
-					SourceClass = ViewModel->GetViewModelClass();
-				}
-			}
-		}
-
-		if (SourceClass == nullptr)
-		{
-			return AvailablePaths;
-		}
-
-		UMVVMSubsystem* Subsystem = GEngine->GetEngineSubsystem<UMVVMSubsystem>();
-		TArray<FMVVMAvailableBinding> AvailableBindings = Subsystem->GetAvailableBindings(SourceClass, WidgetBlueprint->GeneratedClass);
-
-		const UEdGraphSchema_K2* Schema = GetDefault<UEdGraphSchema_K2>();
-		FEdGraphPinType ArgumentType = GetArgumentPinType(ArgumentName, bSourceToDestination);
-
-		Algo::TransformIf(AvailableBindings, AvailablePaths, 
-			// filter handler
-			[ArgumentType, Schema, SourceClass](const FMVVMAvailableBinding& Binding) -> bool
-			{
-				FName BindingName = Binding.GetBindingName().ToName();
-				
-				if (const FProperty* Property = SourceClass->FindPropertyByName(BindingName))
-				{
-					return true;
-				}
-
-				if (const UFunction* Function = SourceClass->FindFunctionByName(BindingName))
-				{
-					const FProperty* ReturnProperty = UE::MVVM::BindingHelper::GetReturnProperty(Function);
-					if (ReturnProperty != nullptr)
-					{
-						FEdGraphPinType ReturnType;
-						if (Schema->ConvertPropertyToPinType(ReturnProperty, ReturnType))
-						{
-							if (Schema->ArePinTypesCompatible(ReturnType, ArgumentType, SourceClass))
-							{
-								return true;
-							}
-						}
-					}
-				}
-
-				return false;
-			},
-			[SourceClass, Source](const FMVVMAvailableBinding& Binding) -> FMVVMBlueprintPropertyPath
-			{
-				FName BindingName = Binding.GetBindingName().ToName();
-
-				UE::MVVM::FMVVMConstFieldVariant Variant;
-				if (const UFunction* Function = SourceClass->FindFunctionByName(BindingName))
-				{
-					Variant = UE::MVVM::FMVVMConstFieldVariant(Function);
-				}
-				else if (const FProperty* Property = SourceClass->FindPropertyByName(BindingName))
-				{
-					Variant = UE::MVVM::FMVVMConstFieldVariant(Property);
-				}
-
-				FMVVMBlueprintPropertyPath Path;
-				if (!Source.Name.IsNone())
-				{
-					Path.SetWidgetName(Source.Name);
-				}
-				else
-				{
-					Path.SetViewModelId(Source.ViewModelId);
-				}
-				Path.SetBasePropertyPath(Variant);
-
-				return Path;
-			});
-
-		return AvailablePaths;
-	}
-
 	ECheckBoxState FConversionPathCustomization::OnGetIsArgumentBound(FName ArgumentName, bool bSourceToDestination) const
 	{
 		UMVVMEditorSubsystem* EditorSubsystem = GEditor->GetEditorSubsystem<UMVVMEditorSubsystem>();
@@ -638,6 +442,17 @@ namespace UE::MVVM
 		UEdGraphPin* Pin = CallFunctionNode->FindPin(ArgumentName);
 		return Pin != nullptr ? Pin->PinType : FEdGraphPinType();
 	}
+
+	EMVVMBindingMode FConversionPathCustomization::GetBindingMode() const
+	{
+		uint8 EnumValue = 0;
+		if (BindingModeHandle->GetValue(EnumValue) == FPropertyAccess::Success)
+		{
+			return static_cast<EMVVMBindingMode>(EnumValue);
+		}
+
+		return EMVVMBindingMode::OneWayToDestination;
+	}
 }
 
 #undef LOCTEXT_NAMESPACE 
diff --git a/Engine/Plugins/Runtime/ModelViewViewModel/Source/ModelViewViewModelEditor/Private/Customizations/MVVMConversionPathCustomization.h b/Engine/Plugins/Runtime/ModelViewViewModel/Source/ModelViewViewModelEditor/Private/Customizations/MVVMConversionPathCustomization.h
--- a/Engine/Plugins/Runtime/ModelViewViewModel/Source/ModelViewViewModelEditor/Private/Customizations/MVVMConversionPathCustomization.h
+++ b/Engine/Plugins/Runtime/ModelViewViewModel/Source/ModelViewViewModelEditor/Private/Customizations/MVVMConversionPathCustomization.h
@@ -4,63 +4,64 @@
 
 #include "IPropertyTypeCustomization.h"
 #include "Layout/Visibility.h"
+#include "MVVMPropertyPath.h"
 #include "Styling/SlateTypes.h"
 #include "Templates/ValueOrError.h"
+#include "Types/MVVMBindingMode.h"
 #include "Types/MVVMBindingSource.h"
 #include "Types/MVVMFieldVariant.h"
-#include "MVVMPropertyPath.h"
 
 struct FMVVMPropertyPathBase;
 class FStructOnScope;
 class SGraphPin;
-class SMVVMFieldSelector;
 class UWidgetBlueprint;
 
 namespace UE::MVVM
 {
-	class FConversionPathCustomization : public IPropertyTypeCustomization
+
+class SFieldSelector;
+
+class FConversionPathCustomization : public IPropertyTypeCustomization
+{
+public:
+	static TSharedRef<IPropertyTypeCustomization> MakeInstance(UWidgetBlueprint* InWidgetBlueprint)
 	{
-	public:
-		static TSharedRef<IPropertyTypeCustomization> MakeInstance(UWidgetBlueprint* InWidgetBlueprint)
-		{
-			return MakeShared<FConversionPathCustomization>(InWidgetBlueprint);
-		}
+		return MakeShared<FConversionPathCustomization>(InWidgetBlueprint);
+	}
+
+	FConversionPathCustomization(UWidgetBlueprint* InWidgetBlueprint);
 
-		FConversionPathCustomization(UWidgetBlueprint* InWidgetBlueprint);
+	virtual void CustomizeHeader(TSharedRef<IPropertyHandle> InPropertyHandle, FDetailWidgetRow& HeaderRow, IPropertyTypeCustomizationUtils& CustomizationUtils) override;
+	virtual void CustomizeChildren(TSharedRef<IPropertyHandle> InPropertyHandle, IDetailChildrenBuilder& ChildBuilder, IPropertyTypeCustomizationUtils& CustomizationUtils) override;
 
-		virtual void CustomizeHeader(TSharedRef<IPropertyHandle> InPropertyHandle, FDetailWidgetRow& HeaderRow, IPropertyTypeCustomizationUtils& CustomizationUtils) override;
-		virtual void CustomizeChildren(TSharedRef<IPropertyHandle> InPropertyHandle, IDetailChildrenBuilder& ChildBuilder, IPropertyTypeCustomizationUtils& CustomizationUtils) override;
+private:
+	FText GetFunctionPathText(TSharedRef<IPropertyHandle> Property) const;
+	void AddRowForProperty(IDetailChildrenBuilder& ChildBuilder, const TSharedPtr<IPropertyHandle>& Property, bool bSourceToDestination);
 
-	private:
-		FText GetFunctionPathText(TSharedRef<IPropertyHandle> Property) const;
-		void AddRowForProperty(IDetailChildrenBuilder& ChildBuilder, const TSharedPtr<IPropertyHandle>& Property, bool bSourceToDestination);
+	void OnTextCommitted(const FText& NewValue, ETextCommit::Type CommitType, TSharedRef<IPropertyHandle> Property, bool bSourceToDestination);
+	void OnFunctionPathChanged(const UFunction* NewFunction, TSharedRef<IPropertyHandle> Property, bool bSourceToDestination);
+	void RefreshDetailsView() const;
 
-		void OnTextCommitted(const FText& NewValue, ETextCommit::Type CommitType, TSharedRef<IPropertyHandle> Property, bool bSourceToDestination);
-		void OnFunctionPathChanged(const UFunction* NewFunction, TSharedRef<IPropertyHandle> Property, bool bSourceToDestination);
-		void RefreshDetailsView() const;
+	FMVVMBlueprintPropertyPath OnGetSelectedField(FName ArgumentName, bool bSourceToDestination) const;
+	void OnSetProperty(FMVVMBlueprintPropertyPath NewSelection, FName ArgumentName, bool bSourceToDestination);
 
-		TArray<FBindingSource> OnGetAvailableSources(bool bSourceToDestination) const;
-		FBindingSource OnGetSelectedSource(FName ArgumentName, bool bSourceToDestination) const;
-		void OnSetSource(FBindingSource Source, FName ArgumentName, bool bSourceToDestination);
+	ECheckBoxState OnGetIsArgumentBound(FName ArgumentName, bool bSourceToDestination) const;
+	void OnCheckedBindArgument(ECheckBoxState CheckState, FName ArgumentName, bool bSourceToDestination);
+	EVisibility GetArgumentWidgetVisibility(FName ArgumentName, bool bSourceToDestination, bool bDefaultValue) const;
 
-		TArray<FMVVMBlueprintPropertyPath> OnGetAvailableFields(FName ArgumentName, bool bSourceToDestination) const;
-		FMVVMBlueprintPropertyPath OnGetSelectedField(FName ArgumentName, bool bSourceToDestination) const;
-		void OnSetProperty(FMVVMBlueprintPropertyPath NewSelection, FName ArgumentName, bool bSourceToDestination);
+	FEdGraphPinType GetArgumentPinType(FName ArgumentName, bool bSourceToDestination) const;
 
-		ECheckBoxState OnGetIsArgumentBound(FName ArgumentName, bool bSourceToDestination) const;
-		void OnCheckedBindArgument(ECheckBoxState CheckState, FName ArgumentName, bool bSourceToDestination);
-		EVisibility GetArgumentWidgetVisibility(FName ArgumentName, bool bSourceToDestination, bool bDefaultValue) const;
+	EMVVMBindingMode GetBindingMode() const;
 
-		FEdGraphPinType GetArgumentPinType(FName ArgumentName, bool bSourceToDestination) const;
+private:
+	UWidgetBlueprint* WidgetBlueprint = nullptr;
 
-	private:
-		UWidgetBlueprint* WidgetBlueprint = nullptr;
+	TSharedPtr<IPropertyHandle> ParentHandle;
+	TSharedPtr<IPropertyHandle> BindingModeHandle;
+	TWeakPtr<IPropertyUtilities> WeakPropertyUtilities;
 
-		TSharedPtr<IPropertyHandle> ParentHandle;
-		TSharedPtr<IPropertyHandle> BindingModeHandle;
-		TWeakPtr<IPropertyUtilities> WeakPropertyUtilities;
+	TArray<TSharedPtr<SGraphPin>> ArgumentPinWidgets; 
+	TArray<TSharedPtr<UE::MVVM::SFieldSelector>> ArgumentFieldSelectors;
+};
 
-		TArray<TSharedPtr<SGraphPin>> ArgumentPinWidgets; 
-		TArray<TSharedPtr<SMVVMFieldSelector>> ArgumentFieldSelectors;
-	};
-}
\ No newline at end of file
+} // namespace UE::MVVM
diff --git a/Engine/Plugins/Runtime/ModelViewViewModel/Source/ModelViewViewModelEditor/Private/Customizations/MVVMPropertyPathCustomization.cpp b/Engine/Plugins/Runtime/ModelViewViewModel/Source/ModelViewViewModelEditor/Private/Customizations/MVVMPropertyPathCustomization.cpp
--- a/Engine/Plugins/Runtime/ModelViewViewModel/Source/ModelViewViewModelEditor/Private/Customizations/MVVMPropertyPathCustomization.cpp
+++ b/Engine/Plugins/Runtime/ModelViewViewModel/Source/ModelViewViewModelEditor/Private/Customizations/MVVMPropertyPathCustomization.cpp
@@ -19,7 +19,6 @@
 #include "Widgets/Input/SComboBox.h" 
 #include "Widgets/SMVVMFieldIcon.h"
 #include "Widgets/SMVVMFieldSelector.h"
-#include "Widgets/SMVVMSourceSelector.h"
 #include "Widgets/SOverlay.h"
 #include "Widgets/Text/STextBlock.h"
 
@@ -27,324 +26,227 @@
 
 namespace UE::MVVM
 {
-	FPropertyPathCustomization::FPropertyPathCustomization(UWidgetBlueprint* InWidgetBlueprint) :
-		WidgetBlueprint(InWidgetBlueprint)
-	{
-		check(WidgetBlueprint != nullptr);
-
-		OnBlueprintChangedHandle = WidgetBlueprint->OnChanged().AddRaw(this, &FPropertyPathCustomization::HandleBlueprintChanged);
-	}
 
-	FPropertyPathCustomization::~FPropertyPathCustomization()
-	{
-		WidgetBlueprint->OnChanged().Remove(OnBlueprintChangedHandle);
-	}
-
-	void FPropertyPathCustomization::CustomizeHeader(TSharedRef<IPropertyHandle> InPropertyHandle, FDetailWidgetRow& HeaderRow, IPropertyTypeCustomizationUtils& CustomizationUtils)
-	{
-		PropertyHandle = InPropertyHandle;
+FPropertyPathCustomization::FPropertyPathCustomization(UWidgetBlueprint* InWidgetBlueprint) :
+	WidgetBlueprint(InWidgetBlueprint)
+{
+	check(WidgetBlueprint != nullptr);
 
-		const FName PropertyName = PropertyHandle->GetProperty()->GetFName();
-		if (PropertyName == GET_MEMBER_NAME_CHECKED(FMVVMBlueprintViewBinding, WidgetPath))
-		{
-			bIsWidget = true;
-		}
-		else if (PropertyName == GET_MEMBER_NAME_CHECKED(FMVVMBlueprintViewBinding, ViewModelPath))
-		{
-			bIsWidget = false;
-		}
-		else
-		{
-			ensureMsgf(false, TEXT("MVVMPropertyPathCustomization used in unknown context."));
-		}
+	OnBlueprintChangedHandle = WidgetBlueprint->OnChanged().AddRaw(this, &FPropertyPathCustomization::HandleBlueprintChanged);
+}
 
-		TSharedPtr<IPropertyHandle> ParentHandle = PropertyHandle->GetParentHandle();
-		TSharedPtr<IPropertyHandle> OtherHandle;
+FPropertyPathCustomization::~FPropertyPathCustomization()
+{
+	WidgetBlueprint->OnChanged().Remove(OnBlueprintChangedHandle);
+}
 
-		if (bIsWidget)
-		{
-			OtherHandle = ParentHandle->GetChildHandle(GET_MEMBER_NAME_CHECKED(FMVVMBlueprintViewBinding, ViewModelPath));
-		}
-		else
-		{
-			OtherHandle = ParentHandle->GetChildHandle(GET_MEMBER_NAME_CHECKED(FMVVMBlueprintViewBinding, WidgetPath));
-		}
+void FPropertyPathCustomization::CustomizeHeader(TSharedRef<IPropertyHandle> InPropertyHandle, FDetailWidgetRow& HeaderRow, IPropertyTypeCustomizationUtils& CustomizationUtils)
+{
+	PropertyHandle = InPropertyHandle;
 
-		if (OtherHandle.IsValid() && OtherHandle->IsValidHandle())
-		{
-			OtherHandle->SetOnPropertyValueChanged(FSimpleDelegate::CreateSP(this, &FPropertyPathCustomization::OnOtherPropertyChanged));
-		}
-		else
-		{
-			ensureMsgf(false, TEXT("MVVMPropertyPathCustomization used in unknown context."));
-		}
+	const FName PropertyName = PropertyHandle->GetProperty()->GetFName();
+	if (PropertyName == GET_MEMBER_NAME_CHECKED(FMVVMBlueprintViewBinding, WidgetPath))
+	{
+		bIsWidget = true;
+	}
+	else if (PropertyName == GET_MEMBER_NAME_CHECKED(FMVVMBlueprintViewBinding, ViewModelPath))
+	{
+		bIsWidget = false;
+	}
+	else
+	{
+		ensureMsgf(false, TEXT("MVVMPropertyPathCustomization used in unknown context."));
+	}
 
-		TSharedPtr<IPropertyHandle> BindingTypeHandle = ParentHandle->GetChildHandle(GET_MEMBER_NAME_CHECKED(FMVVMBlueprintViewBinding, BindingType));
-		if (BindingTypeHandle.IsValid() && BindingTypeHandle->IsValidHandle())
-		{
-			BindingTypeHandle->SetOnPropertyValueChanged(FSimpleDelegate::CreateSP(this, &FPropertyPathCustomization::OnOtherPropertyChanged));
-		}
-		else
-		{
-			ensureMsgf(false, TEXT("MVVMPropertyPathCustomization used in unknown context."));
-		}
+	TSharedPtr<IPropertyHandle> ParentHandle = PropertyHandle->GetParentHandle();
+	TSharedPtr<IPropertyHandle> OtherHandle;
 
-		uint8 Value;
-		BindingTypeHandle->GetValue(Value);
-
-		HeaderRow
-			.NameWidget
-			.VAlign(VAlign_Center)
-			[
-				PropertyHandle->CreatePropertyNameWidget()
-			]
-			.ValueWidget
-			.VAlign(VAlign_Center)
-			.HAlign(HAlign_Fill)
-			[
-				SNew(SHorizontalBox)
-				+ SHorizontalBox::Slot()
-				.VAlign(VAlign_Center)
-				[
-					SAssignNew(SourceSelector, SMVVMSourceSelector)
-					.TextStyle(FAppStyle::Get(), "SmallText")
-					.AvailableSources(this, &FPropertyPathCustomization::OnGetSources)
-					.SelectedSource(this, &FPropertyPathCustomization::OnGetSelectedSource)
-					.OnSelectionChanged(this, &FPropertyPathCustomization::OnSourceSelectionChanged)
-				]
-			+ SHorizontalBox::Slot()
-				.Padding(5, 0, 0, 0)
-				.VAlign(VAlign_Center)
-				[
-					SAssignNew(FieldSelector, SMVVMFieldSelector)
-					.TextStyle(FAppStyle::Get(), "SmallText")
-					.BindingMode(this, &FPropertyPathCustomization::GetCurrentBindingMode)
-					.SelectedSource(this, &FPropertyPathCustomization::OnGetSelectedSource)
-					.OnSelectionChanged(this, &FPropertyPathCustomization::OnPropertySelectionChanged)
-					.AvailableFields(this, &FPropertyPathCustomization::OnGetFields)
-					.SelectedField(this, &FPropertyPathCustomization::OnGetSelectedField)
-				]				
-			];
+	if (bIsWidget)
+	{
+		OtherHandle = ParentHandle->GetChildHandle(GET_MEMBER_NAME_CHECKED(FMVVMBlueprintViewBinding, ViewModelPath));
 	}
-
-	void FPropertyPathCustomization::OnOtherPropertyChanged()
+	else
 	{
-		FieldSelector->Refresh();
+		OtherHandle = ParentHandle->GetChildHandle(GET_MEMBER_NAME_CHECKED(FMVVMBlueprintViewBinding, WidgetPath));
 	}
 
-	EMVVMBindingMode FPropertyPathCustomization::GetCurrentBindingMode() const
+	if (OtherHandle.IsValid() && OtherHandle->IsValidHandle())
 	{
-		TSharedPtr<IPropertyHandle> BindingTypeHandle = PropertyHandle->GetParentHandle()->GetChildHandle(GET_MEMBER_NAME_CHECKED(FMVVMBlueprintViewBinding, BindingType));
-		uint8 BindingMode;
-		if (BindingTypeHandle->GetValue(BindingMode) == FPropertyAccess::Success)
-		{
-			return (EMVVMBindingMode) BindingMode;
-		}
-
-		return EMVVMBindingMode::OneWayToDestination;
+		OtherHandle->SetOnPropertyValueChanged(FSimpleDelegate::CreateSP(this, &FPropertyPathCustomization::OnOtherPropertyChanged));
+	}
+	else
+	{
+		ensureMsgf(false, TEXT("MVVMPropertyPathCustomization used in unknown context."));
 	}
 
-	void FPropertyPathCustomization::OnSourceSelectionChanged(FBindingSource Selected)
+	TSharedPtr<IPropertyHandle> BindingTypeHandle = ParentHandle->GetChildHandle(GET_MEMBER_NAME_CHECKED(FMVVMBlueprintViewBinding, BindingType));
+	if (BindingTypeHandle.IsValid() && BindingTypeHandle->IsValidHandle())
+	{
+		BindingTypeHandle->SetOnPropertyValueChanged(FSimpleDelegate::CreateSP(this, &FPropertyPathCustomization::OnOtherPropertyChanged));
+	}
+	else
 	{
-		PropertyHandle->NotifyPreChange();
+		ensureMsgf(false, TEXT("MVVMPropertyPathCustomization used in unknown context."));
+	}
 
-		UMVVMEditorSubsystem* Subsystem = GEditor->GetEditorSubsystem<UMVVMEditorSubsystem>();
-		
-		TArray<void*> RawData;
-		PropertyHandle->AccessRawData(RawData);
+	uint8 Value;
+	BindingTypeHandle->GetValue(Value);
+
+	HeaderRow
+		.NameWidget
+		.VAlign(VAlign_Center)
+		[
+			PropertyHandle->CreatePropertyNameWidget()
+		]
+		.ValueWidget
+		.VAlign(VAlign_Center)
+		.HAlign(HAlign_Fill)
+		[
+			SAssignNew(FieldSelector, SFieldSelector, WidgetBlueprint, !bIsWidget)
+			.TextStyle(FAppStyle::Get(), "SmallText")
+			.OnSelectionChanged(this, &FPropertyPathCustomization::OnPropertySelectionChanged)
+			.SelectedField(this, &FPropertyPathCustomization::OnGetSelectedField)
+			.BindingMode(this, &FPropertyPathCustomization::GetCurrentBindingMode)
+		];
+}
 
-		for (void* RawPtr : RawData)
-		{
-			if (RawPtr == nullptr)
-			{
-				continue;
-			}
+void FPropertyPathCustomization::OnOtherPropertyChanged()
+{
+	FieldSelector->Refresh();
+}
 
-			FMVVMBlueprintPropertyPath* PropertyPath = (FMVVMBlueprintPropertyPath*) RawPtr;
-			if (bIsWidget)
-			{
-				PropertyPath->SetWidgetName(Selected.Name);
-			}
-			else
-			{
-				PropertyPath->SetViewModelId(Selected.ViewModelId);
-			}
-		}
+EMVVMBindingMode FPropertyPathCustomization::GetCurrentBindingMode() const
+{
+	TSharedPtr<IPropertyHandle> BindingTypeHandle = PropertyHandle->GetParentHandle()->GetChildHandle(GET_MEMBER_NAME_CHECKED(FMVVMBlueprintViewBinding, BindingType));
+	uint8 BindingMode;
+	if (BindingTypeHandle->GetValue(BindingMode) == FPropertyAccess::Success)
+	{
+		return (EMVVMBindingMode) BindingMode;
+	}
 
-		PropertyHandle->NotifyPostChange(EPropertyChangeType::ValueSet);
+	return EMVVMBindingMode::OneWayToDestination;
+}
 
-		FieldSelector->Refresh();
+void FPropertyPathCustomization::OnPropertySelectionChanged(FMVVMBlueprintPropertyPath Selected)
+{
+	if (bPropertySelectionChanging)
+	{
+		return;
 	}
+	TGuardValue<bool> ReentrantGuard(bPropertySelectionChanging, true);
+
+	PropertyHandle->NotifyPreChange();
+
+	TArray<void*> RawData;
+	PropertyHandle->AccessRawData(RawData);
 
-	void FPropertyPathCustomization::OnPropertySelectionChanged(FMVVMBlueprintPropertyPath Selected)
+	for (void* RawPtr : RawData)
 	{
-		if (bPropertySelectionChanging)
+		if (RawPtr == nullptr)
 		{
-			return;
+			continue;
 		}
-		TGuardValue<bool> ReentrantGuard(bPropertySelectionChanging, true);
 
-		PropertyHandle->NotifyPreChange();
+		FMVVMBlueprintPropertyPath* PropertyPath = (FMVVMBlueprintPropertyPath*) RawPtr;
+		*PropertyPath = Selected;
+	}
 
-		TArray<void*> RawData;
-		PropertyHandle->AccessRawData(RawData);
+	PropertyHandle->NotifyPostChange(EPropertyChangeType::ValueSet);
+}
 
-		for (void* RawPtr : RawData)
-		{
-			if (RawPtr == nullptr)
-			{
-				continue;
-			}
+void FPropertyPathCustomization::HandleBlueprintChanged(UBlueprint* Blueprint)
+{
+	FieldSelector->Refresh();
+}
 
-			FMVVMBlueprintPropertyPath* PropertyPath = (FMVVMBlueprintPropertyPath*) RawPtr;
-			*PropertyPath = Selected;
-		}
+FBindingSource FPropertyPathCustomization::OnGetSelectedSource() const
+{
+	TArray<void*> RawData;
+	PropertyHandle->AccessRawData(RawData);
 
-		PropertyHandle->NotifyPostChange(EPropertyChangeType::ValueSet);
-	}
+	FGuid SelectedGuid;
+	FName SelectedName;
 
-	void FPropertyPathCustomization::HandleBlueprintChanged(UBlueprint* Blueprint)
-	{
-		SourceSelector->Refresh();
-		FieldSelector->Refresh();
-	}
+	bool bFirst = true;
 
-	TArray<UE::MVVM::FBindingSource> FPropertyPathCustomization::OnGetSources() const 
+	for (void* RawPtr : RawData)
 	{
-		UMVVMEditorSubsystem* Subsystem = GEditor->GetEditorSubsystem<UMVVMEditorSubsystem>();
-		if (bIsWidget)
-		{
-			return Subsystem->GetBindableWidgets(WidgetBlueprint);
-		}
-		else
+		if (RawPtr == nullptr)
 		{
-			return Subsystem->GetAllViewModels(WidgetBlueprint);
+			continue;
 		}
-	}
-
-	UE::MVVM::FBindingSource FPropertyPathCustomization::OnGetSelectedSource() const
-	{
-		TArray<void*> RawData;
-		PropertyHandle->AccessRawData(RawData);
 
-		FGuid SelectedGuid;
-		FName SelectedName;
-
-		bool bFirst = true;
-
-		for (void* RawPtr : RawData)
+		FMVVMBlueprintPropertyPath* PropertyPath = (FMVVMBlueprintPropertyPath*) RawPtr;
+		if (bIsWidget)
 		{
-			if (RawPtr == nullptr)
-			{
-				continue;
-			}
-
-			FMVVMBlueprintPropertyPath* PropertyPath = (FMVVMBlueprintPropertyPath*) RawPtr;
-			if (bIsWidget)
+			if (bFirst)
 			{
-				if (bFirst)
-				{
-					SelectedName = PropertyPath->GetWidgetName();
-				}
-				else if (SelectedName != PropertyPath->GetWidgetName())
-				{
-					SelectedName = FName();
-					break;
-				}
+				SelectedName = PropertyPath->GetWidgetName();
 			}
-			else
+			else if (SelectedName != PropertyPath->GetWidgetName())
 			{
-				if (bFirst)
-				{
-					SelectedGuid = PropertyPath->GetViewModelId();
-				}
-				else if (SelectedGuid != PropertyPath->GetViewModelId())
-				{
-					SelectedGuid = FGuid();
-					break;
-				}
+				SelectedName = FName();
+				break;
 			}
-
-			bFirst = false;
-		}
-
-		if (bIsWidget)
-		{
-			return UE::MVVM::FBindingSource::CreateForWidget(WidgetBlueprint, SelectedName);
 		}
 		else
 		{
-			return UE::MVVM::FBindingSource::CreateForViewModel(WidgetBlueprint, SelectedGuid);
+			if (bFirst)
+			{
+				SelectedGuid = PropertyPath->GetViewModelId();
+			}
+			else if (SelectedGuid != PropertyPath->GetViewModelId())
+			{
+				SelectedGuid = FGuid();
+				break;
+			}
 		}
+
+		bFirst = false;
 	}
 
-	TArray<FMVVMBlueprintPropertyPath> FPropertyPathCustomization::OnGetFields() const
+	if (bIsWidget)
 	{
-		UE::MVVM::FBindingSource Source = OnGetSelectedSource();
-
-		TArray<FMVVMAvailableBinding> AvailableBindings;
-
-		UClass* Class = Source.Class;
-		UMVVMSubsystem* Subsystem = GEngine->GetEngineSubsystem<UMVVMSubsystem>();
-		AvailableBindings = Subsystem->GetAvailableBindings(Class, WidgetBlueprint->GeneratedClass);
-
-		TArray<FMVVMBlueprintPropertyPath> AvailablePaths;
-		Algo::Transform(AvailableBindings, AvailablePaths, 
-			[this, Class, Source](const FMVVMAvailableBinding& Binding)
-			{
-				UE::MVVM::FMVVMConstFieldVariant Variant = UE::MVVM::BindingHelper::FindFieldByName(Class, Binding.GetBindingName());
-				FMVVMBlueprintPropertyPath Path;
-				if (!Variant.IsEmpty())
-				{
-					if (bIsWidget)
-					{
-						Path.SetWidgetName(Source.Name);
-					}
-					else
-					{
-						Path.SetViewModelId(Source.ViewModelId);
-					}
-
-					Path.SetBasePropertyPath(Variant);
-				}
-
-				return Path;
-			});
-
-		return AvailablePaths;
+		return FBindingSource::CreateForWidget(WidgetBlueprint, SelectedName);
 	}
-
-	FMVVMBlueprintPropertyPath FPropertyPathCustomization::OnGetSelectedField() const
+	else
 	{
-		TArray<void*> RawData;
-		PropertyHandle->AccessRawData(RawData);
+		return FBindingSource::CreateForViewModel(WidgetBlueprint, SelectedGuid);
+	}
+}
 
-		FMVVMBlueprintPropertyPath SelectedField;
+FMVVMBlueprintPropertyPath FPropertyPathCustomization::OnGetSelectedField() const
+{
+	TArray<void*> RawData;
+	PropertyHandle->AccessRawData(RawData);
 
-		bool bFirst = true;
+	FMVVMBlueprintPropertyPath SelectedField;
 
-		for (void* RawPtr : RawData)
-		{
-			if (RawPtr == nullptr)
-			{
-				continue;
-			}
+	bool bFirst = true;
 
-			FMVVMBlueprintPropertyPath* PropertyPath = (FMVVMBlueprintPropertyPath*) RawPtr;
-			if (bFirst)
-			{
-				SelectedField = *PropertyPath;
-			}
-			else if (SelectedField != *PropertyPath)
-			{
-				SelectedField = FMVVMBlueprintPropertyPath();
-				break;
-			}
-			bFirst = false;
+	for (void* RawPtr : RawData)
+	{
+		if (RawPtr == nullptr)
+		{
+			continue;
 		}
 
-		return SelectedField;
+		FMVVMBlueprintPropertyPath* PropertyPath = (FMVVMBlueprintPropertyPath*) RawPtr;
+		if (bFirst)
+		{
+			SelectedField = *PropertyPath;
+		}
+		else if (SelectedField != *PropertyPath)
+		{
+			SelectedField = FMVVMBlueprintPropertyPath();
+			break;
+		}
+		bFirst = false;
 	}
+
+	return SelectedField;
 }
 
+} // namespace UE::MVVM
+
 #undef LOCTEXT_NAMESPACE
\ No newline at end of file
diff --git a/Engine/Plugins/Runtime/ModelViewViewModel/Source/ModelViewViewModelEditor/Private/Customizations/MVVMPropertyPathCustomization.h b/Engine/Plugins/Runtime/ModelViewViewModel/Source/ModelViewViewModelEditor/Private/Customizations/MVVMPropertyPathCustomization.h
--- a/Engine/Plugins/Runtime/ModelViewViewModel/Source/ModelViewViewModelEditor/Private/Customizations/MVVMPropertyPathCustomization.h
+++ b/Engine/Plugins/Runtime/ModelViewViewModel/Source/ModelViewViewModelEditor/Private/Customizations/MVVMPropertyPathCustomization.h
@@ -7,54 +7,50 @@
 #include "PropertyHandle.h"
 #include "Types/MVVMFieldVariant.h"
 
-class SMVVMSourceSelector;
-class SMVVMFieldSelector;
-
 namespace UE::MVVM
 {
-	struct FBindingSource;
+ 
+struct FBindingSource;
+class SFieldSelector;
 
-	class FPropertyPathCustomization : public IPropertyTypeCustomization
-	{
-	public:
-		FPropertyPathCustomization(UWidgetBlueprint* WidgetBlueprint);
-		~FPropertyPathCustomization();
+class FPropertyPathCustomization : public IPropertyTypeCustomization
+{
+public:
+	FPropertyPathCustomization(UWidgetBlueprint* WidgetBlueprint);
+	~FPropertyPathCustomization();
 
-		virtual void CustomizeHeader(TSharedRef<IPropertyHandle> InPropertyHandle, FDetailWidgetRow& HeaderRow, IPropertyTypeCustomizationUtils& CustomizationUtils) override;
-		virtual void CustomizeChildren(TSharedRef<IPropertyHandle> InPropertyHandle, IDetailChildrenBuilder& ChildBuilder, IPropertyTypeCustomizationUtils& CustomizationUtils) override {}
+	virtual void CustomizeHeader(TSharedRef<IPropertyHandle> InPropertyHandle, FDetailWidgetRow& HeaderRow, IPropertyTypeCustomizationUtils& CustomizationUtils) override;
+	virtual void CustomizeChildren(TSharedRef<IPropertyHandle> InPropertyHandle, IDetailChildrenBuilder& ChildBuilder, IPropertyTypeCustomizationUtils& CustomizationUtils) override {}
+
+	static TSharedRef<IPropertyTypeCustomization> MakeInstance(UWidgetBlueprint* InWidgetBlueprint)
+	{
+		check(InWidgetBlueprint != nullptr);
 
-		static TSharedRef<IPropertyTypeCustomization> MakeInstance(UWidgetBlueprint* InWidgetBlueprint)
-		{
-			check(InWidgetBlueprint != nullptr);
+		return MakeShared<FPropertyPathCustomization>(InWidgetBlueprint);
+	}
 
-			return MakeShared<FPropertyPathCustomization>(InWidgetBlueprint);
-		}
+private:
+	void OnPropertySelectionChanged(FMVVMBlueprintPropertyPath Selected);
 
-	private:
-		void OnSourceSelectionChanged(FBindingSource Selected);
-		void OnPropertySelectionChanged(FMVVMBlueprintPropertyPath Selected);
+	void OnOtherPropertyChanged();
 
-		void OnOtherPropertyChanged();
+	EMVVMBindingMode GetCurrentBindingMode() const;
+	FBindingSource OnGetSelectedSource() const;
 
-		EMVVMBindingMode GetCurrentBindingMode() const;
-		TArray<UE::MVVM::FBindingSource> OnGetSources() const;
-		UE::MVVM::FBindingSource OnGetSelectedSource() const;
+	FMVVMBlueprintPropertyPath OnGetSelectedField() const;
 
-		TArray<FMVVMBlueprintPropertyPath> OnGetFields() const;
-		FMVVMBlueprintPropertyPath OnGetSelectedField() const;
+	void HandleBlueprintChanged(UBlueprint* Blueprint);
 
-		void HandleBlueprintChanged(UBlueprint* Blueprint);
+private:
+	TSharedPtr<IPropertyHandle> PropertyHandle;
+	UWidgetBlueprint* WidgetBlueprint = nullptr;
 
-	private:
-		TSharedPtr<IPropertyHandle> PropertyHandle;
-		UWidgetBlueprint* WidgetBlueprint = nullptr;
+	TSharedPtr<SFieldSelector> FieldSelector;
 
-		TSharedPtr<SMVVMSourceSelector> SourceSelector;
-		TSharedPtr<SMVVMFieldSelector> FieldSelector;
+	bool bIsWidget = false;
+	bool bPropertySelectionChanging = false;
 
-		bool bIsWidget = false;
-		bool bPropertySelectionChanging = false;
+	FDelegateHandle OnBlueprintChangedHandle;
+};
 
-		FDelegateHandle OnBlueprintChangedHandle;
-	};
-}
\ No newline at end of file
+} // namespace UE::MVVM
\ No newline at end of file
diff --git a/Engine/Plugins/Runtime/ModelViewViewModel/Source/ModelViewViewModelEditor/Private/MVVMEditorSubsystem.cpp b/Engine/Plugins/Runtime/ModelViewViewModel/Source/ModelViewViewModelEditor/Private/MVVMEditorSubsystem.cpp
--- a/Engine/Plugins/Runtime/ModelViewViewModel/Source/ModelViewViewModelEditor/Private/MVVMEditorSubsystem.cpp
+++ b/Engine/Plugins/Runtime/ModelViewViewModel/Source/ModelViewViewModelEditor/Private/MVVMEditorSubsystem.cpp
@@ -444,20 +444,6 @@ void UMVVMEditorSubsystem::SetDestinationToSourceConversionFunction(UWidgetBluep
 	}
 }
 
-void UMVVMEditorSubsystem::SetWidgetForBinding(UWidgetBlueprint* WidgetBlueprint, FMVVMBlueprintViewBinding& Binding, const FName WidgetName)
-{
-	if (UMVVMBlueprintView* View = GetView(WidgetBlueprint))
-	{
-		FScopedTransaction Transaction(LOCTEXT("SetBindingSource", "Set Binding Source"));
-
-		UE::MVVM::Private::OnBindingPreEditChange(View, GET_MEMBER_NAME_CHECKED(FMVVMBlueprintViewBinding, WidgetPath));
-
-		Binding.WidgetPath.SetWidgetName(WidgetName);
-
-		UE::MVVM::Private::OnBindingPostEditChange(View, GET_MEMBER_NAME_CHECKED(FMVVMBlueprintViewBinding, WidgetPath));
-	}
-}
-
 void UMVVMEditorSubsystem::SetWidgetPropertyForBinding(UWidgetBlueprint* WidgetBlueprint, FMVVMBlueprintViewBinding& Binding, FMVVMBlueprintPropertyPath Field)
 {
 	if (UMVVMBlueprintView* View = GetView(WidgetBlueprint))
@@ -472,20 +458,6 @@ void UMVVMEditorSubsystem::SetWidgetPropertyForBinding(UWidgetBlueprint* WidgetB
 	}
 }
 
-void UMVVMEditorSubsystem::SetViewModelForBinding(UWidgetBlueprint* WidgetBlueprint, FMVVMBlueprintViewBinding& Binding, FMVVMBlueprintViewModelContext ViewModel)
-{
-	if (UMVVMBlueprintView* View = GetView(WidgetBlueprint))
-	{
-		FScopedTransaction Transaction(LOCTEXT("SetBindingSource", "Set Binding Source"));
-
-		UE::MVVM::Private::OnBindingPreEditChange(View, GET_MEMBER_NAME_CHECKED(FMVVMBlueprintViewBinding, ViewModelPath));
-
-		Binding.ViewModelPath.SetViewModelId(ViewModel.GetViewModelId());
-
-		UE::MVVM::Private::OnBindingPostEditChange(View, GET_MEMBER_NAME_CHECKED(FMVVMBlueprintViewBinding, ViewModelPath));
-	}
-}
-
 void UMVVMEditorSubsystem::SetViewModelPropertyForBinding(UWidgetBlueprint* WidgetBlueprint, FMVVMBlueprintViewBinding& Binding, FMVVMBlueprintPropertyPath Field)
 {
 	if (UMVVMBlueprintView* View = GetView(WidgetBlueprint))
diff --git a/Engine/Plugins/Runtime/ModelViewViewModel/Source/ModelViewViewModelEditor/Private/Styling/MVVMEditorStyle.cpp b/Engine/Plugins/Runtime/ModelViewViewModel/Source/ModelViewViewModelEditor/Private/Styling/MVVMEditorStyle.cpp
--- a/Engine/Plugins/Runtime/ModelViewViewModel/Source/ModelViewViewModelEditor/Private/Styling/MVVMEditorStyle.cpp
+++ b/Engine/Plugins/Runtime/ModelViewViewModel/Source/ModelViewViewModelEditor/Private/Styling/MVVMEditorStyle.cpp
@@ -65,6 +65,12 @@ FMVVMEditorStyle::FMVVMEditorStyle()
 	Set("ConversionFunction.DestToSource", new IMAGE_BRUSH_SVG("Slate/ConversionFunction_DestToSource", Icon16x16));
 	Set("ConversionFunction.SourceToDest", new IMAGE_BRUSH_SVG("Slate/ConversionFunction_SourceToDest", Icon16x16));
 
+	Set("FieldSelector.ComboButton", 
+		FComboButtonStyle(FAppStyle::Get().GetWidgetStyle<FComboButtonStyle>("ComboButton"))
+		.SetDownArrowImage(*FCoreStyle::Get().GetBrush("Icons.Edit")));
+
+	Set("FieldSelector.MenuBorderBrush", new FSlateRoundedBoxBrush(FStyleColors::Dropdown, 0.0f, FStyleColors::DropdownOutline, 1.0f));
+
 	FSlateStyleRegistry::RegisterSlateStyle(*this);
 }
 
diff --git a/Engine/Plugins/Runtime/ModelViewViewModel/Source/ModelViewViewModelEditor/Private/Widgets/SMVVMFieldEntry.cpp b/Engine/Plugins/Runtime/ModelViewViewModel/Source/ModelViewViewModelEditor/Private/Widgets/SMVVMFieldEntry.cpp
--- a/Engine/Plugins/Runtime/ModelViewViewModel/Source/ModelViewViewModelEditor/Private/Widgets/SMVVMFieldEntry.cpp
+++ b/Engine/Plugins/Runtime/ModelViewViewModel/Source/ModelViewViewModelEditor/Private/Widgets/SMVVMFieldEntry.cpp
@@ -6,25 +6,28 @@
 
 #define LOCTEXT_NAMESPACE "MVVMFieldEntry"
 
-using namespace UE::MVVM;
+namespace UE::MVVM
+{
+
+namespace Private
+{
 
-namespace UE::MVVM::Private
+FText GetFieldDisplayName(const FMVVMConstFieldVariant& Field)
 {
-	FText GetFieldDisplayName(const FMVVMConstFieldVariant& Field)
+	if (Field.IsProperty())
 	{
-		if (Field.IsProperty())
-		{
-			return Field.GetProperty()->GetDisplayNameText();
-		}
-		else if (Field.IsFunction())
-		{
-			return Field.GetFunction()->GetDisplayNameText();
-		}
-		return LOCTEXT("None", "<None>");
+		return Field.GetProperty()->GetDisplayNameText();
+	}
+	else if (Field.IsFunction())
+	{
+		return Field.GetFunction()->GetDisplayNameText();
 	}
+	return LOCTEXT("None", "<None>");
 }
 
-void SMVVMFieldEntry::Construct(const FArguments& InArgs)
+} // namespace Private
+
+void SFieldEntry::Construct(const FArguments& InArgs)
 {
 	Field = InArgs._Field;
 	OnValidateField = InArgs._OnValidateField;
@@ -37,7 +40,7 @@ void SMVVMFieldEntry::Construct(const FArguments& InArgs)
 		.VAlign(VAlign_Center)
 		.AutoWidth()
 		[
-			SAssignNew(Icon, SMVVMFieldIcon)
+			SAssignNew(Icon, SFieldIcon)
 		]
 		+ SHorizontalBox::Slot()
 		.HAlign(HAlign_Fill)
@@ -52,7 +55,7 @@ void SMVVMFieldEntry::Construct(const FArguments& InArgs)
 	Refresh();
 }
 
-void SMVVMFieldEntry::Refresh()
+void SFieldEntry::Refresh()
 {
 	FText ToolTipText = FText::GetEmpty();
 	bool bEnabled = true;
@@ -93,11 +96,13 @@ void SMVVMFieldEntry::Refresh()
 	Label->SetText(Private::GetFieldDisplayName(Variant));
 }
 
-void SMVVMFieldEntry::SetField(const FMVVMBlueprintPropertyPath& InField)
+void SFieldEntry::SetField(const FMVVMBlueprintPropertyPath& InField)
 {
 	Field = InField;
 
 	Refresh();
 }
 
+} // namespace UE::MVVM
+
 #undef LOCTEXT_NAMESPACE
diff --git a/Engine/Plugins/Runtime/ModelViewViewModel/Source/ModelViewViewModelEditor/Private/Widgets/SMVVMFieldEntry.h b/Engine/Plugins/Runtime/ModelViewViewModel/Source/ModelViewViewModelEditor/Private/Widgets/SMVVMFieldEntry.h
--- a/Engine/Plugins/Runtime/ModelViewViewModel/Source/ModelViewViewModelEditor/Private/Widgets/SMVVMFieldEntry.h
+++ b/Engine/Plugins/Runtime/ModelViewViewModel/Source/ModelViewViewModelEditor/Private/Widgets/SMVVMFieldEntry.h
@@ -11,24 +11,25 @@
 #include "Widgets/SCompoundWidget.h"
 #include "Widgets/Text/STextBlock.h"
 
-class SMVVMFieldIcon;
 class STextBlock;
 
 namespace UE::MVVM
 {
-	using FIsFieldValidResult = TValueOrError<bool, FString>;
-	DECLARE_DELEGATE_RetVal_OneParam(FIsFieldValidResult, FIsFieldValid, FMVVMBlueprintPropertyPath);
-}
 
-class SMVVMFieldEntry : public SCompoundWidget
+class SFieldIcon;
+
+using FIsFieldValidResult = TValueOrError<bool, FString>;
+DECLARE_DELEGATE_RetVal_OneParam(FIsFieldValidResult, FIsFieldValid, FMVVMBlueprintPropertyPath);
+
+class SFieldEntry : public SCompoundWidget
 {
 public:
-	SLATE_BEGIN_ARGS(SMVVMFieldEntry) :
+	SLATE_BEGIN_ARGS(SFieldEntry) :
 		_TextStyle(&FCoreStyle::Get().GetWidgetStyle<FTextBlockStyle>("NormalText"))
 	{}
 		SLATE_STYLE_ARGUMENT(FTextBlockStyle, TextStyle)
 		SLATE_ARGUMENT(FMVVMBlueprintPropertyPath, Field)
-		SLATE_EVENT(UE::MVVM::FIsFieldValid, OnValidateField)
+		SLATE_EVENT(FIsFieldValid, OnValidateField)
 	SLATE_END_ARGS()
 
 	void Construct(const FArguments& InArgs);
@@ -37,7 +38,9 @@ public:
 
 private:
 	FMVVMBlueprintPropertyPath Field;
-	UE::MVVM::FIsFieldValid OnValidateField;
-	TSharedPtr<SMVVMFieldIcon> Icon;
+	FIsFieldValid OnValidateField;
+	TSharedPtr<SFieldIcon> Icon;
 	TSharedPtr<STextBlock> Label;
 };
+
+} // namespace UE::MVVM
diff --git a/Engine/Plugins/Runtime/ModelViewViewModel/Source/ModelViewViewModelEditor/Private/Widgets/SMVVMFieldIcon.cpp b/Engine/Plugins/Runtime/ModelViewViewModel/Source/ModelViewViewModelEditor/Private/Widgets/SMVVMFieldIcon.cpp
--- a/Engine/Plugins/Runtime/ModelViewViewModel/Source/ModelViewViewModelEditor/Private/Widgets/SMVVMFieldIcon.cpp
+++ b/Engine/Plugins/Runtime/ModelViewViewModel/Source/ModelViewViewModelEditor/Private/Widgets/SMVVMFieldIcon.cpp
@@ -9,9 +9,12 @@
 #include "Styling/StyleColors.h"
 #include "Widgets/Images/SLayeredImage.h"
 
-#define LOCTEXT_NAMESPACE "SMVVMFieldIcon"
+#define LOCTEXT_NAMESPACE "SFieldIcon"
 
-void SMVVMFieldIcon::Construct(const FArguments& Args)
+namespace UE::MVVM
+{
+
+void SFieldIcon::Construct(const FArguments& Args)
 {
 	ChildSlot
 	[
@@ -21,7 +24,7 @@ void SMVVMFieldIcon::Construct(const FArguments& Args)
 	RefreshBinding(Args._Field);
 }
 
-void SMVVMFieldIcon::RefreshBinding(const UE::MVVM::FMVVMConstFieldVariant& Field)
+void SFieldIcon::RefreshBinding(const FMVVMConstFieldVariant& Field)
 {
 	LayeredImage->SetImage(nullptr);
 	LayeredImage->RemoveAllLayers();
@@ -40,14 +43,14 @@ void SMVVMFieldIcon::RefreshBinding(const UE::MVVM::FMVVMConstFieldVariant& Fiel
 	else if (Field.IsFunction())
 	{
 		const UFunction* Function = Field.GetFunction();
-		const FProperty* ReturnProperty = UE::MVVM::BindingHelper::GetReturnProperty(Function);
+		const FProperty* ReturnProperty = BindingHelper::GetReturnProperty(Function);
 		if (ReturnProperty != nullptr)
 		{
 			IconProperty = ReturnProperty;
 		}
 		else
 		{
-			IconProperty = UE::MVVM::BindingHelper::GetFirstArgumentProperty(Function);
+			IconProperty = BindingHelper::GetFirstArgumentProperty(Function);
 		}
 	}
 
@@ -63,4 +66,6 @@ void SMVVMFieldIcon::RefreshBinding(const UE::MVVM::FMVVMConstFieldVariant& Fiel
 	}
 }
 
+} // namespace UE::MVVM
+
 #undef LOCTEXT_NAMESPACE
diff --git a/Engine/Plugins/Runtime/ModelViewViewModel/Source/ModelViewViewModelEditor/Private/Widgets/SMVVMFieldIcon.h b/Engine/Plugins/Runtime/ModelViewViewModel/Source/ModelViewViewModelEditor/Private/Widgets/SMVVMFieldIcon.h
--- a/Engine/Plugins/Runtime/ModelViewViewModel/Source/ModelViewViewModelEditor/Private/Widgets/SMVVMFieldIcon.h
+++ b/Engine/Plugins/Runtime/ModelViewViewModel/Source/ModelViewViewModelEditor/Private/Widgets/SMVVMFieldIcon.h
@@ -7,16 +7,21 @@
 
 class SLayeredImage;
 
-class SMVVMFieldIcon : public SCompoundWidget
+namespace UE::MVVM
+{
+
+class SFieldIcon : public SCompoundWidget
 {
 public:
-	SLATE_BEGIN_ARGS(SMVVMFieldIcon) {}
-		SLATE_ARGUMENT(UE::MVVM::FMVVMConstFieldVariant, Field)
+	SLATE_BEGIN_ARGS(SFieldIcon) {}
+		SLATE_ARGUMENT(FMVVMConstFieldVariant, Field)
 	SLATE_END_ARGS()
 
 	void Construct(const FArguments& Args);
-	void RefreshBinding(const UE::MVVM::FMVVMConstFieldVariant& Field);
+	void RefreshBinding(const FMVVMConstFieldVariant& Field);
 
 private:
 	TSharedPtr<SLayeredImage> LayeredImage;
-};
\ No newline at end of file
+};
+
+} // namespace UE::MVVM
diff --git a/Engine/Plugins/Runtime/ModelViewViewModel/Source/ModelViewViewModelEditor/Private/Widgets/SMVVMFieldSelector.cpp b/Engine/Plugins/Runtime/ModelViewViewModel/Source/ModelViewViewModelEditor/Private/Widgets/SMVVMFieldSelector.cpp
--- a/Engine/Plugins/Runtime/ModelViewViewModel/Source/ModelViewViewModelEditor/Private/Widgets/SMVVMFieldSelector.cpp
+++ b/Engine/Plugins/Runtime/ModelViewViewModel/Source/ModelViewViewModelEditor/Private/Widgets/SMVVMFieldSelector.cpp
@@ -1,111 +1,141 @@
 // Copyright Epic Games, Inc. All Rights Reserved.
 
 #include "Widgets/SMVVMFieldSelector.h"
-
-#include "Algo/Transform.h"
-#include "Engine/Engine.h"
-#include "Widgets/SMVVMFieldIcon.h"
-#include "MVVMSubsystem.h"
-#include "Widgets/Input/SComboBox.h"
+ 
+#include "Editor.h"
+#include "MVVMEditorSubsystem.h"
+#include "SNegativeActionButton.h"
+#include "SPositiveActionButton.h"
+#include "SPrimaryButton.h"
 #include "SSimpleButton.h"
+#include "Styling/MVVMEditorStyle.h"
+#include "Widgets/SMVVMFieldEntry.h"
+#include "Widgets/SMVVMSourceEntry.h"
+#include "Widgets/Images/SImage.h"
+#include "Widgets/Input/SComboButton.h"
+#include "Widgets/Input/SSearchBox.h"
 
 #define LOCTEXT_NAMESPACE "MVVMFieldSelector"
 
-using namespace UE::MVVM;
+namespace UE::MVVM
+{
+
+namespace Private
+{
+
+FBindingSource GetSourceFromPath(const UWidgetBlueprint* WidgetBlueprint, const FMVVMBlueprintPropertyPath& Path)
+{
+	if (Path.IsFromViewModel())
+	{
+		return FBindingSource::CreateForViewModel(WidgetBlueprint, Path.GetViewModelId());
+	}
+	else if (Path.IsFromWidget())
+	{
+		return FBindingSource::CreateForWidget(WidgetBlueprint, Path.GetWidgetName());
+	}
+	return FBindingSource();
+}
+
+} // namespace Private
 
-void SMVVMFieldSelector::Construct(const FArguments& InArgs)
+void SFieldSelector::Construct(const FArguments& InArgs, const UWidgetBlueprint* InWidgetBlueprint, bool bInViewModelProperty)
 {
 	OnValidateFieldDelegate = InArgs._OnValidateField;
 	OnSelectionChangedDelegate = InArgs._OnSelectionChanged;
+	check(OnSelectionChangedDelegate.IsBound());
+	WidgetBlueprint = InWidgetBlueprint;
+	check(WidgetBlueprint);
+
 	SelectedField = InArgs._SelectedField;
 	check(SelectedField.IsSet());
 
-	SelectedSource = InArgs._SelectedSource;
-	check(SelectedSource.IsSet());
-
-	AvailableFields = InArgs._AvailableFields;
-	check(AvailableFields.IsSet());
+	BindingMode = InArgs._BindingMode;
+	check(BindingMode.IsSet());
 
-	CachedAvailableFields = AvailableFields.Get(TArray<FMVVMBlueprintPropertyPath>());
+	bViewModelProperty = bInViewModelProperty;
 
-	BindingMode = InArgs._BindingMode;
-	bIsSource = InArgs._IsSource;
 	TextStyle = InArgs._TextStyle;
 
-	Refresh();
-
 	ChildSlot
 	[
-		SNew(SHorizontalBox)
-		+ SHorizontalBox::Slot()
-		.HAlign(HAlign_Fill)
+		SNew(SBox)
+		.HAlign(HAlign_Left)
 		.VAlign(VAlign_Center)
 		[
-			SNew(SOverlay)
-			+ SOverlay::Slot()
-			.HAlign(HAlign_Center)
-			.VAlign(VAlign_Center)
+			SAssignNew(ComboButton, SComboButton)
+			.ComboButtonStyle(FMVVMEditorStyle::Get(), "FieldSelector.ComboButton")
+			.OnGetMenuContent(this, &SFieldSelector::OnGetMenuContent)
+			.ContentPadding(FMargin(4, 2))
+			.ButtonContent()
 			[
-				SNew(STextBlock)
-				.Text(LOCTEXT("NoSourceSelected", "No source selected"))
-				.TextStyle(FAppStyle::Get(), "HintText")
-				.Visibility_Lambda([this]()
-				{
-					return SelectedSource.Get(UE::MVVM::FBindingSource()).IsValid() ? EVisibility::Collapsed : EVisibility::Visible;
-				})
-			]
-			+ SOverlay::Slot()
-			[
-				SAssignNew(FieldComboBox, SComboBox<FMVVMBlueprintPropertyPath>)
-				.Visibility_Lambda([this]()
-				{
-					return SelectedSource.Get(UE::MVVM::FBindingSource()).IsValid() ? EVisibility::Visible : EVisibility::Hidden;
-				})
-				.OptionsSource(&CachedAvailableFields)
-				.InitiallySelectedItem(SelectedField.Get(FMVVMBlueprintPropertyPath()))
-				.OnGenerateWidget(this, &SMVVMFieldSelector::OnGenerateFieldWidget)
-				.OnSelectionChanged(this, &SMVVMFieldSelector::OnComboBoxSelectionChanged)
+				SNew(SOverlay)
+				+ SOverlay::Slot()
 				[
-					SAssignNew(SelectedEntryWidget, SMVVMFieldEntry)
-					.TextStyle(TextStyle)
-					.Field(SelectedField.Get(FMVVMBlueprintPropertyPath()))
-					.OnValidateField(this, &SMVVMFieldSelector::ValidateField)
+					SNew(SHorizontalBox)
+					.Visibility_Lambda([this]() { return CachedSelectedField.IsEmpty() ? EVisibility::Collapsed : EVisibility::Visible; })
+					+ SHorizontalBox::Slot()
+					.Padding(8, 0, 0, 0)
+					.AutoWidth()
+					[
+						SAssignNew(SelectedSourceWidget, SSourceEntry)
+					]
+					+ SHorizontalBox::Slot()
+					.Padding(8, 0)
+					.AutoWidth()
+					[
+						SNew(SImage)
+						.Image(FAppStyle::Get().GetBrush("Icons.ChevronRight"))
+					]
+					+ SHorizontalBox::Slot()
+					.Padding(0, 0, 8, 0)
+					.AutoWidth()
+					[
+						SAssignNew(SelectedEntryWidget, SFieldEntry)
+						.TextStyle(TextStyle)
+						.OnValidateField(this, &SFieldSelector::ValidateField)
+					]
+				]
+				+ SOverlay::Slot()
+				[
+					SNew(SBox)
+					.Padding(FMargin(8, 0, 8, 0))
+					[
+						SNew(STextBlock)
+						.Visibility_Lambda([this]() { return CachedSelectedField.IsEmpty() ? EVisibility::Visible : EVisibility::Collapsed; })
+						.TextStyle(FAppStyle::Get(), "HintText")
+						.Text(LOCTEXT("None", "No field selected"))
+					]
 				]
 			]
 		]
-		+ SHorizontalBox::Slot()
-		.HAlign(HAlign_Right)
-		.VAlign(VAlign_Center)
-		.AutoWidth()
-		[
-			SNew(SSimpleButton)
-			.Icon(FAppStyle::Get().GetBrush("Icons.X"))
-			.ToolTipText(LOCTEXT("ClearField", "Clear field selection."))
-			.Visibility(this, &SMVVMFieldSelector::GetClearVisibility)
-			.OnClicked(this, &SMVVMFieldSelector::OnClearBinding)
-		]
 	];
+
+	Refresh();
 }
 
-void SMVVMFieldSelector::OnComboBoxSelectionChanged(FMVVMBlueprintPropertyPath Selected, ESelectInfo::Type SelectionType)
+void SFieldSelector::OnFieldSelected(FMVVMBlueprintPropertyPath Selected)
 {
-	SelectedEntryWidget->SetField(Selected);
-
 	OnSelectionChangedDelegate.ExecuteIfBound(Selected);
+	
+	Refresh();
 }
 
-void SMVVMFieldSelector::Refresh()
+void SFieldSelector::Refresh()
 {
-	CachedAvailableFields = AvailableFields.Get(TArray<FMVVMBlueprintPropertyPath>());
+	CachedSelectedField = SelectedField.Get();
 
-	if (FieldComboBox.IsValid())
+	if (SelectedEntryWidget.IsValid())
+	{
+		SelectedEntryWidget->SetField(CachedSelectedField);
+	}
+
+	if (SelectedSourceWidget.IsValid())
 	{
-		FieldComboBox->RefreshOptions();
-		FieldComboBox->SetSelectedItem(SelectedField.Get(FMVVMBlueprintPropertyPath()));
+		SelectedSourceWidget->RefreshSource(Private::GetSourceFromPath(WidgetBlueprint, CachedSelectedField));
 	}
 }
 
-TValueOrError<bool, FString> SMVVMFieldSelector::ValidateField(FMVVMBlueprintPropertyPath Field) const
+TValueOrError<bool, FString> SFieldSelector::ValidateField(FMVVMBlueprintPropertyPath Field) const
 {
 	if (OnValidateFieldDelegate.IsBound())
 	{
@@ -115,38 +145,193 @@ TValueOrError<bool, FString> SMVVMFieldSelector::ValidateField(FMVVMBlueprintPro
 	return MakeValue(true);
 }
 
-TSharedRef<SWidget> SMVVMFieldSelector::OnGenerateFieldWidget(FMVVMBlueprintPropertyPath Path) const
+TSharedRef<SWidget> SFieldSelector::OnGenerateFieldWidget(FMVVMBlueprintPropertyPath Path) const
 {
-	return SNew(SMVVMFieldEntry)
+	return SNew(SFieldEntry)
 		.TextStyle(TextStyle)
 		.Field(Path)
-		.OnValidateField(this, &SMVVMFieldSelector::ValidateField);
+		.OnValidateField(this, &SFieldSelector::ValidateField);
 }
 
-EVisibility SMVVMFieldSelector::GetClearVisibility() const
+bool SFieldSelector::IsSelectEnabled() const
 {
-	FMVVMBlueprintPropertyPath SelectedPath = SelectedField.Get(FMVVMBlueprintPropertyPath());
-	if (!SelectedPath.IsEmpty())
+	if (BindingList.IsValid())
 	{
-		return EVisibility::Visible;
+		FMVVMBlueprintPropertyPath Path = BindingList->GetSelectedProperty();
+		return (Path.IsFromViewModel() || Path.IsFromWidget()) && !Path.GetBasePropertyPath().IsEmpty();
 	}
-
-	return EVisibility::Collapsed;
+	return false;
 }
 
-FReply SMVVMFieldSelector::OnClearBinding()
+FReply SFieldSelector::OnSelectProperty()
 {
-	if (FieldComboBox.IsValid())
+	if (BindingList.IsValid())
 	{
-		FieldComboBox->ClearSelection();
+		FMVVMBlueprintPropertyPath Path = BindingList->GetSelectedProperty();
+		SetSelection(Path);
 	}
+	return FReply::Handled();
+}
 
-	if (SelectedEntryWidget.IsValid())
+bool SFieldSelector::IsClearEnabled() const
+{
+	return !CachedSelectedField.IsEmpty();
+}
+
+FReply SFieldSelector::OnClearBinding()
+{
+	SetSelection(FMVVMBlueprintPropertyPath());
+	return FReply::Handled();
+}
+
+void SFieldSelector::HandleSearchChanged(const FText& InFilterText)
+{
+	BindingList->SetRawFilterText(InFilterText);
+}
+
+FReply SFieldSelector::OnCancel()
+{
+	if (ComboButton.IsValid())
 	{
-		SelectedEntryWidget->Refresh();
+		ComboButton->SetIsOpen(false);
 	}
 
 	return FReply::Handled();
 }
 
+void SFieldSelector::SetSelection(const FMVVMBlueprintPropertyPath& SelectedPath)
+{
+	if (ComboButton.IsValid())
+	{
+		ComboButton->SetIsOpen(false);
+	}
+
+	OnSelectionChangedDelegate.Execute(SelectedPath);
+	Refresh();
+}
+
+TSharedRef<SWidget> SFieldSelector::OnGetMenuContent()
+{
+	TSharedRef<SWidget> MenuWidget = 
+		SNew(SBorder)
+		.BorderImage(FMVVMEditorStyle::Get().GetBrush("FieldSelector.MenuBorderBrush"))
+		.Padding(FMargin(8, 1, 8, 3))
+		[
+			SNew(SVerticalBox)
+			+ SVerticalBox::Slot()
+			.Padding(4)
+			.AutoHeight()
+			[
+				SAssignNew(SearchBox, SSearchBox)
+				.HintText(LOCTEXT("SearchViewmodel", "Search"))
+				.OnTextChanged(this, &SFieldSelector::HandleSearchChanged)
+			]
+			+ SVerticalBox::Slot()
+			.FillHeight(1.0f)
+			[
+				SNew(SBorder)
+				.BorderImage(FAppStyle::Get().GetBrush("Brushes.Dropdown"))
+				.Padding(0)
+				[
+					SAssignNew(BindingList, SSourceBindingList, WidgetBlueprint)
+					.OnDoubleClicked(this, &SFieldSelector::SetSelection)
+					.FieldVisibilityFlags(GetFieldVisibilityFlags())
+				]
+			]
+			+ SVerticalBox::Slot()
+			.Padding(4, 4, 4, 0)
+			.AutoHeight()
+			[
+				SNew(SHorizontalBox)
+				+ SHorizontalBox::Slot()
+				[
+					SNew(SButton)
+					.OnClicked(this, &SFieldSelector::OnCancel)
+					.HAlign(HAlign_Center)
+					[
+						SNew(STextBlock)
+						.Text(LOCTEXT("Cancel", "Cancel"))
+					]
+				]
+				+ SHorizontalBox::Slot()
+				.Padding(4, 0, 0, 0)
+				[
+					SNew(SButton)
+					.OnClicked(this, &SFieldSelector::OnClearBinding)
+					.HAlign(HAlign_Center)
+					[
+						SNew(STextBlock)
+						.Text(LOCTEXT("Clear", "Clear"))
+					]
+				]
+				+ SHorizontalBox::Slot()
+				.Padding(4, 0, 0, 0)
+				[
+					SNew(SPrimaryButton)
+					.OnClicked(this, &SFieldSelector::OnSelectProperty)
+					.Text(LOCTEXT("Select", "Select"))
+				]
+			]
+		];
+
+	UMVVMEditorSubsystem* EditorSubsystem = GEditor->GetEditorSubsystem<UMVVMEditorSubsystem>();
+	if (bViewModelProperty)
+	{
+		if (UMVVMBlueprintView* View = EditorSubsystem->GetView(WidgetBlueprint))
+		{
+			BindingList->AddViewModels(View->GetViewModels());
+		}
+	}
+	else
+	{
+		TArray<FBindingSource> BindableWidgets = EditorSubsystem->GetBindableWidgets(WidgetBlueprint);
+		BindingList->AddSources(BindableWidgets);
+	}
+
+	return MenuWidget;
+}
+
+EFieldVisibility SFieldSelector::GetFieldVisibilityFlags() const
+{
+	EFieldVisibility Flags = EFieldVisibility::None;
+
+	EMVVMBindingMode Mode = BindingMode.Get();
+	if (bViewModelProperty)
+	{
+		if (IsForwardBinding(Mode))
+		{
+			Flags |= EFieldVisibility::Readable;
+
+			if (!IsOneTimeBinding(Mode))
+			{
+				Flags |= EFieldVisibility::Notify;
+			}
+		}
+		if (IsBackwardBinding(Mode))
+		{
+			Flags |= EFieldVisibility::Writable;
+		}
+	}
+	else
+	{
+		if (IsForwardBinding(Mode))
+		{
+			Flags |= EFieldVisibility::Writable;
+		}
+		if (IsBackwardBinding(Mode))
+		{
+			Flags |= EFieldVisibility::Readable;
+
+			if (!IsOneTimeBinding(Mode))
+			{
+				Flags |= EFieldVisibility::Notify;
+			}
+		}
+	}
+
+	return Flags;
+}
+
+} // namespace UE::MVVM
+
 #undef LOCTEXT_NAMESPACE
diff --git a/Engine/Plugins/Runtime/ModelViewViewModel/Source/ModelViewViewModelEditor/Private/Widgets/SMVVMFieldSelector.h b/Engine/Plugins/Runtime/ModelViewViewModel/Source/ModelViewViewModelEditor/Private/Widgets/SMVVMFieldSelector.h
--- a/Engine/Plugins/Runtime/ModelViewViewModel/Source/ModelViewViewModelEditor/Private/Widgets/SMVVMFieldSelector.h
+++ b/Engine/Plugins/Runtime/ModelViewViewModel/Source/ModelViewViewModelEditor/Private/Widgets/SMVVMFieldSelector.h
@@ -2,7 +2,7 @@
 
 #pragma once
 
-#include "MVVMFieldVariantListTraits.h"
+#include "MVVMPropertyPath.h"
 #include "Styling/CoreStyle.h"
 #include "Templates/ValueOrError.h"
 #include "Types/MVVMBindingMode.h"
@@ -13,66 +13,83 @@
 #include "Widgets/Input/SComboBox.h"
 #include "Widgets/SCompoundWidget.h"
 #include "Widgets/SMVVMFieldEntry.h"
+#include "Widgets/SMVVMViewModelBindingListWidget.h"
+#include "Widgets/Text/STextBlock.h"
+
+class SComboButton;
+class SSearchBox;
 
 namespace UE::MVVM
 {
-	class IFieldPathHelper;
-}
 
-class SMVVMFieldIcon;
-class STextBlock; 
+class SSourceEntry;
+class SSourceBindingList;
 
-class SMVVMFieldSelector : public SCompoundWidget
+class SFieldSelector : public SCompoundWidget
 {
 public:
 	DECLARE_DELEGATE_OneParam(FSelectionChanged, FMVVMBlueprintPropertyPath);
 
-	SLATE_BEGIN_ARGS(SMVVMFieldSelector) :
-		_TextStyle(&FCoreStyle::Get().GetWidgetStyle<FTextBlockStyle>("NormalText")),
-		_IsSource(false)
+	SLATE_BEGIN_ARGS(SFieldSelector) :
+		_TextStyle(&FCoreStyle::Get().GetWidgetStyle<FTextBlockStyle>("NormalText"))
 		{
 		}
 		SLATE_STYLE_ARGUMENT(FTextBlockStyle, TextStyle)
-		SLATE_ATTRIBUTE(EMVVMBindingMode, BindingMode)
-		SLATE_ARGUMENT(bool, IsSource)
-		SLATE_ATTRIBUTE(UE::MVVM::FBindingSource, SelectedSource)
 		SLATE_ATTRIBUTE(FMVVMBlueprintPropertyPath, SelectedField)
-		SLATE_ATTRIBUTE(TArray<FMVVMBlueprintPropertyPath>, AvailableFields)
 		SLATE_EVENT(FSelectionChanged, OnSelectionChanged)
+		SLATE_ATTRIBUTE(EMVVMBindingMode, BindingMode)
 
 		/**
 		  * Would the given field be a valid entry for this combo box? 
 		  * The error string returned will be used as a tooltip.
 		  */
-		SLATE_EVENT(UE::MVVM::FIsFieldValid, OnValidateField)
+		SLATE_EVENT(FIsFieldValid, OnValidateField)
 	SLATE_END_ARGS()
 
-	void Construct(const FArguments& InArgs);
+	void Construct(const FArguments& InArgs, const UWidgetBlueprint* InWidgetBlueprint, bool bInViewModelProperty);
 	void Refresh();
 
 private:
 	TSharedRef<SWidget> OnGenerateFieldWidget(FMVVMBlueprintPropertyPath Path) const;
-	void OnComboBoxSelectionChanged(FMVVMBlueprintPropertyPath Selected, ESelectInfo::Type SelectionType);
+
+	TSharedRef<SWidget> OnGetMenuContent();
+	void OnFieldSelected(FMVVMBlueprintPropertyPath Selected);
 
 	TValueOrError<bool, FString> ValidateField(FMVVMBlueprintPropertyPath Field) const;
 
-	EVisibility GetClearVisibility() const;
+	void HandleSearchChanged(const FText& InFilterText);
+
+	bool IsClearEnabled() const;
 	FReply OnClearBinding();
 
+	bool IsSelectEnabled() const;
+	FReply OnSelectProperty();
+
+	FReply OnCancel();
+
+	void SetSelection(const FMVVMBlueprintPropertyPath& SelectedPath);
+	EFieldVisibility GetFieldVisibilityFlags() const;
+
 private:
-	TAttribute<EMVVMBindingMode> BindingMode;
-	TSharedPtr<SMVVMFieldEntry> SelectedEntryWidget;
-	TAttribute<UE::MVVM::FBindingSource> SelectedSource;
 	TAttribute<FMVVMBlueprintPropertyPath> SelectedField;
-	TAttribute<TArray<FMVVMBlueprintPropertyPath>> AvailableFields;
+	TAttribute<EMVVMBindingMode> BindingMode;
 
-	const FTextBlockStyle* TextStyle = nullptr;
-	bool bIsSource = false;
+	TSharedPtr<SFieldEntry> SelectedEntryWidget;
+	TSharedPtr<SSourceEntry> SelectedSourceWidget;
 
-	TSharedPtr<SComboBox<FMVVMBlueprintPropertyPath>> FieldComboBox;
+	const FTextBlockStyle* TextStyle = nullptr;
 
 	FSelectionChanged OnSelectionChangedDelegate;
-	UE::MVVM::FIsFieldValid OnValidateFieldDelegate;
+	FIsFieldValid OnValidateFieldDelegate;
 
-	TArray<FMVVMBlueprintPropertyPath> CachedAvailableFields;
+	FMVVMBlueprintPropertyPath CachedSelectedField;
+
+	TSharedPtr<SSourceBindingList> BindingList;
+	TSharedPtr<SSearchBox> SearchBox;
+	TSharedPtr<SComboButton> ComboButton;
+
+	const UWidgetBlueprint* WidgetBlueprint = nullptr;
+	bool bViewModelProperty = false;
 }; 
+
+} // namespace UE::MVVM
\ No newline at end of file
diff --git a/Engine/Plugins/Runtime/ModelViewViewModel/Source/ModelViewViewModelEditor/Private/Widgets/SMVVMManageViewModelsWidget.cpp b/Engine/Plugins/Runtime/ModelViewViewModel/Source/ModelViewViewModelEditor/Private/Widgets/SMVVMManageViewModelsWidget.cpp
--- a/Engine/Plugins/Runtime/ModelViewViewModel/Source/ModelViewViewModelEditor/Private/Widgets/SMVVMManageViewModelsWidget.cpp
+++ b/Engine/Plugins/Runtime/ModelViewViewModel/Source/ModelViewViewModelEditor/Private/Widgets/SMVVMManageViewModelsWidget.cpp
@@ -285,7 +285,7 @@ void SMVVMManageViewModelsWidget::Construct(const FArguments& InArgs)
 						SNew(SBorder)
 						.BorderImage(FStyleDefaults::GetNoBrush())
 						[
-							SAssignNew(ViewModelBindingListWidget, UE::MVVM::SViewModelBindingListWidget, nullptr)
+							SAssignNew(BindingListWidget, UE::MVVM::SSourceBindingList, InArgs._WidgetBlueprint)
 						]
 					]
 				]
@@ -769,7 +769,7 @@ void SMVVMManageViewModelsWidget::HandleTreeViewSelectionChanged(TSharedPtr<FMVV
 	{
 		InClass = GetClassFromNode(Item.ToSharedRef());
 	}
-	ViewModelBindingListWidget->SetViewModel(InClass, FName(), FGuid());
+	BindingListWidget->AddSource(InClass, FName(), FGuid());
 }
 
 void SMVVMManageViewModelsWidget::HandleTreeViewExpansionChanged(TSharedPtr<FMVVMViewModelTreeNode> Item, bool bExpanded)
diff --git a/Engine/Plugins/Runtime/ModelViewViewModel/Source/ModelViewViewModelEditor/Private/Widgets/SMVVMManageViewModelsWidget.h b/Engine/Plugins/Runtime/ModelViewViewModel/Source/ModelViewViewModelEditor/Private/Widgets/SMVVMManageViewModelsWidget.h
--- a/Engine/Plugins/Runtime/ModelViewViewModel/Source/ModelViewViewModelEditor/Private/Widgets/SMVVMManageViewModelsWidget.h
+++ b/Engine/Plugins/Runtime/ModelViewViewModel/Source/ModelViewViewModelEditor/Private/Widgets/SMVVMManageViewModelsWidget.h
@@ -23,7 +23,7 @@ class UWidgetBlueprint;
 namespace UE::MVVM
 {
 	struct FMVVMFieldVariant;
-	class SViewModelBindingListWidget;
+	class SSourceBindingList;
 
 	namespace Private
 	{
@@ -108,7 +108,7 @@ private:
 
 	TWeakPtr<SWindow> WeakParentWindow;
 
-	TSharedPtr<UE::MVVM::SViewModelBindingListWidget> ViewModelBindingListWidget;
+	TSharedPtr<UE::MVVM::SSourceBindingList> BindingListWidget;
 
 	bool bSearchStringIsEmpty = true;
 	UUserWidget* OuterWidget = nullptr;
diff --git a/Engine/Plugins/Runtime/ModelViewViewModel/Source/ModelViewViewModelEditor/Private/Widgets/SMVVMPropertyPath.cpp b/Engine/Plugins/Runtime/ModelViewViewModel/Source/ModelViewViewModelEditor/Private/Widgets/SMVVMPropertyPath.cpp
--- a/Engine/Plugins/Runtime/ModelViewViewModel/Source/ModelViewViewModelEditor/Private/Widgets/SMVVMPropertyPath.cpp
+++ b/Engine/Plugins/Runtime/ModelViewViewModel/Source/ModelViewViewModelEditor/Private/Widgets/SMVVMPropertyPath.cpp
@@ -15,9 +15,12 @@
 #include "Widgets/Layout/SSeparator.h"
 #include "Widgets/Text/STextBlock.h"
 
-#define LOCTEXT_NAMESPACE "SMVVMPropertyPath"
+#define LOCTEXT_NAMESPACE "SPropertyPath"
 
-void SMVVMPropertyPath::Construct(const FArguments& InArgs)
+namespace UE::MVVM
+{
+
+void SPropertyPath::Construct(const FArguments& InArgs)
 {
 	WidgetBlueprint = InArgs._WidgetBlueprint;
 	check(InArgs._WidgetBlueprint != nullptr);
@@ -33,14 +36,14 @@ void SMVVMPropertyPath::Construct(const FArguments& InArgs)
 		[
 			SNew(STextBlock)
 			.TextStyle(FMVVMEditorStyle::Get(), "PropertyPath.ContextText")
-			.Text(this, &SMVVMPropertyPath::GetSourceDisplayName)
+			.Text(this, &SPropertyPath::GetSourceDisplayName)
 		]
 		+ SHorizontalBox::Slot()
 		.Padding(FMargin(5, 0, 0, 0))
 		.VAlign(VAlign_Center)
 		.AutoWidth()
 		[
-			SNew(SMVVMFieldIcon)
+			SNew(SFieldIcon)
 			.Field(GetLastField())
 		]
 		+ SHorizontalBox::Slot()
@@ -50,12 +53,12 @@ void SMVVMPropertyPath::Construct(const FArguments& InArgs)
 		[
 			SNew(STextBlock)
 			.TextStyle(FMVVMEditorStyle::Get(), "PropertyPath.ContextText")
-			.Text(this, &SMVVMPropertyPath::GetFieldDisplayName)
+			.Text(this, &SPropertyPath::GetFieldDisplayName)
 		]
 	];
 }
 
-FText SMVVMPropertyPath::GetSourceDisplayName() const
+FText SPropertyPath::GetSourceDisplayName() const
 {
 	if (PropertyPath->IsFromWidget())
 	{
@@ -74,14 +77,16 @@ FText SMVVMPropertyPath::GetSourceDisplayName() const
 	return FText::GetEmpty();
 }
 
-FText SMVVMPropertyPath::GetFieldDisplayName() const
+FText SPropertyPath::GetFieldDisplayName() const
 {
 	return FText::FromString(PropertyPath->GetBasePropertyPath());
 }
 
-UE::MVVM::FMVVMConstFieldVariant SMVVMPropertyPath::GetLastField() const
+FMVVMConstFieldVariant SPropertyPath::GetLastField() const
 {
 	return PropertyPath->GetFields().Last();
 }
 
-#undef LOCTEXT_NAMESPACE
\ No newline at end of file
+} // namespace UE::MVVM
+
+#undef LOCTEXT_NAMESPACE
diff --git a/Engine/Plugins/Runtime/ModelViewViewModel/Source/ModelViewViewModelEditor/Private/Widgets/SMVVMPropertyPath.h b/Engine/Plugins/Runtime/ModelViewViewModel/Source/ModelViewViewModelEditor/Private/Widgets/SMVVMPropertyPath.h
--- a/Engine/Plugins/Runtime/ModelViewViewModel/Source/ModelViewViewModelEditor/Private/Widgets/SMVVMPropertyPath.h
+++ b/Engine/Plugins/Runtime/ModelViewViewModel/Source/ModelViewViewModelEditor/Private/Widgets/SMVVMPropertyPath.h
@@ -7,10 +7,13 @@
 #include "Widgets/SMVVMFieldIcon.h"
 #include "Types/MVVMFieldVariant.h"
 
-class SMVVMPropertyPath : public SCompoundWidget
+namespace UE::MVVM
+{
+
+class SPropertyPath : public SCompoundWidget
 {
 public:
-	SLATE_BEGIN_ARGS(SMVVMPropertyPath)
+	SLATE_BEGIN_ARGS(SPropertyPath)
 	{}
 		SLATE_ARGUMENT(FMVVMBlueprintPropertyPath*, PropertyPath)
 		SLATE_ARGUMENT(const UWidgetBlueprint*, WidgetBlueprint)
@@ -21,9 +24,11 @@ public:
 private:
 	FText GetSourceDisplayName() const;
 	FText GetFieldDisplayName() const;
-	UE::MVVM::FMVVMConstFieldVariant GetLastField() const;
+	FMVVMConstFieldVariant GetLastField() const;
 
 private:
 	FMVVMBlueprintPropertyPath* PropertyPath = nullptr;
 	TWeakObjectPtr<const UWidgetBlueprint> WidgetBlueprint = nullptr;
 };
+
+} // namespace UE::MVVM
diff --git a/Engine/Plugins/Runtime/ModelViewViewModel/Source/ModelViewViewModelEditor/Private/Widgets/SMVVMSelectViewModel.cpp b/Engine/Plugins/Runtime/ModelViewViewModel/Source/ModelViewViewModelEditor/Private/Widgets/SMVVMSelectViewModel.cpp
--- a/Engine/Plugins/Runtime/ModelViewViewModel/Source/ModelViewViewModelEditor/Private/Widgets/SMVVMSelectViewModel.cpp
+++ b/Engine/Plugins/Runtime/ModelViewViewModel/Source/ModelViewViewModelEditor/Private/Widgets/SMVVMSelectViewModel.cpp
@@ -109,7 +109,7 @@ void SMVVMSelectViewModel::Construct(const FArguments& InArgs, const UWidgetBlue
 					SNew(SBorder)
 					.BorderImage(FStyleDefaults::GetNoBrush())
 					[
-						SAssignNew(ViewModelBindingListWidget, SViewModelBindingListWidget, WidgetBlueprint)
+						SAssignNew(BindingListWidget, SSourceBindingList, WidgetBlueprint)
 					]
 				]
 			]
@@ -128,7 +128,7 @@ void SMVVMSelectViewModel::Construct(const FArguments& InArgs, const UWidgetBlue
 void SMVVMSelectViewModel::HandleClassPicked(UClass* ClassPicked)
 {
 	SelectedClass = ClassPicked;
-	ViewModelBindingListWidget->SetViewModel(ClassPicked, FName(), FGuid());
+	BindingListWidget->AddSource(ClassPicked, FName(), FGuid());
 }
 
 
diff --git a/Engine/Plugins/Runtime/ModelViewViewModel/Source/ModelViewViewModelEditor/Private/Widgets/SMVVMSelectViewModel.h b/Engine/Plugins/Runtime/ModelViewViewModel/Source/ModelViewViewModelEditor/Private/Widgets/SMVVMSelectViewModel.h
--- a/Engine/Plugins/Runtime/ModelViewViewModel/Source/ModelViewViewModelEditor/Private/Widgets/SMVVMSelectViewModel.h
+++ b/Engine/Plugins/Runtime/ModelViewViewModel/Source/ModelViewViewModelEditor/Private/Widgets/SMVVMSelectViewModel.h
@@ -10,7 +10,7 @@ class UWidgetBlueprint;
 
 namespace UE::MVVM
 {
-class SViewModelBindingListWidget;
+class SSourceBindingList;
 
 class SMVVMSelectViewModel : public SCompoundWidget
 {
@@ -33,7 +33,7 @@ private:
 
 private:
 	TSharedPtr<SClassViewer> ClassViewer;
-	TSharedPtr<SViewModelBindingListWidget> ViewModelBindingListWidget;
+	TSharedPtr<SSourceBindingList> BindingListWidget;
 	TWeakObjectPtr<const UClass> SelectedClass;
 
 	FSimpleDelegate OnCancel;
diff --git a/Engine/Plugins/Runtime/ModelViewViewModel/Source/ModelViewViewModelEditor/Private/Widgets/SMVVMSourceEntry.cpp b/Engine/Plugins/Runtime/ModelViewViewModel/Source/ModelViewViewModelEditor/Private/Widgets/SMVVMSourceEntry.cpp
--- a/Engine/Plugins/Runtime/ModelViewViewModel/Source/ModelViewViewModelEditor/Private/Widgets/SMVVMSourceEntry.cpp
+++ b/Engine/Plugins/Runtime/ModelViewViewModel/Source/ModelViewViewModelEditor/Private/Widgets/SMVVMSourceEntry.cpp
@@ -9,7 +9,13 @@
 
 #define LOCTEXT_NAMESPACE "SMVVMSource"
 
-static const FSlateBrush* GetSourceIcon(const UE::MVVM::FBindingSource& Source)
+namespace UE::MVVM
+{
+	
+namespace Private
+{
+
+const FSlateBrush* GetSourceIcon(const FBindingSource& Source)
 {
 	if (!Source.IsValid())
 	{
@@ -19,7 +25,7 @@ static const FSlateBrush* GetSourceIcon(const UE::MVVM::FBindingSource& Source)
 	return FSlateIconFinder::FindIconBrushForClass(Source.Class);
 }
 
-static FLinearColor GetSourceColor(const UE::MVVM::FBindingSource& Source)
+FLinearColor GetSourceColor(const FBindingSource& Source)
 {
 	if (!Source.IsValid())
 	{
@@ -34,7 +40,7 @@ static FLinearColor GetSourceColor(const UE::MVVM::FBindingSource& Source)
 	return Color;
 }
 
-static FText GetSourceDisplayName(const UE::MVVM::FBindingSource& Source)
+FText GetSourceDisplayName(const FBindingSource& Source)
 {
 	if (!Source.IsValid())
 	{
@@ -44,7 +50,9 @@ static FText GetSourceDisplayName(const UE::MVVM::FBindingSource& Source)
 	return !Source.DisplayName.IsEmpty() ? Source.DisplayName : FText::FromString(Source.Name.ToString());
 }
 
-void SMVVMSourceEntry::Construct(const FArguments& InArgs)
+} // namespace Private
+
+void SSourceEntry::Construct(const FArguments& InArgs)
 {
 	ChildSlot
 	[
@@ -69,11 +77,13 @@ void SMVVMSourceEntry::Construct(const FArguments& InArgs)
 	RefreshSource(InArgs._Source);
 }
 
-void SMVVMSourceEntry::RefreshSource(const UE::MVVM::FBindingSource& Source)
+void SSourceEntry::RefreshSource(const FBindingSource& Source)
 {
-	Image->SetImage(GetSourceIcon(Source));
-	Image->SetColorAndOpacity(GetSourceColor(Source));
-	Label->SetText(GetSourceDisplayName(Source));
+	Image->SetImage(Private::GetSourceIcon(Source));
+	Image->SetColorAndOpacity(Private::GetSourceColor(Source));
+	Label->SetText(Private::GetSourceDisplayName(Source));
 }
 
+} // namespace UE::MVVM
+
 #undef LOCTEXT_NAMESPACE
\ No newline at end of file
diff --git a/Engine/Plugins/Runtime/ModelViewViewModel/Source/ModelViewViewModelEditor/Private/Widgets/SMVVMSourceEntry.h b/Engine/Plugins/Runtime/ModelViewViewModel/Source/ModelViewViewModelEditor/Private/Widgets/SMVVMSourceEntry.h
--- a/Engine/Plugins/Runtime/ModelViewViewModel/Source/ModelViewViewModelEditor/Private/Widgets/SMVVMSourceEntry.h
+++ b/Engine/Plugins/Runtime/ModelViewViewModel/Source/ModelViewViewModelEditor/Private/Widgets/SMVVMSourceEntry.h
@@ -13,10 +13,13 @@
 class SImage;
 class STextBlock;
 
-class SMVVMSourceEntry : public SCompoundWidget
+namespace UE::MVVM
+{
+
+class SSourceEntry : public SCompoundWidget
 {
 public:
-	SLATE_BEGIN_ARGS(SMVVMSourceEntry) :
+	SLATE_BEGIN_ARGS(SSourceEntry) :
 		_TextStyle(&FCoreStyle::Get().GetWidgetStyle<FTextBlockStyle>("NormalText"))
 		{
 		}
@@ -30,4 +33,6 @@ public:
 private:
 	TSharedPtr<STextBlock> Label;
 	TSharedPtr<SImage> Image;
-};
\ No newline at end of file
+};
+
+} // namespace UE::MVVM
diff --git a/Engine/Plugins/Runtime/ModelViewViewModel/Source/ModelViewViewModelEditor/Private/Widgets/SMVVMSourceSelector.h b/Engine/Plugins/Runtime/ModelViewViewModel/Source/ModelViewViewModelEditor/Private/Widgets/SMVVMSourceSelector.h
--- a/Engine/Plugins/Runtime/ModelViewViewModel/Source/ModelViewViewModelEditor/Private/Widgets/SMVVMSourceSelector.h
+++ b/Engine/Plugins/Runtime/ModelViewViewModel/Source/ModelViewViewModelEditor/Private/Widgets/SMVVMSourceSelector.h
@@ -8,20 +8,23 @@
 #include "Widgets/SCompoundWidget.h"
 #include "Widgets/Input/SComboBox.h"
 
-class SMVVMSourceEntry;
+namespace UE::MVVM
+{
+
+class SSourceEntry;
 
-class SMVVMSourceSelector : public SCompoundWidget
+class SSourceSelector : public SCompoundWidget
 {
 public:
-	DECLARE_DELEGATE_OneParam(FSelectionChanged, UE::MVVM::FBindingSource);
+	DECLARE_DELEGATE_OneParam(FSelectionChanged, FBindingSource);
 
-	SLATE_BEGIN_ARGS(SMVVMSourceSelector) :
+	SLATE_BEGIN_ARGS(SSourceSelector) :
 		_TextStyle( &FCoreStyle::Get().GetWidgetStyle<FTextBlockStyle>( "NormalText" ) )
 		{
 		}
 		SLATE_STYLE_ARGUMENT(FTextBlockStyle, TextStyle)
-		SLATE_ATTRIBUTE(UE::MVVM::FBindingSource, SelectedSource)
-		SLATE_ATTRIBUTE(TArray<UE::MVVM::FBindingSource>, AvailableSources)
+		SLATE_ATTRIBUTE(FBindingSource, SelectedSource)
+		SLATE_ATTRIBUTE(TArray<FBindingSource>, AvailableSources)
 		SLATE_EVENT(FSelectionChanged, OnSelectionChanged)
 	SLATE_END_ARGS()
 
@@ -29,19 +32,20 @@ public:
 	void Refresh();
 
 private:
-	void OnComboBoxSelectionChanged(UE::MVVM::FBindingSource Selected, ESelectInfo::Type SelectionType);
+	void OnComboBoxSelectionChanged(FBindingSource Selected, ESelectInfo::Type SelectionType);
 
 	EVisibility GetClearVisibility() const;
 	FReply OnClearSource();
 
 private:
 	FSelectionChanged OnSelectionChanged;
-	TAttribute<TArray<UE::MVVM::FBindingSource>> AvailableSourcesAttribute;
-	TAttribute< UE::MVVM::FBindingSource> SelectedSourceAttribute;
-	TSharedPtr<SComboBox<UE::MVVM::FBindingSource>> SourceComboBox;
-	TArray<UE::MVVM::FBindingSource> AvailableSources;
-	UE::MVVM::FBindingSource SelectedSource;
-	TSharedPtr<SMVVMSourceEntry> SelectedSourceWidget;
+	TAttribute<TArray<FBindingSource>> AvailableSourcesAttribute;
+	TAttribute<FBindingSource> SelectedSourceAttribute;
+	TSharedPtr<SComboBox<FBindingSource>> SourceComboBox;
+	TArray<FBindingSource> AvailableSources;
+	FBindingSource SelectedSource;
+	TSharedPtr<SSourceEntry> SelectedSourceWidget;
 	const FTextBlockStyle* TextStyle = nullptr;
 };
 
+} // namespace UE::MVVM
diff --git a/Engine/Plugins/Runtime/ModelViewViewModel/Source/ModelViewViewModelEditor/Private/Widgets/SMVVMViewBindingListView.cpp b/Engine/Plugins/Runtime/ModelViewViewModel/Source/ModelViewViewModelEditor/Private/Widgets/SMVVMViewBindingListView.cpp
--- a/Engine/Plugins/Runtime/ModelViewViewModel/Source/ModelViewViewModelEditor/Private/Widgets/SMVVMViewBindingListView.cpp
+++ b/Engine/Plugins/Runtime/ModelViewViewModel/Source/ModelViewViewModelEditor/Private/Widgets/SMVVMViewBindingListView.cpp
@@ -50,10 +50,9 @@ public:
 	static FName CompileColumnName;
 	static FName ErrorColumnName;
 	static FName ViewModelColumnName;
-	static FName ViewModelPropertyColumnName;
 	static FName ModeColumnName;
 	static FName WidgetColumnName;
-	static FName WidgetPropertyColumnName;
+	static FName WidgetColumName;
 	static FName UpdateColumnName;
 	static FName ConversionFunctionColumnName;
 	static FName DropDownOptionsColumnName;
@@ -127,31 +126,16 @@ public:
 				.ToolTipText(this, &SMVVMViewBindingListEntryRow::GetErrorToolTip)
 				.OnClicked(this, &SMVVMViewBindingListEntryRow::OnErrorButtonClicked);
 		}
-		else if (ColumnName == ViewModelColumnName)
-		{
-			return SNew(SBox)
-				.Padding(FMargin(2, 0))
-				.VAlign(VAlign_Center)
-				[
-					SAssignNew(ViewModelSourceSelector, SMVVMSourceSelector)
-					.AvailableSources(this, &SMVVMViewBindingListEntryRow::GetAvailableViewModels)
-					.SelectedSource(this, &SMVVMViewBindingListEntryRow::GetSelectedViewModel)
-					.OnSelectionChanged(this, &SMVVMViewBindingListEntryRow::OnViewModelSelectionChanged)
-				];
-		}
-		else if (ColumnName == ViewModelPropertyColumnName)
+		if (ColumnName == ViewModelColumnName)
 		{
 
 			return SNew(SBox)
 				.Padding(FMargin(2, 0))
 				.VAlign(VAlign_Center)
 				[
-					SAssignNew(ViewModelFieldSelector, SMVVMFieldSelector)
-					.SelectedSource(this, &SMVVMViewBindingListEntryRow::GetSelectedViewModel)
-					.AvailableFields(this, &SMVVMViewBindingListEntryRow::GetAvailableViewModelProperties)
+					SAssignNew(ViewModelFieldSelector, UE::MVVM::SFieldSelector, WidgetBlueprint, true)
 					.SelectedField(this, &SMVVMViewBindingListEntryRow::GetSelectedViewModelProperty)
 					.BindingMode(this, &SMVVMViewBindingListEntryRow::GetCurrentBindingMode)
-					.IsSource(true)
 					.OnSelectionChanged(this, &SMVVMViewBindingListEntryRow::OnViewModelPropertySelectionChanged)
 				];
 		}
@@ -191,30 +175,15 @@ public:
 					]
 				];
 		}
-		else if (ColumnName == WidgetColumnName)
+		else if (ColumnName == WidgetColumName)
 		{
 			return SNew(SBox)
 				.Padding(FMargin(2, 0))
 				.VAlign(VAlign_Center)
 				[
-					SAssignNew(WidgetSourceSelector, SMVVMSourceSelector)
-					.AvailableSources(this, &SMVVMViewBindingListEntryRow::GetAvailableWidgets)
-					.SelectedSource(this, &SMVVMViewBindingListEntryRow::GetSelectedWidget)
-					.OnSelectionChanged(this, &SMVVMViewBindingListEntryRow::OnWidgetSelectionChanged)
-				];
-		}
-		else if (ColumnName == WidgetPropertyColumnName)
-		{
-			return SNew(SBox)
-				.Padding(FMargin(2, 0))
-				.VAlign(VAlign_Center)
-				[
-					SAssignNew(WidgetFieldSelector, SMVVMFieldSelector)
-					.SelectedSource(this, &SMVVMViewBindingListEntryRow::GetSelectedWidget)
-					.AvailableFields(this, &SMVVMViewBindingListEntryRow::GetAvailableWidgetProperties)
+					SAssignNew(WidgetFieldSelector, UE::MVVM::SFieldSelector, WidgetBlueprint, false)
 					.SelectedField(this, &SMVVMViewBindingListEntryRow::GetSelectedWidgetProperty)
 					.BindingMode(this, &SMVVMViewBindingListEntryRow::GetCurrentBindingMode)
-					.IsSource(false)
 					.OnSelectionChanged(this, &SMVVMViewBindingListEntryRow::OnWidgetPropertySelectionChanged)
 				];
 		}
@@ -395,60 +364,6 @@ private:
 		return EditorSubsystem->GetAllViewModels(WidgetBlueprint);
 	}
 
-	UE::MVVM::FBindingSource GetSelectedViewModel() const
-	{
-		if (FMVVMBlueprintViewBinding* ViewBinding = GetThisViewBinding())
-		{
-			return UE::MVVM::FBindingSource::CreateForViewModel(WidgetBlueprint, ViewBinding->ViewModelPath.GetViewModelId());
-		}
-
-		return UE::MVVM::FBindingSource();
-	}
-
-	TArray<FMVVMBlueprintPropertyPath> GetAvailableViewModelProperties() const
-	{
-		TArray<FMVVMBlueprintPropertyPath> AvailablePaths;
-
-		if (FMVVMBlueprintViewBinding* ViewBinding = GetThisViewBinding())
-		{
-			UMVVMBlueprintView* BlueprintViewPtr = BlueprintView.Get();
-			if (const FMVVMBlueprintViewModelContext* ViewModel = BlueprintViewPtr->FindViewModel(ViewBinding->ViewModelPath.GetViewModelId()))
-			{
-				UMVVMSubsystem* Subsystem = GEngine->GetEngineSubsystem<UMVVMSubsystem>();
-				TArray<FMVVMAvailableBinding> AvailableBindings = Subsystem->GetAvailableBindings(ViewModel->GetViewModelClass(), WidgetBlueprint->GeneratedClass);
-
-				Algo::Transform(AvailableBindings, AvailablePaths, [ViewModel](const FMVVMAvailableBinding& Binding)
-					{
-						const UClass* Class = ViewModel->GetViewModelClass();
-						FName BindingName = Binding.GetBindingName().ToName();
-						
-						UE::MVVM::FMVVMConstFieldVariant Variant; 
-
-						if (const UFunction* Function = Class->FindFunctionByName(BindingName))
-						{
-							Variant = UE::MVVM::FMVVMConstFieldVariant(Function);
-						}
-						else if (const FProperty* Property = Class->FindPropertyByName(BindingName))
-						{
-							Variant = UE::MVVM::FMVVMConstFieldVariant(Property);
-						}
-
-						FMVVMBlueprintPropertyPath Path;
-
-						if (!Variant.IsEmpty())
-						{
-							Path.SetViewModelId(ViewModel->GetViewModelId());
-							Path.SetBasePropertyPath(Variant);
-						}
-
-						return Path;
-					});
-			}
-		}
-
-		return AvailablePaths;
-	}
-
 	FMVVMBlueprintPropertyPath GetSelectedViewModelProperty() const
 	{
 		if (FMVVMBlueprintViewBinding* ViewBinding = GetThisViewBinding())
@@ -459,82 +374,6 @@ private:
 		return FMVVMBlueprintPropertyPath();
 	}
 
-	void OnViewModelSelectionChanged(UE::MVVM::FBindingSource Source)
-	{
-		if (FMVVMBlueprintViewBinding* ViewBinding = GetThisViewBinding())
-		{
-			const FMVVMBlueprintViewModelContext* ViewModel = BlueprintView.Get()->FindViewModel(Source.ViewModelId);
-
-			UMVVMEditorSubsystem* EditorSubsystem = GEditor->GetEditorSubsystem<UMVVMEditorSubsystem>();
-			EditorSubsystem->SetViewModelForBinding(WidgetBlueprint, *ViewBinding, ViewModel != nullptr ? *ViewModel : FMVVMBlueprintViewModelContext());
-		}
-
-		if (ViewModelFieldSelector.IsValid())
-		{
-			ViewModelFieldSelector->Refresh();
-		}
-	}
-
-	TArray<UE::MVVM::FBindingSource> GetAvailableWidgets() const
-	{
-		UMVVMEditorSubsystem* EditorSubsystem = GEditor->GetEditorSubsystem<UMVVMEditorSubsystem>();
-		return EditorSubsystem->GetBindableWidgets(WidgetBlueprint);
-	}
-	
-	TArray<FMVVMBlueprintPropertyPath> GetAvailableWidgetProperties() const
-	{
-		TArray<FMVVMBlueprintPropertyPath> AvailablePaths;
-
-		if (FMVVMBlueprintViewBinding* ViewBinding = GetThisViewBinding())
-		{
-			UClass* Class = nullptr;
-
-			const FName WidgetName = ViewBinding->WidgetPath.GetWidgetName();
-			if (WidgetName.IsEqual(WidgetBlueprint->GetFName()))
-			{
-				Class = WidgetBlueprint->GeneratedClass;
-			}
-			else if (const UWidget* Widget = WidgetBlueprint->WidgetTree->FindWidget(WidgetName))
-			{
-				Class = Widget->GetClass();
-			}
-
-			if (Class != nullptr)
-			{
-				UMVVMSubsystem* Subsystem = GEngine->GetEngineSubsystem<UMVVMSubsystem>();
-				TArray<FMVVMAvailableBinding> AvailableBindings = Subsystem->GetAvailableBindings(Class, WidgetBlueprint->GeneratedClass);
-
-				Algo::Transform(AvailableBindings, AvailablePaths, [Class, WidgetName](const FMVVMAvailableBinding& Binding)
-					{
-						FName BindingName = Binding.GetBindingName().ToName();
-
-						UE::MVVM::FMVVMConstFieldVariant Variant;
-
-						if (const UFunction* Function = Class->FindFunctionByName(BindingName))
-						{
-							Variant = UE::MVVM::FMVVMConstFieldVariant(Function);
-						}
-						else if (const FProperty* Property = Class->FindPropertyByName(BindingName))
-						{
-							Variant = UE::MVVM::FMVVMConstFieldVariant(Property);
-						}
-
-						FMVVMBlueprintPropertyPath Path;
-
-						if (!Variant.IsEmpty())
-						{
-							Path.SetWidgetName(WidgetName);
-							Path.SetBasePropertyPath(Variant);
-						}
-
-						return Path;
-					});
-			}
-		}
-
-		return AvailablePaths;
-	}
-
 	FMVVMBlueprintPropertyPath GetSelectedWidgetProperty() const
 	{
 		if (FMVVMBlueprintViewBinding* ViewBinding = GetThisViewBinding())
@@ -545,29 +384,6 @@ private:
 		return FMVVMBlueprintPropertyPath();
 	}
 
-	UE::MVVM::FBindingSource GetSelectedWidget() const
-	{
-		if (FMVVMBlueprintViewBinding* ViewBinding = GetThisViewBinding())
-		{
-			return UE::MVVM::FBindingSource::CreateForWidget(WidgetBlueprint, ViewBinding->WidgetPath.GetWidgetName());
-		}
-
-		return UE::MVVM::FBindingSource();
-	}
-
-	void OnWidgetSelectionChanged(UE::MVVM::FBindingSource Source)
-	{
-		if (FMVVMBlueprintViewBinding* ViewBinding = GetThisViewBinding())
-		{
-			GEditor->GetEditorSubsystem<UMVVMEditorSubsystem>()->SetWidgetForBinding(WidgetBlueprint, *ViewBinding, Source.Name);
-		}
-
-		if (WidgetFieldSelector.IsValid())
-		{
-			WidgetFieldSelector->Refresh();
-		}
-	}
-
 	void OnViewModelPropertySelectionChanged(FMVVMBlueprintPropertyPath SelectedField)
 	{
 		if (FMVVMBlueprintViewBinding* ViewBinding = GetThisViewBinding())
@@ -575,7 +391,7 @@ private:
 			if (ViewBinding->ViewModelPath != SelectedField)
 			{
 				UMVVMEditorSubsystem* Subsystem = GEditor->GetEditorSubsystem<UMVVMEditorSubsystem>();
-				Subsystem->SetWidgetPropertyForBinding(WidgetBlueprint, *ViewBinding, SelectedField);
+				Subsystem->SetViewModelPropertyForBinding(WidgetBlueprint, *ViewBinding, SelectedField);
 
 				if (WidgetFieldSelector.IsValid())
 				{
@@ -592,7 +408,7 @@ private:
 			if (ViewBinding->WidgetPath != SelectedField)
 			{
 				UMVVMEditorSubsystem* Subsystem = GEditor->GetEditorSubsystem<UMVVMEditorSubsystem>();
-				Subsystem->SetViewModelPropertyForBinding(WidgetBlueprint, *ViewBinding, SelectedField);
+				Subsystem->SetWidgetPropertyForBinding(WidgetBlueprint, *ViewBinding, SelectedField);
 
 				if (ViewModelFieldSelector.IsValid())
 				{
@@ -768,9 +584,7 @@ private:
 
 	void HandleBlueprintChanged(UBlueprint* Blueprint)
 	{
-		ViewModelSourceSelector->Refresh();
 		ViewModelFieldSelector->Refresh();
-		WidgetSourceSelector->Refresh();
 		WidgetFieldSelector->Refresh();
 	}
 
@@ -778,10 +592,8 @@ private:
 	FMVVMViewBindingListEntryPtr Entry;
 	TWeakObjectPtr<UMVVMBlueprintView> BlueprintView;
 	UWidgetBlueprint* WidgetBlueprint = nullptr;
-	TSharedPtr<SMVVMSourceSelector> ViewModelSourceSelector;
-	TSharedPtr<SMVVMFieldSelector> ViewModelFieldSelector;
-	TSharedPtr<SMVVMSourceSelector> WidgetSourceSelector;
-	TSharedPtr<SMVVMFieldSelector> WidgetFieldSelector;
+	TSharedPtr<UE::MVVM::SFieldSelector> ViewModelFieldSelector;
+	TSharedPtr<UE::MVVM::SFieldSelector> WidgetFieldSelector;
 	TSharedPtr<SWidget> ContextMenuOptionHelper;
 	TSharedPtr<SCustomDialog> ErrorDialog;
 	TArray<TSharedPtr<FText>> ErrorItems;
@@ -792,11 +604,9 @@ private:
 FName SMVVMViewBindingListEntryRow::EnabledColumnName = "Enabled";
 FName SMVVMViewBindingListEntryRow::CompileColumnName = "Compile";
 FName SMVVMViewBindingListEntryRow::ErrorColumnName = "Error";
-FName SMVVMViewBindingListEntryRow::ViewModelColumnName = "ViewModel";
-FName SMVVMViewBindingListEntryRow::ViewModelPropertyColumnName = "ViewModelProperty";
+FName SMVVMViewBindingListEntryRow::ViewModelColumnName = "ViewModelProperty";
 FName SMVVMViewBindingListEntryRow::ModeColumnName = "Mode";
-FName SMVVMViewBindingListEntryRow::WidgetColumnName = "Widget";
-FName SMVVMViewBindingListEntryRow::WidgetPropertyColumnName = "WidgetProperty";
+FName SMVVMViewBindingListEntryRow::WidgetColumName = "WidgetProperty";
 FName SMVVMViewBindingListEntryRow::UpdateColumnName = "Update";
 FName SMVVMViewBindingListEntryRow::ConversionFunctionColumnName = "ConversionFunction";
 FName SMVVMViewBindingListEntryRow::DropDownOptionsColumnName = "DropdownOptions";
@@ -844,19 +654,13 @@ void SMVVMViewBindingListView::Construct(const FArguments& InArgs, TSharedPtr<SM
 				]
 			]
 			+ SHeaderRow::Column(SMVVMViewBindingListEntryRow::ViewModelColumnName)
-			.FillWidth(0.125f)
-			.DefaultLabel(LOCTEXT("ViewModel", "View Model"))
-			+ SHeaderRow::Column(SMVVMViewBindingListEntryRow::ViewModelPropertyColumnName)
-			.FillWidth(0.125f)
+			.FillWidth(0.25f)
 			.DefaultLabel(LOCTEXT("ViewModelProperty", "View Model Property"))
 			+ SHeaderRow::Column(SMVVMViewBindingListEntryRow::ModeColumnName)
 			.FixedWidth(52)
 			.DefaultLabel(LOCTEXT("Mode", "Mode"))
-			+ SHeaderRow::Column(SMVVMViewBindingListEntryRow::WidgetColumnName)
-			.FillWidth(0.125f)
-			.DefaultLabel(LOCTEXT("Widget", "Widget"))
-			+ SHeaderRow::Column(SMVVMViewBindingListEntryRow::WidgetPropertyColumnName)
-			.FillWidth(0.125f)
+			+ SHeaderRow::Column(SMVVMViewBindingListEntryRow::WidgetColumName)
+			.FillWidth(0.25f)
 			.DefaultLabel(LOCTEXT("WidgetProperty", "Widget Property"))
 			+ SHeaderRow::Column(SMVVMViewBindingListEntryRow::UpdateColumnName)
 			.FillWidth(0.05f)
diff --git a/Engine/Plugins/Runtime/ModelViewViewModel/Source/ModelViewViewModelEditor/Private/Widgets/SMVVMViewModelBindingListWidget.cpp b/Engine/Plugins/Runtime/ModelViewViewModel/Source/ModelViewViewModelEditor/Private/Widgets/SMVVMViewModelBindingListWidget.cpp
--- a/Engine/Plugins/Runtime/ModelViewViewModel/Source/ModelViewViewModelEditor/Private/Widgets/SMVVMViewModelBindingListWidget.cpp
+++ b/Engine/Plugins/Runtime/ModelViewViewModel/Source/ModelViewViewModelEditor/Private/Widgets/SMVVMViewModelBindingListWidget.cpp
@@ -2,10 +2,13 @@
 
 #include "SMVVMViewModelBindingListWidget.h"
 
+#include "Algo/Transform.h"
 #include "Bindings/MVVMBindingHelper.h"
+#include "Components/Widget.h"
 #include "Editor.h"
 #include "Engine/Engine.h"
 #include "FieldNotification/IFieldValueChanged.h"
+#include "MVVMBlueprintViewModelContext.h"
 #include "MVVMEditorSubsystem.h"
 #include "MVVMSubsystem.h"
 #include "MVVMViewModelBase.h"
@@ -14,26 +17,43 @@
 #include "WidgetBlueprint.h"
 #include "Widgets/PropertyViewer/SPropertyViewer.h"
 
-#define LOCTEXT_NAMESPACE "SViewModelBindingListWidget"
+#define LOCTEXT_NAMESPACE "SSourceBindingList"
+
+using UE::PropertyViewer::SPropertyViewer;
 
 namespace UE::MVVM
 {
 
-/** */
-FFieldIterator_ViewModel::FFieldIterator_ViewModel(const UWidgetBlueprint* InWidgetBlueprint)
+	FFieldIterator_Bindable::FFieldIterator_Bindable(const UWidgetBlueprint* InWidgetBlueprint, EFieldVisibility InVisibilityFlags)
 	: WidgetBlueprint(InWidgetBlueprint)
+	, FieldVisibilityFlags(InVisibilityFlags)
 { }
 
 
-TArray<FFieldVariant> FFieldIterator_ViewModel::GetFields(const UStruct* Struct) const
+TArray<FFieldVariant> FFieldIterator_Bindable::GetFields(const UStruct* Struct) const
 {
 	TArray<FFieldVariant> Result;
 
-	auto AddResult = [&Result, Struct](const TArray<FMVVMAvailableBinding>& AvailableBindingsList)
+	auto AddResult = [Flags = FieldVisibilityFlags, &Result, Struct](const TArray<FMVVMAvailableBinding>& AvailableBindingsList)
 	{
 		Result.Reserve(AvailableBindingsList.Num());
 		for (const FMVVMAvailableBinding& Value : AvailableBindingsList)
 		{
+			if (EnumHasAllFlags(Flags, EFieldVisibility::Readable) && !Value.IsReadable())
+			{
+				continue;
+			} 
+
+			if (EnumHasAllFlags(Flags, EFieldVisibility::Writable) && !Value.IsWritable())
+			{
+				continue;
+			}
+			
+			if (EnumHasAllFlags(Flags, EFieldVisibility::Notify) && !Value.HasNotify())
+			{
+				continue;
+			}
+
 			FMVVMFieldVariant FieldVariant = BindingHelper::FindFieldByName(Struct, Value.GetBindingName());
 			if (FieldVariant.IsFunction())
 			{
@@ -46,7 +66,6 @@ TArray<FFieldVariant> FFieldIterator_ViewModel::GetFields(const UStruct* Struct)
 		}
 	};
 
-
 	if (const UClass* Class = Cast<const UClass>(Struct))
 	{
 		const UWidgetBlueprint* WidgetBlueprintPtr = WidgetBlueprint.Get();
@@ -79,7 +98,7 @@ TArray<FFieldVariant> FFieldIterator_ViewModel::GetFields(const UStruct* Struct)
 			{
 				return false;
 			}
-			else if(A.IsUObject())
+			else if (A.IsUObject())
 			{
 				return true;
 			}
@@ -90,26 +109,24 @@ TArray<FFieldVariant> FFieldIterator_ViewModel::GetFields(const UStruct* Struct)
 			return A.GetFName().LexicalLess(B.GetFName());
 		});
 
-
 	return Result;
 }
 
-
 /** */
-void SViewModelBindingListWidget::Construct(const FArguments& InArgs, const UWidgetBlueprint* WidgetBlueprint)
+void SSourceBindingList::Construct(const FArguments& InArgs, const UWidgetBlueprint* WidgetBlueprint)
 {
-	ViewModelFieldIterator = MakeUnique<FFieldIterator_ViewModel>(WidgetBlueprint);
+	FieldIterator = MakeUnique<FFieldIterator_Bindable>(WidgetBlueprint, InArgs._FieldVisibilityFlags);
+
+	OnDoubleClicked = InArgs._OnDoubleClicked;
 
-	PropertyViewer = SNew(UE::PropertyViewer::SPropertyViewer)
-		.PropertyVisibility(UE::PropertyViewer::SPropertyViewer::EPropertyVisibility::Hidden)
+	PropertyViewer = SNew(SPropertyViewer)
+		.PropertyVisibility(SPropertyViewer::EPropertyVisibility::Hidden)
 		.bShowFieldIcon(true)
 		.bSanitizeName(true)
-		.FieldIterator(ViewModelFieldIterator.Get());
-
-	if (InArgs._ViewModel.Class.Get())
-	{
-		SetViewModels(MakeArrayView(&InArgs._ViewModel, 1));
-	}
+		.FieldIterator(FieldIterator.Get())
+		.bShowSearchBox(false)
+		.OnSelectionChanged(this, &SSourceBindingList::HandleSelectionChanged)
+		.OnDoubleClicked(this, &SSourceBindingList::HandleDoubleClicked);
 
 	ChildSlot
 	[
@@ -117,35 +134,123 @@ void SViewModelBindingListWidget::Construct(const FArguments& InArgs, const UWid
 	];
 }
 
+void SSourceBindingList::Clear()
+{
+	if (PropertyViewer)
+	{
+		PropertyViewer->RemoveAll();
+	}
+}
+
+void SSourceBindingList::AddSource(UClass* Class, FName Name, FGuid Guid)
+{
+	FBindingSource Source;
+	Source.Class = Class;
+	Source.Name = Name;
+	Source.ViewModelId = Guid;
+
+	AddSources(MakeArrayView(&Source, 1));
+}
 
-void SViewModelBindingListWidget::SetViewModel(UClass* Class, FName Name, FGuid Guid)
+void SSourceBindingList::AddWidgets(TArrayView<const UWidget*> InWidgets)
 {
-	FViewModel ViewModel;
-	ViewModel.Class = Class;
-	ViewModel.ViewModelId = Guid;
-	ViewModel.Name = Name;
-	SetViewModels(MakeArrayView(&ViewModel, 1));
+	TArray<FBindingSource, TInlineAllocator<16>> NewSources;
+	NewSources.Reserve(InWidgets.Num());
+
+	for (const UWidget* Widget : InWidgets)
+	{
+		FBindingSource& Source = NewSources.AddDefaulted_GetRef();
+		Source.Class = Widget->GetClass();
+		Source.Name = Widget->GetFName();
+		Source.DisplayName = Widget->GetLabelText();
+	}
+
+	AddSources(NewSources);
 }
 
+void SSourceBindingList::AddViewModels(TArrayView<const FMVVMBlueprintViewModelContext> InViewModels)
+{
+	TArray<FBindingSource, TInlineAllocator<16>> NewSources;
+	NewSources.Reserve(InViewModels.Num());
+
+	for (const FMVVMBlueprintViewModelContext& ViewModelContext : InViewModels)
+	{
+		FBindingSource& Source = NewSources.AddDefaulted_GetRef();
+		Source.Class = ViewModelContext.GetViewModelClass();
+		Source.ViewModelId = ViewModelContext.GetViewModelId();
+	}
+
+	AddSources(NewSources);
+}
 
-void SViewModelBindingListWidget::SetViewModels(TArrayView<const FViewModel> InViewModels)
+void SSourceBindingList::AddSources(TArrayView<const FBindingSource> InSources)
 {
-	ViewModels = InViewModels;
+	Algo::Transform(InSources, Sources, [](const FBindingSource& Source)
+		{
+			return TPair<FBindingSource, SPropertyViewer::FHandle>(Source, SPropertyViewer::FHandle());
+		});
+
 	if (PropertyViewer)
 	{
-		PropertyViewer->RemoveAll();
-		for (const FViewModel& ViewModel : ViewModels)
+		for (TPair<FBindingSource, SPropertyViewer::FHandle>& Source : Sources)
 		{
-			if (ViewModel.Class && ViewModel.Class->ImplementsInterface(UNotifyFieldValueChanged::StaticClass()))
+			const UClass* Class = Source.Key.Class;
+			if (Class && Class->ImplementsInterface(UNotifyFieldValueChanged::StaticClass()))
 			{
-				PropertyViewer->AddContainer(ViewModel.Class.Get());
+				Source.Value = PropertyViewer->AddContainer(Class);
 			}
 		}
 	}
 }
 
+FMVVMBlueprintPropertyPath SSourceBindingList::CreateBlueprintPropertyPath(SPropertyViewer::FHandle Handle, TArrayView<const FFieldVariant> FieldPath) const
+{
+	const TPair<FBindingSource, SPropertyViewer::FHandle>* Source = Sources.FindByPredicate(
+		[Handle](const TPair<FBindingSource, SPropertyViewer::FHandle>& Source)
+		{
+			return Source.Value == Handle;
+		});
+		
+	FMVVMBlueprintPropertyPath PropertyPath;
+	if (Source != nullptr)
+	{
+		if (Source->Key.ViewModelId.IsValid())
+		{
+			PropertyPath.SetViewModelId(Source->Key.ViewModelId);
+		}
+		else
+		{
+			PropertyPath.SetWidgetName(Source->Key.Name);
+		}
+
+		PropertyPath.ResetBasePropertyPath();
+	}
 
-void SViewModelBindingListWidget::SetRawFilterText(const FText& InFilterText)
+	// TODO (sebastiann): FMVVMBlueprintPropertyPath doesn't support long paths yet
+	FFieldVariant LastField = FieldPath.Num() > 0 ? FieldPath.Last() : FFieldVariant();
+	if (LastField.IsValid())
+	{
+		PropertyPath.SetBasePropertyPath(FMVVMConstFieldVariant(LastField));
+	}
+
+	return PropertyPath;
+}
+
+void SSourceBindingList::HandleSelectionChanged(SPropertyViewer::FHandle Handle, TArrayView<const FFieldVariant> FieldPath, ESelectInfo::Type SelectionType)
+{
+	SelectedPath = CreateBlueprintPropertyPath(Handle, FieldPath);
+}
+
+void SSourceBindingList::HandleDoubleClicked(SPropertyViewer::FHandle Handle, TArrayView<const FFieldVariant> FieldPath)
+{
+	if (OnDoubleClicked.IsBound())
+	{
+		FMVVMBlueprintPropertyPath ClickedPath = CreateBlueprintPropertyPath(Handle, FieldPath);
+		OnDoubleClicked.Execute(ClickedPath);
+	}
+}
+
+void SSourceBindingList::SetRawFilterText(const FText& InFilterText)
 {
 	if (PropertyViewer)
 	{
@@ -153,6 +258,11 @@ void SViewModelBindingListWidget::SetRawFilterText(const FText& InFilterText)
 	}
 }
 
+FMVVMBlueprintPropertyPath SSourceBindingList::GetSelectedProperty() const
+{
+	return SelectedPath;
+}
+
 } // namespace UE::MVVM
 
 #undef LOCTEXT_NAMESPACE
\ No newline at end of file
diff --git a/Engine/Plugins/Runtime/ModelViewViewModel/Source/ModelViewViewModelEditor/Private/Widgets/SMVVMViewModelBindingListWidget.h b/Engine/Plugins/Runtime/ModelViewViewModel/Source/ModelViewViewModelEditor/Private/Widgets/SMVVMViewModelBindingListWidget.h
--- a/Engine/Plugins/Runtime/ModelViewViewModel/Source/ModelViewViewModelEditor/Private/Widgets/SMVVMViewModelBindingListWidget.h
+++ b/Engine/Plugins/Runtime/ModelViewViewModel/Source/ModelViewViewModelEditor/Private/Widgets/SMVVMViewModelBindingListWidget.h
@@ -4,59 +4,90 @@
 
 #include "CoreMinimal.h"
 #include "Framework/PropertyViewer/IFieldIterator.h"
+#include "Misc/EnumClassFlags.h"
+#include "MVVMPropertyPath.h"
 #include "Templates/SubclassOf.h"
 #include "UObject/WeakObjectPtr.h"
 #include "Widgets/SCompoundWidget.h"
+#include "Widgets/PropertyViewer/SPropertyViewer.h"
 
 class UWidgetBlueprint;
-namespace UE::PropertyViewer
-{
-class SPropertyViewer;
-}
+struct FMVVMBlueprintViewModelContext;
+class UWidgetBlueprint;
 
 namespace UE::MVVM
 {
+	struct FBindingSource;
 
-/**
- * 
- */
-class FFieldIterator_ViewModel : public UE::PropertyViewer::FFieldIterator_BlueprintVisible
-{
-public:
-	FFieldIterator_ViewModel(const UWidgetBlueprint* WidgetBlueprint);
-	virtual TArray<FFieldVariant> GetFields(const UStruct*) const override;
-
-private:
-	TWeakObjectPtr<const UWidgetBlueprint> WidgetBlueprint;
-};
+	enum class EFieldVisibility : uint8
+	{
+		None = 0,
+		Readable = 1 << 1,
+		Writable = 1 << 2,
+		Notify = 1 << 3,
+		All = 0xff
+	};
 
+	ENUM_CLASS_FLAGS(EFieldVisibility)
 
-/** 
- * 
- */
-class SViewModelBindingListWidget : public SCompoundWidget
-{
-public:
-	struct FViewModel
+	/**
+	 * 
+	 */
+	class FFieldIterator_Bindable : public UE::PropertyViewer::FFieldIterator_BlueprintVisible
 	{
-		TSubclassOf<UObject> Class;
-		FGuid ViewModelId;
-		FName Name;
+	public:
+		FFieldIterator_Bindable(const UWidgetBlueprint* WidgetBlueprint, EFieldVisibility InVisibilityFlags);
+		virtual TArray<FFieldVariant> GetFields(const UStruct*) const override;
+
+	private:
+		TWeakObjectPtr<const UWidgetBlueprint> WidgetBlueprint; 
+		EFieldVisibility FieldVisibilityFlags = EFieldVisibility::All;
 	};
-	SLATE_BEGIN_ARGS(SViewModelBindingListWidget) {}
-		SLATE_ARGUMENT(FViewModel, ViewModel)
-	SLATE_END_ARGS()
 
-	void Construct(const FArguments& InArgs, const UWidgetBlueprint* WidgetBlueprint);
-	void SetViewModel(UClass* Class, FName Name, FGuid Guid);
-	void SetViewModels(TArrayView<const FViewModel> ViewModels);
+	/** 
+	 * 
+	 */
+	class SSourceBindingList : public SCompoundWidget
+	{
+	public:
+
+		DECLARE_DELEGATE_OneParam(FOnDoubleClicked, const FMVVMBlueprintPropertyPath&);
+	
+		SLATE_BEGIN_ARGS(SSourceBindingList) : 
+			_FieldVisibilityFlags(EFieldVisibility::All)
+			{}
+			SLATE_ARGUMENT(EFieldVisibility, FieldVisibilityFlags)
+			SLATE_EVENT(FOnDoubleClicked, OnDoubleClicked)
+		SLATE_END_ARGS()
 
-	void SetRawFilterText(const FText& InFilterText);
+		void Construct(const FArguments& InArgs, const UWidgetBlueprint* WidgetBlueprint);
 
-private:
-	TUniquePtr<FFieldIterator_ViewModel> ViewModelFieldIterator;
-	TArray<FViewModel> ViewModels;
-	TSharedPtr<UE::PropertyViewer::SPropertyViewer> PropertyViewer;
-};
+		void Clear();
+
+		void AddSource(UClass* Class, FName Name, FGuid Guid);
+		void AddSources(TArrayView<const FBindingSource> InSources);
+
+		void AddWidgets(TArrayView<const UWidget*> Widgets);
+		void AddViewModels(TArrayView<const FMVVMBlueprintViewModelContext> ViewModels);
+
+		FMVVMBlueprintPropertyPath GetSelectedProperty() const;
+
+		void SetRawFilterText(const FText& InFilterText);
+
+	private:
+		using SPropertyViewer = UE::PropertyViewer::SPropertyViewer;
+
+		void HandleSelectionChanged(SPropertyViewer::FHandle ViewModel, TArrayView<const FFieldVariant> Path, ESelectInfo::Type SelectionType);
+		void HandleDoubleClicked(SPropertyViewer::FHandle ViewModel, TArrayView<const FFieldVariant> Path);
+
+		FMVVMBlueprintPropertyPath CreateBlueprintPropertyPath(SPropertyViewer::FHandle Handle, TArrayView<const FFieldVariant> Path) const;
+
+	private:
+		FOnDoubleClicked OnDoubleClicked;
+		TUniquePtr<FFieldIterator_Bindable> FieldIterator;
+		TArray<TPair<FBindingSource, SPropertyViewer::FHandle>> Sources;
+		FMVVMBlueprintPropertyPath SelectedPath;
+		TSharedPtr<SPropertyViewer> PropertyViewer;
+	};
 
 } //namespace UE::MVVM
\ No newline at end of file
diff --git a/Engine/Plugins/Runtime/ModelViewViewModel/Source/ModelViewViewModelEditor/Private/Widgets/SMVVMViewModelContextListWidget.cpp b/Engine/Plugins/Runtime/ModelViewViewModel/Source/ModelViewViewModelEditor/Private/Widgets/SMVVMViewModelContextListWidget.cpp
--- a/Engine/Plugins/Runtime/ModelViewViewModel/Source/ModelViewViewModelEditor/Private/Widgets/SMVVMViewModelContextListWidget.cpp
+++ b/Engine/Plugins/Runtime/ModelViewViewModel/Source/ModelViewViewModelEditor/Private/Widgets/SMVVMViewModelContextListWidget.cpp
@@ -16,7 +16,7 @@
 #include "Styling/MVVMEditorStyle.h"
 #include "Styling/SlateIconFinder.h"
 #include "Types/MVVMFieldVariant.h"
-#include "Widgets/SMVVMFieldSelector.h"
+#include "Widgets/SMVVMFieldEntry.h"
 #include "Widgets/Input/SComboButton.h"
 #include "Widgets/Input/SEditableTextBox.h"
 #include "Widgets/Images/SImage.h"
@@ -29,7 +29,7 @@
 
 namespace UE::MVVM::Private
 {
-	static bool IsPropertyTypeChildOf(const FProperty* Property, const UClass* ParentClass)
+	bool IsPropertyTypeChildOf(const FProperty* Property, const UClass* ParentClass)
 	{
 		if (const FObjectPropertyBase* ObjectProperty = CastField<FObjectPropertyBase>(Property))
 		{
@@ -42,7 +42,7 @@ namespace UE::MVVM::Private
 		return false;
 	}
 
-	static bool HasBindableViewModelFieldRecursive(UStruct* InStruct, const UClass* InParentClass, TSet<UStruct*>& VisitedSet, uint32 RecursionDepth)
+	bool HasBindableViewModelFieldRecursive(UStruct* InStruct, const UClass* InParentClass, TSet<UStruct*>& VisitedSet, uint32 RecursionDepth)
 	{
 		if (RecursionDepth > 10)
 		{
@@ -388,7 +388,7 @@ namespace UE::MVVM::Private
 				FMVVMBlueprintPropertyPath Path;
 				Path.SetBasePropertyPath(Field);
 
-				return SNew(SMVVMFieldEntry)
+				return SNew(SFieldEntry)
 					.Field(Path);
 			};
 
diff --git a/Engine/Plugins/Runtime/ModelViewViewModel/Source/ModelViewViewModelEditor/Private/Widgets/SMVVMViewModelPanel.cpp b/Engine/Plugins/Runtime/ModelViewViewModel/Source/ModelViewViewModelEditor/Private/Widgets/SMVVMViewModelPanel.cpp
--- a/Engine/Plugins/Runtime/ModelViewViewModel/Source/ModelViewViewModelEditor/Private/Widgets/SMVVMViewModelPanel.cpp
+++ b/Engine/Plugins/Runtime/ModelViewViewModel/Source/ModelViewViewModelEditor/Private/Widgets/SMVVMViewModelPanel.cpp
@@ -24,7 +24,6 @@
 #include "Widgets/Layout/SBox.h"
 #include "Widgets/PropertyViewer/SPropertyViewer.h"
 #include "Widgets/SMVVMSelectViewModel.h"
-#include "Widgets/SMVVMViewModelBindingListWidget.h"
 
 #include "SPositiveActionButton.h"
 
@@ -41,7 +40,7 @@ void SMVVMViewModelPanel::Construct(const FArguments& InArgs, TSharedPtr<FWidget
 
 	WeakBlueprintEditor = WidgetBlueprintEditor;
 	WeakBlueprintView = CurrentBlueprintView;
-	ViewModelFieldIterator = MakeUnique<FFieldIterator_ViewModel>(WidgetBlueprint);
+	FieldIterator = MakeUnique<FFieldIterator_Bindable>(WidgetBlueprint, EFieldVisibility::All);
 
 	if (CurrentBlueprintView)
 	{
@@ -54,7 +53,7 @@ void SMVVMViewModelPanel::Construct(const FArguments& InArgs, TSharedPtr<FWidget
 		.bShowSearchBox(true)
 		.bShowFieldIcon(true)
 		.bSanitizeName(true)
-		.FieldIterator(ViewModelFieldIterator.Get())
+		.FieldIterator(FieldIterator.Get())
 		.SearchBoxPreSlot()
 		[
 			SAssignNew(AddMenuButton, SPositiveActionButton)
@@ -103,7 +102,6 @@ void SMVVMViewModelPanel::FillViewModel()
 	if (ViewModelTreeView)
 	{
 		ViewModelTreeView->RemoveAll();
-		ViewModelHandles.Reset();
 
 		if (UMVVMBlueprintView* View = WeakBlueprintView.Get())
 		{
@@ -112,19 +110,12 @@ void SMVVMViewModelPanel::FillViewModel()
 				if (ViewModelContext.GetViewModelClass())
 				{
 					UObject* ViewModelInstance = ViewModelContext.GetViewModelClass()->GetDefaultObject();
-
-					FViewModelHandle ViewModelHandle;
-					ViewModelHandle.Key = ViewModelContext.GetViewModelId();
-					ViewModelHandle.Value = ViewModelTreeView->AddInstance(ViewModelInstance, ViewModelContext.GetDisplayName());
-					ViewModelHandles.Add(ViewModelHandle);
+					ViewModelTreeView->AddInstance(ViewModelInstance, ViewModelContext.GetDisplayName());
 				}
 				else
 				{
-					FViewModelHandle ViewModelHandle;
-					ViewModelHandle.Key = ViewModelContext.GetViewModelId();
 					// Find a way to show context that are not valid anymore
 					//ViewModelHandle.Value = ViewModelTreeView->AddInstance(ViewModelInstance, ViewModelContext.GetDisplayName());
-					ViewModelHandles.Add(ViewModelHandle);
 				}
 			}
 		}
@@ -193,7 +184,6 @@ void SMVVMViewModelPanel::HandleCancelAddMenu()
 	}
 }
 
-
 void SMVVMViewModelPanel::HandleAddMenuViewModel(const UClass* SelectedClass)
 {
 	if (AddMenuButton)
@@ -224,49 +214,6 @@ bool SMVVMViewModelPanel::HandleCanEditViewmodelList() const
 	return false;
 }
 
-
-//bool SMVVMViewModelPanel::HandleViewModelRenamed(const SViewModelBindingListWidget::FViewModel& ViewModel, const FText& RenameTo, bool bCommit, FText& OutErrorMessage)
-//{
-//	if (RenameTo.IsEmptyOrWhitespace())
-//	{
-//		OutErrorMessage = LOCTEXT("EmptyViewmodelName", "Empty viewmodel name");
-//		return false;
-//	}
-//
-//	const FString& NewNameString = RenameTo.ToString();
-//	if (NewNameString.Len() >= NAME_SIZE)
-//	{
-//		OutErrorMessage = LOCTEXT("ViewmodelNameTooLong", "Viewmodel name is too long");
-//		return false;
-//	}
-//
-//	FString GeneratedName = SlugStringForValidName(NewNameString);
-//	if (GeneratedName.IsEmpty())
-//	{
-//		OutErrorMessage = LOCTEXT("EmptyViewmodelName", "Empty viewmodel name");
-//		return false;
-//	}
-//
-//	const FName GeneratedFName(*GeneratedName);
-//	check(GeneratedFName.IsValidXName(INVALID_OBJECTNAME_CHARACTERS));
-//
-//	if (TSharedPtr<FWidgetBlueprintEditor> WidgetBlueprintEditor = WeakBlueprintEditor.Pin())
-//	{
-//		if (UWidgetBlueprint* WidgetBP = WidgetBlueprintEditor->GetWidgetBlueprintObj())
-//		{
-//			if (bCommit)
-//			{
-//				return GEditor->GetEditorSubsystem<UMVVMEditorSubsystem>()->RenameViewModel(WidgetBP, ViewModel.Name, *NewNameString, OutErrorMessage);
-//			}
-//			else
-//			{
-//				return GEditor->GetEditorSubsystem<UMVVMEditorSubsystem>()->VerifyViewModelRename(WidgetBP, ViewModel.Name, *NewNameString, OutErrorMessage);
-//			}
-//		}
-//	}
-//	return false;
-//}
-
 } // namespace
 
 #undef LOCTEXT_NAMESPACE
\ No newline at end of file
diff --git a/Engine/Plugins/Runtime/ModelViewViewModel/Source/ModelViewViewModelEditor/Private/Widgets/SMVVMViewModelPanel.h b/Engine/Plugins/Runtime/ModelViewViewModel/Source/ModelViewViewModelEditor/Private/Widgets/SMVVMViewModelPanel.h
--- a/Engine/Plugins/Runtime/ModelViewViewModel/Source/ModelViewViewModelEditor/Private/Widgets/SMVVMViewModelPanel.h
+++ b/Engine/Plugins/Runtime/ModelViewViewModel/Source/ModelViewViewModelEditor/Private/Widgets/SMVVMViewModelPanel.h
@@ -15,41 +15,36 @@ class UMVVMBlueprintView;
 
 namespace UE::MVVM
 {
-
-/**
- * 
- */
-class SMVVMViewModelPanel : public SCompoundWidget
-{
-public:
-	SLATE_BEGIN_ARGS(SMVVMViewModelPanel) {}
-	SLATE_END_ARGS()
-
-	void Construct(const FArguments& InArgs, TSharedPtr<FWidgetBlueprintEditor> Editor);
-	virtual ~SMVVMViewModelPanel();
-
-private:
-	void HandleViewUpdated(UBlueprintExtension* Extension);
-	void HandleViewModelsUpdated();
-	//bool HandleViewModelRenamed(const SViewModelBindingListWidget::FViewModel& ViewModel, const FText& RenameTo, bool bCommit, FText& OutErrorMessage);
-	TSharedRef<SWidget> MakeAddMenu();
-	void HandleCancelAddMenu();
-	void HandleAddMenuViewModel(const UClass* SelectedClass);
-	bool HandleCanEditViewmodelList() const;
-
-	void FillViewModel();
-
-private:
-	TSharedPtr<SPositiveActionButton> AddMenuButton;
-	TSharedPtr<UE::PropertyViewer::SPropertyViewer> ViewModelTreeView;
-
-	using FViewModelHandle = TPair<FGuid, UE::PropertyViewer::SPropertyViewer::FHandle>;
-	TArray<FViewModelHandle> ViewModelHandles;
-
-	TWeakPtr<FWidgetBlueprintEditor> WeakBlueprintEditor;
-	TWeakObjectPtr<UMVVMBlueprintView> WeakBlueprintView;
-	TUniquePtr<FFieldIterator_ViewModel> ViewModelFieldIterator;
-	FDelegateHandle ViewModelsUpdatedHandle;
-};
+	/**
+	 * 
+	 */
+	class SMVVMViewModelPanel : public SCompoundWidget
+	{
+	public:
+		SLATE_BEGIN_ARGS(SMVVMViewModelPanel) {}
+		SLATE_END_ARGS()
+
+		void Construct(const FArguments& InArgs, TSharedPtr<FWidgetBlueprintEditor> Editor);
+		virtual ~SMVVMViewModelPanel();
+
+	private:
+		void HandleViewUpdated(UBlueprintExtension* Extension);
+		void HandleViewModelsUpdated();
+		TSharedRef<SWidget> MakeAddMenu();
+		void HandleCancelAddMenu();
+		void HandleAddMenuViewModel(const UClass* SelectedClass);
+		bool HandleCanEditViewmodelList() const;
+
+		void FillViewModel();
+
+	private:
+		TSharedPtr<SPositiveActionButton> AddMenuButton;
+		TSharedPtr<UE::PropertyViewer::SPropertyViewer> ViewModelTreeView;
+		TUniquePtr<FFieldIterator_Bindable> FieldIterator;
+
+		TWeakPtr<FWidgetBlueprintEditor> WeakBlueprintEditor;
+		TWeakObjectPtr<UMVVMBlueprintView> WeakBlueprintView;
+		FDelegateHandle ViewModelsUpdatedHandle;
+	};
 
 } // namespace
diff --git a/Engine/Plugins/Runtime/ModelViewViewModel/Source/ModelViewViewModelEditor/Public/MVVMEditorSubsystem.h b/Engine/Plugins/Runtime/ModelViewViewModel/Source/ModelViewViewModelEditor/Public/MVVMEditorSubsystem.h
--- a/Engine/Plugins/Runtime/ModelViewViewModel/Source/ModelViewViewModelEditor/Public/MVVMEditorSubsystem.h
+++ b/Engine/Plugins/Runtime/ModelViewViewModel/Source/ModelViewViewModelEditor/Public/MVVMEditorSubsystem.h
@@ -52,9 +52,7 @@ public:
 
 	void SetSourceToDestinationConversionFunction(UWidgetBlueprint* WidgetBlueprint, FMVVMBlueprintViewBinding& Binding, const UFunction* ConversionFunction);
 	void SetDestinationToSourceConversionFunction(UWidgetBlueprint* WidgetBlueprint, FMVVMBlueprintViewBinding& Binding, const UFunction* ConversionFunction);
-	void SetWidgetForBinding(UWidgetBlueprint* WidgetBlueprint, FMVVMBlueprintViewBinding& Binding, FName Widget);
 	void SetWidgetPropertyForBinding(UWidgetBlueprint* WidgetBlueprint, FMVVMBlueprintViewBinding& Binding, FMVVMBlueprintPropertyPath Field);
-	void SetViewModelForBinding(UWidgetBlueprint* WidgetBlueprint, FMVVMBlueprintViewBinding& Binding, FMVVMBlueprintViewModelContext ViewModel);
 	void SetViewModelPropertyForBinding(UWidgetBlueprint* WidgetBlueprint, FMVVMBlueprintViewBinding& Binding, FMVVMBlueprintPropertyPath Field);
 	void SetUpdateModeForBinding(UWidgetBlueprint* WidgetBlueprint, FMVVMBlueprintViewBinding& Binding, EMVVMViewBindingUpdateMode Mode);
 	void SetBindingTypeForBinding(UWidgetBlueprint* WidgetBlueprint, FMVVMBlueprintViewBinding& Binding, EMVVMBindingMode Type);
