Missed file in previous submit.

#rnx

diff --git a/Engine/Plugins/Runtime/ModelViewViewModel/Source/ModelViewViewModelEditor/Private/Widgets/SMVVMSourceSelector.cpp b/Engine/Plugins/Runtime/ModelViewViewModel/Source/ModelViewViewModelEditor/Private/Widgets/SMVVMSourceSelector.cpp
--- a/Engine/Plugins/Runtime/ModelViewViewModel/Source/ModelViewViewModelEditor/Private/Widgets/SMVVMSourceSelector.cpp
+++ b/Engine/Plugins/Runtime/ModelViewViewModel/Source/ModelViewViewModelEditor/Private/Widgets/SMVVMSourceSelector.cpp
@@ -9,9 +9,10 @@
 
 #define LOCTEXT_NAMESPACE "MVVMSourceSelector"
 
-using namespace UE::MVVM;
+namespace UE::MVVM
+{
 
-void SMVVMSourceSelector::Construct(const FArguments& Args)
+void SSourceSelector::Construct(const FArguments& Args)
 {
 	TextStyle = Args._TextStyle;
 	AvailableSourcesAttribute = Args._AvailableSources;
@@ -32,14 +33,14 @@ void SMVVMSourceSelector::Construct(const FArguments& Args)
 			.OptionsSource(&AvailableSources)
 			.OnGenerateWidget_Lambda([this](FBindingSource Source)
 			{
-				return SNew(SMVVMSourceEntry)
+				return SNew(SSourceEntry)
 					.TextStyle(TextStyle)
 					.Source(Source);
 			})
-			.OnSelectionChanged(this, &SMVVMSourceSelector::OnComboBoxSelectionChanged)
+			.OnSelectionChanged(this, &SSourceSelector::OnComboBoxSelectionChanged)
 			.InitiallySelectedItem(SelectedSource)
 			[
-				SAssignNew(SelectedSourceWidget, SMVVMSourceEntry)
+				SAssignNew(SelectedSourceWidget, SSourceEntry)
 				.TextStyle(TextStyle)
 				.Source(SelectedSource)
 			]
@@ -52,15 +53,15 @@ void SMVVMSourceSelector::Construct(const FArguments& Args)
 			SNew(SSimpleButton)
 			.Icon(FAppStyle::Get().GetBrush("Icons.X"))
 			.ToolTipText(LOCTEXT("ClearField", "Clear source selection."))
-			.Visibility(this, &SMVVMSourceSelector::GetClearVisibility)
-			.OnClicked(this, &SMVVMSourceSelector::OnClearSource)
+			.Visibility(this, &SSourceSelector::GetClearVisibility)
+			.OnClicked(this, &SSourceSelector::OnClearSource)
 		]
 	];
 
 	OnSelectionChanged = Args._OnSelectionChanged;
 }
 
-void SMVVMSourceSelector::OnComboBoxSelectionChanged(FBindingSource Selected, ESelectInfo::Type SelectionType)
+void SSourceSelector::OnComboBoxSelectionChanged(FBindingSource Selected, ESelectInfo::Type SelectionType)
 {
 	SelectedSource = Selected;
 
@@ -69,7 +70,7 @@ void SMVVMSourceSelector::OnComboBoxSelectionChanged(FBindingSource Selected, ES
 	OnSelectionChanged.ExecuteIfBound(Selected);
 }
 
-void SMVVMSourceSelector::Refresh()
+void SSourceSelector::Refresh()
 {
 	SelectedSource = SelectedSourceAttribute.Get();
 	AvailableSources = AvailableSourcesAttribute.Get();
@@ -81,12 +82,12 @@ void SMVVMSourceSelector::Refresh()
 	}
 }
 
-EVisibility SMVVMSourceSelector::GetClearVisibility() const
+EVisibility SSourceSelector::GetClearVisibility() const
 {
 	return SelectedSource.IsValid() ? EVisibility::Visible : EVisibility::Collapsed;
 }
 
-FReply SMVVMSourceSelector::OnClearSource()
+FReply SSourceSelector::OnClearSource()
 {
 	if (SourceComboBox.IsValid())
 	{
@@ -96,4 +97,6 @@ FReply SMVVMSourceSelector::OnClearSource()
 	return FReply::Handled();
 }
 
+} // namespace UE::MVVM
+
 #undef LOCTEXT_NAMESPACE
