MVVM: Can only select viewmodel that are not abstract.

diff --git a/Engine/Plugins/Runtime/ModelViewViewModel/Source/ModelViewViewModelEditor/Private/Widgets/SMVVMSelectViewModel.cpp b/Engine/Plugins/Runtime/ModelViewViewModel/Source/ModelViewViewModelEditor/Private/Widgets/SMVVMSelectViewModel.cpp
--- a/Engine/Plugins/Runtime/ModelViewViewModel/Source/ModelViewViewModelEditor/Private/Widgets/SMVVMSelectViewModel.cpp
+++ b/Engine/Plugins/Runtime/ModelViewViewModel/Source/ModelViewViewModelEditor/Private/Widgets/SMVVMSelectViewModel.cpp
@@ -29,18 +29,17 @@ namespace UE::MVVM
 
 namespace Private
 {
+	bool IsValidViewModel(const UClass * InClass);
+
 	class FSelectViewModelClassFilter : public IClassViewerFilter
 	{
 	public:
 		FSelectViewModelClassFilter() = default;
-		static const EClassFlags ClassFlag = CLASS_HideDropDown | CLASS_Hidden | CLASS_Deprecated;
+
+		static const EClassFlags DisallowedClassFlags = CLASS_HideDropDown | CLASS_Hidden | CLASS_Deprecated | CLASS_Abstract | CLASS_NotPlaceable;
 		virtual bool IsClassAllowed(const FClassViewerInitializationOptions& InInitOptions, const UClass* InClass, TSharedRef<FClassViewerFilterFuncs> InFilterFuncs) override
 		{
-			if (InClass->IsChildOf(UWidget::StaticClass()))
-			{
-				return false;
-			}
-			return InClass->ImplementsInterface(UNotifyFieldValueChanged::StaticClass()) && !InClass->HasAnyClassFlags(ClassFlag);
+			return IsValidViewModel(InClass);
 		}
 
 		virtual bool IsUnloadedClassAllowed(const FClassViewerInitializationOptions& InInitOptions, const TSharedRef<const IUnloadedBlueprintData> InUnloadedClassData, TSharedRef<FClassViewerFilterFuncs> InFilterFuncs) override
@@ -49,9 +48,18 @@ namespace Private
 			{
 				return false;
 			}
-			return InUnloadedClassData->ImplementsInterface(UNotifyFieldValueChanged::StaticClass()) && !InUnloadedClassData->HasAnyClassFlags(ClassFlag);
+			return InUnloadedClassData->ImplementsInterface(UNotifyFieldValueChanged::StaticClass()) && !InUnloadedClassData->HasAnyClassFlags(DisallowedClassFlags);
 		}
 	};
+
+	bool IsValidViewModel(const UClass* InClass)
+	{
+		if (InClass->IsChildOf(UWidget::StaticClass()))
+		{
+			return false;
+		}
+		return InClass->ImplementsInterface(UNotifyFieldValueChanged::StaticClass()) && !InClass->HasAnyClassFlags(FSelectViewModelClassFilter::DisallowedClassFlags);
+	}
 } //namespace
 
 void SMVVMSelectViewModel::Construct(const FArguments& InArgs, const UWidgetBlueprint* WidgetBlueprint)
@@ -127,7 +135,7 @@ void SMVVMSelectViewModel::Construct(const FArguments& InArgs, const UWidgetBlue
 
 void SMVVMSelectViewModel::HandleClassPicked(UClass* ClassPicked)
 {
-	SelectedClass = ClassPicked;
+	SelectedClass = Private::IsValidViewModel(ClassPicked) ? ClassPicked : nullptr;
 	BindingListWidget->AddSource(ClassPicked, FName(), FGuid());
 }
 
