MVVM: Work in progress for viewmodel in viewmodel

diff --git a/Engine/Plugins/Runtime/ModelViewViewModel/Source/ModelViewViewModel/Private/Bindings/MVVMFieldPathHelper.cpp b/Engine/Plugins/Runtime/ModelViewViewModel/Source/ModelViewViewModel/Private/Bindings/MVVMFieldPathHelper.cpp
--- a/Engine/Plugins/Runtime/ModelViewViewModel/Source/ModelViewViewModel/Private/Bindings/MVVMFieldPathHelper.cpp
+++ b/Engine/Plugins/Runtime/ModelViewViewModel/Source/ModelViewViewModel/Private/Bindings/MVVMFieldPathHelper.cpp
@@ -3,7 +3,11 @@
 #include "Bindings/MVVMFieldPathHelper.h"
 
 #include "Bindings/MVVMBindingHelper.h"
+#include "Engine/Engine.h"
+#include "FieldNotification/IFieldValueChanged.h"
 #include "Misc/MemStack.h"
+#include "MVVMSubsystem.h"
+#include "Types/MVVMAvailableBinding.h"
 
 namespace UE::MVVM::FieldPathHelper
 {
@@ -80,9 +84,9 @@ TValueOrError<FMVVMConstFieldVariant, FString> TransformWithAccessor(UStruct* Cu
 } // namespace
 
 
-TValueOrError<TArray<FMVVMConstFieldVariant>, FString> GenerateFieldPathList(TSubclassOf<UObject> InFrom, FStringView InFieldPath, bool bForReading)
+TValueOrError<TArray<FMVVMConstFieldVariant>, FString> GenerateFieldPathList(const UClass* InFrom, FStringView InFieldPath, bool bForReading)
 {
-	if (InFrom.Get() == nullptr)
+	if (InFrom == nullptr)
 	{
 		return MakeError(TEXT("The source class is invalid."));
 	}
@@ -98,7 +102,7 @@ TValueOrError<TArray<FMVVMConstFieldVariant>, FString> GenerateFieldPathList(TSu
 	FMemMark Mark(FMemStack::Get());
 	TArray<FMVVMConstFieldVariant, TMemStackAllocator<>> Result;
 
-	UStruct* CurrentContainer = InFrom.Get();
+	const UStruct* CurrentContainer = InFrom;
 
 	// Split the string into property or function names
 	//ie. myvar.myfunction.myvar
@@ -237,6 +241,123 @@ TValueOrError<TArray<FMVVMConstFieldVariant>, FString> GenerateFieldPathList(TAr
 }
 
 
+TValueOrError<FParsedBindingInfo, FString> GetBindingInfoFromFieldPath(const UClass* InAccessor, TArrayView<const FMVVMConstFieldVariant> InFieldPath)
+{
+	if (InAccessor == nullptr)
+	{
+		return MakeError(TEXT("The Accessor is null."));
+	}
+	if (InFieldPath.Num() == 0)
+	{
+		return MakeError(TEXT("The FieldPath is empty."));
+	}
+
+	FParsedBindingInfo Result;
+	FParsedBindingInfo::FFieldVariantArray LocalFullPath;
+	bool bLookForNotifyFieldId = false;
+	UStruct* CurrentContainerType = InFieldPath[0].GetOwner();
+	if (!InAccessor->IsChildOf(CurrentContainerType))
+	{
+		return MakeError(TEXT("FieldPath that doesn't start with the accessor is not supported."));
+	}
+
+	for (int32 Index = 0; Index < InFieldPath.Num(); ++Index)
+	{
+		bool bLastField = Index == InFieldPath.Num() - 1;
+		FMVVMConstFieldVariant Field = InFieldPath[Index];
+
+		if (Field.IsEmpty())
+		{
+			return MakeError(FString::Printf(TEXT("The field '%d' does not exist."), Index));
+		}
+
+		if (CurrentContainerType == nullptr)
+		{
+			return MakeError(FString::Printf(TEXT("The field '%s' does not exist."), *Field.GetName().ToString()));
+		}
+
+		const bool bIsChild = Field.GetOwner()->IsChildOf(CurrentContainerType);
+		const bool bIsDownCast = CurrentContainerType->IsChildOf(Field.GetOwner());
+		if (!(bIsChild || bIsDownCast))
+		{
+			return MakeError(FString::Printf(TEXT("The field '%s' does not exist in the struct '%s'."), *Field.GetName().ToString(), *CurrentContainerType->GetName()));
+		}
+
+		const FProperty* FieldProperty = nullptr;
+		if (Field.IsProperty())
+		{
+			FieldProperty = Field.GetProperty();
+		}
+		else if (Field.IsFunction() && Field.GetFunction())
+		{
+			FieldProperty = BindingHelper::GetReturnProperty(Field.GetFunction());
+		}
+
+		if (FieldProperty == nullptr)
+		{
+			return MakeError(FString::Printf(TEXT("The field '%s' has an invalid property."), *Field.GetName().ToString()));
+		}
+
+
+		LocalFullPath.Add(Field);
+
+		// Is field a viewmodel
+		bool bSaveNotifyFieldInterfacePath = false;
+		if (!bLastField)
+		{
+			if (const FObjectPropertyBase* FieldObjectProperty = CastField<const FObjectPropertyBase>(FieldProperty))
+			{
+				if (FieldObjectProperty->PropertyClass->ImplementsInterface(UNotifyFieldValueChanged::StaticClass()))
+				{
+					Result.NotifyFieldClass = FieldObjectProperty->PropertyClass;
+					Result.NotifyFieldId = FFieldNotificationId();
+					Result.NotifyFieldInterfacePath = LocalFullPath;
+					Result.ViewModelIndex = Index;
+					bSaveNotifyFieldInterfacePath = true;
+					bLookForNotifyFieldId = true;
+				}
+			}
+		}
+
+		if (!bSaveNotifyFieldInterfacePath)
+		{
+			if (bLookForNotifyFieldId && Result.NotifyFieldClass.Get() != nullptr)
+			{
+				FMVVMBindingName BindingName = FMVVMBindingName(Field.GetName());
+				FMVVMAvailableBinding AvailableBinding = GEngine->GetEngineSubsystem<UMVVMSubsystem>()->GetAvailableBinding(Result.NotifyFieldClass.Get(), BindingName, InAccessor);
+				if (AvailableBinding.IsValid())
+				{
+					Result.NotifyFieldId = FFieldNotificationId(AvailableBinding.GetBindingName().ToName());
+				}
+			}
+			bLookForNotifyFieldId = false;
+		}
+
+		if (!bLastField)
+		{
+			TValueOrError<UStruct*, FString> FoundContainer = Private::FindContainer(FieldProperty);
+			if (FoundContainer.HasError())
+			{
+				return MakeError(FoundContainer.StealError());
+			}
+			CurrentContainerType = FoundContainer.GetValue();
+		}
+	}
+
+	if (Result.ViewModelIndex > 0)
+	{
+		TValueOrError<TArray<FMVVMConstFieldVariant>, FString> FinalPath = GenerateFieldPathList(Result.NotifyFieldInterfacePath, true);
+		if (FinalPath.HasError())
+		{
+			return MakeError(FinalPath.StealError());
+		}
+		Result.NotifyFieldInterfacePath = FinalPath.StealValue();
+	}
+
+	return MakeValue(MoveTemp(Result));
+}
+
+
 FString ToString(TArrayView<const FMVVMFieldVariant> Fields)
 {
 	TStringBuilder<512> Builder;
diff --git a/Engine/Plugins/Runtime/ModelViewViewModel/Source/ModelViewViewModel/Private/MVVMSubsystem.cpp b/Engine/Plugins/Runtime/ModelViewViewModel/Source/ModelViewViewModel/Private/MVVMSubsystem.cpp
--- a/Engine/Plugins/Runtime/ModelViewViewModel/Source/ModelViewViewModel/Private/MVVMSubsystem.cpp
+++ b/Engine/Plugins/Runtime/ModelViewViewModel/Source/ModelViewViewModel/Private/MVVMSubsystem.cpp
@@ -44,12 +44,58 @@ bool UMVVMSubsystem::DoesWidgetTreeContainedWidget(const UWidgetTree* WidgetTree
 
 namespace UE::MVVM::Private
 {
-	TArray<FMVVMAvailableBinding> GetAvailableBindings(const UStruct* Class, const UClass* AccessorType, const UE::FieldNotification::IClassDescriptor* ClassDescriptor)
+
+	FMVVMAvailableBinding GetAvailableBinding(const UFunction* Function, bool bHasNotify, bool bCanAccessPrivate, bool bCanAccessProtected)
+	{
+		// N.B. A function can be private/protected and can still be use when it's a BlueprintGetter/BlueprintSetter.
+		//But we bind to the property not the function.
+		if (Function->HasAnyFunctionFlags(FUNC_Private) && !bCanAccessPrivate)
+		{
+			return FMVVMAvailableBinding();
+		}
+		if (Function->HasAnyFunctionFlags(FUNC_Protected) && !bCanAccessProtected)
+		{
+			return FMVVMAvailableBinding();
+		}
+
+		const bool bIsReadable = BindingHelper::IsValidForSourceBinding(Function);
+		const bool bIsWritable = BindingHelper::IsValidForDestinationBinding(Function);
+		if (bIsReadable || bIsWritable || bHasNotify)
+		{
+			return FMVVMAvailableBinding(FMVVMBindingName(Function->GetFName()), bIsReadable, bIsWritable, bHasNotify && bIsReadable);
+		}
+
+		return FMVVMAvailableBinding();
+	}
+
+	FMVVMAvailableBinding GetAvailableBinding(const FProperty* Property, bool bHasNotify, bool bCanAccessPrivate, bool bCanAccessProtected)
+	{
+		// N.B. Property can be private/protected in cpp but if they are visible in BP that is all that matter.
+		//Only property defined in BP can be BP visible and really private.
+#if WITH_EDITOR
+		static FName NAME_BlueprintPrivate = "BlueprintPrivate";
+		if (Property->GetBoolMetaData(NAME_BlueprintPrivate) && !bCanAccessPrivate)
+		{
+			return FMVVMAvailableBinding();
+		}
+#endif
+
+		const bool bIsReadable = BindingHelper::IsValidForSourceBinding(Property);
+		const bool bIsWritable = BindingHelper::IsValidForDestinationBinding(Property);
+		if (bIsReadable || bIsWritable || bHasNotify)
+		{
+			return FMVVMAvailableBinding(FMVVMBindingName(Property->GetFName()), bIsReadable, bIsWritable, bHasNotify && bIsReadable);
+		}
+
+		return FMVVMAvailableBinding();
+	}
+
+	TArray<FMVVMAvailableBinding> GetAvailableBindings(const UStruct* Container, const UClass* AccessorType, const UE::FieldNotification::IClassDescriptor* ClassDescriptor)
 	{
 		TSet<FName> FieldDescriptors;
 		if (ClassDescriptor)
 		{
-			ClassDescriptor->ForEachField(CastChecked<UClass>(Class), [&FieldDescriptors](UE::FieldNotification::FFieldId FieldId)
+			ClassDescriptor->ForEachField(CastChecked<UClass>(Container), [&FieldDescriptors](UE::FieldNotification::FFieldId FieldId)
 			{
 				FieldDescriptors.Add(FieldId.GetName());
 				return true;
@@ -59,32 +105,21 @@ namespace UE::MVVM::Private
 		TArray<FMVVMAvailableBinding> Result;
 		Result.Reserve(FieldDescriptors.Num());
 
-		const bool bAccessorCanSeeProtectedMember = AccessorType ? AccessorType->IsChildOf(Class) : false;
+		const bool bCanAccessPrivateMember = Container == AccessorType;
+		const bool bCanAccessProtectedMember = AccessorType ? AccessorType->IsChildOf(Container) : false;
 
-		if (Cast<UClass>(Class))
+		if (Cast<UClass>(Container))
 		{
-			for (TFieldIterator<const UFunction> FunctionIt(Class, EFieldIteratorFlags::IncludeSuper); FunctionIt; ++FunctionIt)
+			for (TFieldIterator<const UFunction> FunctionIt(Container, EFieldIteratorFlags::IncludeSuper); FunctionIt; ++FunctionIt)
 			{
 				const UFunction* Function = *FunctionIt;
 				check(Function);
 
-				// N.B. A function can be private/protected and can still be use when it's a BlueprintGetter/BlueprintSetter.
-				//But we bind to the property not the function.
-				if (Function->HasAnyFunctionFlags(FUNC_Private) && Class != AccessorType)
-				{
-					continue;
-				}
-				if (Function->HasAnyFunctionFlags(FUNC_Protected) && !bAccessorCanSeeProtectedMember)
-				{
-					continue;
-				}
-
-				const bool bIsReadable = BindingHelper::IsValidForSourceBinding(Function);
-				const bool bIsWritable = BindingHelper::IsValidForDestinationBinding(Function);
-				const bool bHasNotify = FieldDescriptors.Contains(Function->GetFName()) && bIsReadable;
-				if (bIsReadable || bIsWritable || bHasNotify)
+				const bool bHasNotify = FieldDescriptors.Contains(Function->GetFName());
+				FMVVMAvailableBinding Binding = GetAvailableBinding(Function, bHasNotify, bCanAccessPrivateMember, bCanAccessProtectedMember);
+				if (Binding.IsValid())
 				{
-					Result.Add(FMVVMAvailableBinding(FMVVMBindingName(Function->GetFName()), bIsReadable, bIsWritable, bHasNotify));
+					Result.Add(Binding);
 				}
 			}
 		}
@@ -92,43 +127,33 @@ namespace UE::MVVM::Private
 #if WITH_EDITOR
 		FName NAME_BlueprintGetter = "BlueprintGetter";
 		FName NAME_BlueprintSetter = "BlueprintSetter";
-		FName NAME_BlueprintPrivate = "BlueprintPrivate";
 #endif
 
-		for (TFieldIterator<const FProperty> PropertyIt(Class, EFieldIteratorFlags::IncludeSuper); PropertyIt; ++PropertyIt)
+		for (TFieldIterator<const FProperty> PropertyIt(Container, EFieldIteratorFlags::IncludeSuper); PropertyIt; ++PropertyIt)
 		{
 			const FProperty* Property = *PropertyIt;
+			check(Property);
 
-			// N.B. Property can be private/protected in cpp but if they are visible in BP that is all that matter.
-			//Only property defined in BP can be BP visible and really private.
-#if WITH_EDITOR
-			if (Property->GetBoolMetaData(NAME_BlueprintPrivate) && Class != AccessorType)
-			{
-				continue;
-			}
-#endif
-
-			const bool bIsReadable = BindingHelper::IsValidForSourceBinding(Property);
-			const bool bIsWritable = BindingHelper::IsValidForDestinationBinding(Property);
-			const bool bHasNotify = FieldDescriptors.Contains(Property->GetFName()) && bIsReadable;
-			if (bIsReadable || bIsWritable || bHasNotify)
+			const bool bHasNotify = FieldDescriptors.Contains(Property->GetFName());
+			FMVVMAvailableBinding Binding = GetAvailableBinding(Property, bHasNotify, bCanAccessPrivateMember, bCanAccessProtectedMember);
+			if (Binding.IsValid())
 			{
 #if WITH_EDITOR
 				// Remove any BlueprintGetter & BlueprintSetter and use the Property instead.
-				if (bIsReadable)
+				if (Binding.IsReadable())
 				{
 					const FString& PropertyGetter = Property->GetMetaData(NAME_BlueprintGetter);
 					if (!PropertyGetter.IsEmpty())
 					{
 						FMVVMBindingName BindingName = FMVVMBindingName(*PropertyGetter);
-						int32 BindingIndex = Result.IndexOfByPredicate([BindingName](const FMVVMAvailableBinding& Binding) { return Binding.GetBindingName() == BindingName && Binding.IsReadable() && !Binding.IsWritable(); });
+						int32 BindingIndex = Result.IndexOfByPredicate([BindingName](const FMVVMAvailableBinding& Binding) { return Binding.GetBindingName() == BindingName && Binding.IsReadable() && !Binding.IsWritable() && !Binding.HasNotify(); });
 						if (BindingIndex != INDEX_NONE)
 						{
 							Result.RemoveAtSwap(BindingIndex);
 						}
 					}
 				}
-				if (bIsWritable)
+				if (Binding.IsWritable())
 				{
 					const FString& PropertySetter = Property->GetMetaData(NAME_BlueprintSetter);
 					if (!PropertySetter.IsEmpty())
@@ -142,30 +167,75 @@ namespace UE::MVVM::Private
 					}
 				}
 #endif
-				Result.Add(FMVVMAvailableBinding(FMVVMBindingName(Property->GetFName()), bIsReadable, bIsWritable, bHasNotify));
+				Result.Add(Binding);
 			}
 		}
 
 		return Result;
 	}
 
-	TArray<FMVVMAvailableBinding> GetAvailableBindings(TSubclassOf<UObject> InSubClass, TSubclassOf<UObject> InAccessor)
+	TArray<FMVVMAvailableBinding> GetAvailableBindings(const UClass* InSubClass, const UClass* InAccessor)
 	{
-		if (InSubClass.Get() && InSubClass.Get()->ImplementsInterface(UNotifyFieldValueChanged::StaticClass()))
+		if (InSubClass && InSubClass->ImplementsInterface(UNotifyFieldValueChanged::StaticClass()))
 		{
-			TScriptInterface<INotifyFieldValueChanged> DefaultObject = InSubClass.GetDefaultObject();
+			TScriptInterface<INotifyFieldValueChanged> DefaultObject = InSubClass->GetDefaultObject();
 			const UE::FieldNotification::IClassDescriptor& ClassDescriptor = DefaultObject->GetFieldNotificationDescriptor();
-			return GetAvailableBindings(InSubClass.Get(), InAccessor.Get(), &ClassDescriptor);
+			return GetAvailableBindings(InSubClass, InAccessor, &ClassDescriptor);
 		}
 		else
 		{
-			return GetAvailableBindings(InSubClass.Get(), InAccessor.Get(), nullptr);
+			return GetAvailableBindings(InSubClass, InAccessor, nullptr);
+		}
+	}
+
+
+	FMVVMAvailableBinding GetAvailableBinding(FMVVMBindingName BindingName, const UStruct* Container, const UClass* AccessorType, const UE::FieldNotification::IClassDescriptor* ClassDescriptor)
+	{
+		bool bHasNotify = false;
+		if (ClassDescriptor)
+		{
+			UE::FieldNotification::FFieldId FieldId = ClassDescriptor->GetField(CastChecked<UClass>(Container), BindingName.ToName());
+			bHasNotify = FieldId.IsValid();
+		}
+
+		const bool bCanAccessPrivateMember = Container == AccessorType;
+		const bool bCanAccessProtectedMember = AccessorType ? AccessorType->IsChildOf(Container) : false;
+
+		FMVVMFieldVariant FieldVariant = UE::MVVM::BindingHelper::FindFieldByName(Container, BindingName);
+		if (FieldVariant.IsProperty())
+		{
+			return GetAvailableBinding(FieldVariant.GetProperty(), bHasNotify, bCanAccessPrivateMember, bCanAccessProtectedMember);
+		}
+		else if (FieldVariant.IsFunction())
+		{
+			return GetAvailableBinding(FieldVariant.GetFunction(), bHasNotify, bCanAccessPrivateMember, bCanAccessProtectedMember);
+		}
+
+		return FMVVMAvailableBinding();
+	}
+
+	FMVVMAvailableBinding GetAvailableBinding(FMVVMBindingName BindingName, const UClass* InSubClass, const UClass* InAccessor)
+	{
+		if (!BindingName.IsValid())
+		{
+			return FMVVMAvailableBinding();
+		}
+
+		if (InSubClass && InSubClass->ImplementsInterface(UNotifyFieldValueChanged::StaticClass()))
+		{
+			TScriptInterface<INotifyFieldValueChanged> DefaultObject = InSubClass->GetDefaultObject();
+			const UE::FieldNotification::IClassDescriptor& ClassDescriptor = DefaultObject->GetFieldNotificationDescriptor();
+			return GetAvailableBinding(BindingName, InSubClass, InAccessor, &ClassDescriptor);
+		}
+		else
+		{
+			return GetAvailableBinding(BindingName, InSubClass, InAccessor, nullptr);
 		}
 	}
 } //namespace
 
 
-TArray<FMVVMAvailableBinding> UMVVMSubsystem::GetAvailableBindings(TSubclassOf<UObject> Class, TSubclassOf<UObject> Accessor) const
+TArray<FMVVMAvailableBinding> UMVVMSubsystem::GetAvailableBindings(const UClass* Class, const UClass* Accessor) const
 {
 	return UE::MVVM::Private::GetAvailableBindings(Class, Accessor);
 }
@@ -177,6 +247,12 @@ TArray<FMVVMAvailableBinding> UMVVMSubsystem::GetAvailableBindingsForStruct(cons
 }
 
 
+FMVVMAvailableBinding UMVVMSubsystem::GetAvailableBinding(const UClass* Class, FMVVMBindingName BindingName, const UClass* Accessor) const
+{
+	return UE::MVVM::Private::GetAvailableBinding(BindingName, Class, Accessor);
+}
+
+
 TValueOrError<bool, FString> UMVVMSubsystem::IsBindingValid(FBindingArgs Args) const
 {
 	if (UE::MVVM::IsForwardBinding(Args.Mode))
diff --git a/Engine/Plugins/Runtime/ModelViewViewModel/Source/ModelViewViewModel/Public/Bindings/MVVMFieldPathHelper.h b/Engine/Plugins/Runtime/ModelViewViewModel/Source/ModelViewViewModel/Public/Bindings/MVVMFieldPathHelper.h
--- a/Engine/Plugins/Runtime/ModelViewViewModel/Source/ModelViewViewModel/Public/Bindings/MVVMFieldPathHelper.h
+++ b/Engine/Plugins/Runtime/ModelViewViewModel/Source/ModelViewViewModel/Public/Bindings/MVVMFieldPathHelper.h
@@ -1,28 +1,63 @@
-// Copyright Epic Games, Inc. All Rights Reserved.
-
-#pragma once
-
-#include "CoreMinimal.h"
-#include "Templates/SubclassOf.h"
-#include "Templates/ValueOrError.h"
-#include "Types/MVVMFieldContext.h"
-#include "UObject/FieldPath.h"
-#include "UObject/Object.h"
-
-
-namespace UE::MVVM::FieldPathHelper
-{
-	/** */
-	UE_NODISCARD MODELVIEWVIEWMODEL_API TValueOrError<TArray<FMVVMConstFieldVariant>, FString> GenerateFieldPathList(TSubclassOf<UObject> From, FStringView FieldPath, bool bForSourceBinding);
-
-	/** */
-	UE_NODISCARD MODELVIEWVIEWMODEL_API TValueOrError<TArray<FMVVMConstFieldVariant>, FString> GenerateFieldPathList(TArrayView<const FMVVMConstFieldVariant>, bool bForSourceBinding);
-
-	/** */
-	UE_NODISCARD MODELVIEWVIEWMODEL_API FString ToString(TArrayView<const FMVVMFieldVariant> Fields);
-	UE_NODISCARD MODELVIEWVIEWMODEL_API FString ToString(TArrayView<const FMVVMConstFieldVariant> Fields);
-
-	/** */
-	UE_NODISCARD MODELVIEWVIEWMODEL_API TValueOrError<UObject*, void> EvaluateObjectProperty(const FFieldContext& InSource);
-
-} // namespace
+// Copyright Epic Games, Inc. All Rights Reserved.
+
+#pragma once
+
+#include "CoreMinimal.h"
+#include "FieldNotification/FieldId.h"
+#include "Templates/SubclassOf.h"
+#include "Templates/ValueOrError.h"
+#include "Types/MVVMFieldContext.h"
+#include "UObject/FieldPath.h"
+#include "UObject/Object.h"
+
+
+namespace UE::MVVM::FieldPathHelper
+{
+	/**
+	 * Info about the viewmodel and field.
+	 */
+	struct FParsedBindingInfo
+	{
+		using FFieldVariantArray = TArray<FMVVMConstFieldVariant, TInlineAllocator<4>>;
+
+		/**
+		 * The viewmodel that is used to create the binding. When it's invalid the binding can be valid but only OneTime mode is supported.
+		 * @return the class of "ViewmodelB" when the path is "ViewmodelA.ViewmodelB.Vector.X"
+		 */
+		TSubclassOf<UObject> NotifyFieldClass;
+
+		/**
+		 * The property we will bind to. When it's invalid the binding can be valid but only OneTime mode is supported.
+		 * @return "Vector" when the path is "ViewmodelA.ViewmodelB.Vector.X"
+		 */
+		FFieldNotificationId NotifyFieldId;
+
+		/**
+		 * The property path to get the viewmodel at runtime.
+		 * @return "ViewmodelA.ViewmodelB" or "ViewmodelA.GetViewmodelB" when the path is "ViewmodelA.ViewmodelB.Vector.X"
+		 */
+		FFieldVariantArray NotifyFieldInterfacePath;
+		/**
+		 * The viewmodel index in the original path.
+		 * @return 1 when the path is "ViewmodelA.ViewmodelB.Vector.X"
+		 */
+		int32 ViewModelIndex = INDEX_NONE;
+	};
+
+	/** */
+	UE_NODISCARD MODELVIEWVIEWMODEL_API TValueOrError<TArray<FMVVMConstFieldVariant>, FString> GenerateFieldPathList(const UClass* From, FStringView FieldPath, bool bForSourceBinding);
+
+	/** */
+	UE_NODISCARD MODELVIEWVIEWMODEL_API TValueOrError<TArray<FMVVMConstFieldVariant>, FString> GenerateFieldPathList(TArrayView<const FMVVMConstFieldVariant> Fields, bool bForSourceBinding);
+
+	/** */
+	UE_NODISCARD MODELVIEWVIEWMODEL_API TValueOrError<FParsedBindingInfo, FString> GetBindingInfoFromFieldPath(const UClass* Accessor, TArrayView<const FMVVMConstFieldVariant> Fields);
+
+	/** */
+	UE_NODISCARD MODELVIEWVIEWMODEL_API FString ToString(TArrayView<const FMVVMFieldVariant> Fields);
+	UE_NODISCARD MODELVIEWVIEWMODEL_API FString ToString(TArrayView<const FMVVMConstFieldVariant> Fields);
+
+	/** */
+	UE_NODISCARD MODELVIEWVIEWMODEL_API TValueOrError<UObject*, void> EvaluateObjectProperty(const FFieldContext& InSource);
+
+} // namespace
diff --git a/Engine/Plugins/Runtime/ModelViewViewModel/Source/ModelViewViewModel/Public/MVVMSubsystem.h b/Engine/Plugins/Runtime/ModelViewViewModel/Source/ModelViewViewModel/Public/MVVMSubsystem.h
--- a/Engine/Plugins/Runtime/ModelViewViewModel/Source/ModelViewViewModel/Public/MVVMSubsystem.h
+++ b/Engine/Plugins/Runtime/ModelViewViewModel/Source/ModelViewViewModel/Public/MVVMSubsystem.h
@@ -36,18 +36,23 @@ public:
 	bool DoesWidgetTreeContainedWidget(const UWidgetTree* WidgetTree, const UWidget* ViewWidget) const;
 
 	/**
-	 * Returns the list of all the bindings that are available for the Class.
-	 * The class must implement the UNotifyFieldValueChanged interface.
+	 * @return The list of all the bindings that are available for the Class.
 	 */
 	UFUNCTION(BlueprintCallable, Category = "MVVM")
-	TArray<FMVVMAvailableBinding> GetAvailableBindings(TSubclassOf<UObject> Class, TSubclassOf<UObject> Accessor) const;
+	TArray<FMVVMAvailableBinding> GetAvailableBindings(const UClass* Class, const UClass* Accessor) const;
 
 	/**
-	 * Returns the list of all the bindings that are available from the SriptStuct.
+	 * @return The list of all the bindings that are available from the SriptStuct.
 	 * @note When FMVVMAvailableBinding::HasNotify is false, a notification can still be triggered by the owner of the struct. The struct changed but which property of the struct changed is unknown.
 	 */
 	TArray<FMVVMAvailableBinding> GetAvailableBindingsForStruct(const UScriptStruct* Struct) const;
 
+	/**
+	 * @return the available binding from the binding.
+	 */
+	UFUNCTION(BlueprintCallable, Category = "MVVM")
+	FMVVMAvailableBinding GetAvailableBinding(const UClass* Class, FMVVMBindingName BindingName, const UClass* Accessor) const;
+
 	UFUNCTION(BlueprintCallable, Category = "MVVM")
 	UMVVMViewModelCollectionObject* GetGlobalViewModelCollection() const
 	{
diff --git a/Engine/Plugins/Runtime/ModelViewViewModel/Source/ModelViewViewModelBlueprint/Private/Bindings/MVVMCompiledBindingLibraryCompiler.cpp b/Engine/Plugins/Runtime/ModelViewViewModel/Source/ModelViewViewModelBlueprint/Private/Bindings/MVVMCompiledBindingLibraryCompiler.cpp
--- a/Engine/Plugins/Runtime/ModelViewViewModel/Source/ModelViewViewModelBlueprint/Private/Bindings/MVVMCompiledBindingLibraryCompiler.cpp
+++ b/Engine/Plugins/Runtime/ModelViewViewModel/Source/ModelViewViewModelBlueprint/Private/Bindings/MVVMCompiledBindingLibraryCompiler.cpp
@@ -3,7 +3,6 @@
 #include "Bindings/MVVMCompiledBindingLibraryCompiler.h"
 
 #include "Bindings/MVVMBindingHelper.h"
-#include "Bindings/MVVMFieldPathHelper.h"
 #include "Engine/Engine.h"
 #include "MVVMSubsystem.h"
 #include "Types/MVVMObjectVariant.h"
@@ -182,31 +181,11 @@ TValueOrError<FCompiledBindingLibraryCompiler::FFieldIdHandle, FString> FCompile
 }
 
 
-TValueOrError<FCompiledBindingLibraryCompiler::FFieldPathHandle, FString> FCompiledBindingLibraryCompiler::AddFieldPath(TSubclassOf<UObject> InSourceClass, FStringView InFieldPath, bool bInRead)
-{
-	Impl->bCompiled = false;
-
-	TValueOrError<TArray<FMVVMConstFieldVariant>, FString> GeneratedField = FieldPathHelper::GenerateFieldPathList(InSourceClass, InFieldPath, bInRead);
-	if (GeneratedField.HasError())
-	{
-		return MakeError(GeneratedField.StealError());
-	}
-
-	return AddFieldPathImpl(MakeArrayView(GeneratedField.GetValue()), bInRead);
-}
-
-
 TValueOrError<FCompiledBindingLibraryCompiler::FFieldPathHandle, FString> FCompiledBindingLibraryCompiler::AddFieldPath(TArrayView<const FMVVMConstFieldVariant> InFieldPath, bool bInRead)
 {
 	Impl->bCompiled = false;
 
-	TValueOrError<TArray<FMVVMConstFieldVariant>, FString> GeneratedField = FieldPathHelper::GenerateFieldPathList(InFieldPath, bInRead);
-	if (GeneratedField.HasError())
-	{
-		return MakeError(GeneratedField.StealError());
-	}
-
-	return AddFieldPathImpl(MakeArrayView(GeneratedField.GetValue()), bInRead);
+	return AddFieldPathImpl(InFieldPath, bInRead);
 }
 
 
@@ -335,30 +314,25 @@ TValueOrError<FCompiledBindingLibraryCompiler::FFieldPathHandle, FString> FCompi
 }
 
 
-TValueOrError<FCompiledBindingLibraryCompiler::FFieldPathHandle, FString> FCompiledBindingLibraryCompiler::AddObjectFieldPath(TSubclassOf<UObject> InSourceClass, FStringView InFieldPath, UClass* ExpectedType, bool bInRead)
+TValueOrError<FCompiledBindingLibraryCompiler::FFieldPathHandle, FString> FCompiledBindingLibraryCompiler::AddObjectFieldPath(TArrayView<const UE::MVVM::FMVVMConstFieldVariant> FieldPath, UClass* ExpectedType, bool bInRead)
 {
 	Impl->bCompiled = false;
 
 	check(ExpectedType);
 
-	TValueOrError<TArray<FMVVMConstFieldVariant>, FString> GeneratedField = FieldPathHelper::GenerateFieldPathList(InSourceClass, InFieldPath, bInRead);
-	if (GeneratedField.HasError())
-	{
-		return MakeError(GeneratedField.StealError());
-	}
-	if (GeneratedField.GetValue().Num() == 0)
+	if (FieldPath.Num() == 0)
 	{
 		return MakeError(FString::Printf(TEXT("The field does not return a '%s'."), *ExpectedType->GetName()));
 	}
 
 	const FObjectPropertyBase* ObjectPropertyBase = nullptr;
-	if (GeneratedField.GetValue().Last().IsProperty())
+	if (FieldPath.Last().IsProperty())
 	{
-		ObjectPropertyBase = CastField<const FObjectPropertyBase>(GeneratedField.GetValue().Last().GetProperty());
+		ObjectPropertyBase = CastField<const FObjectPropertyBase>(FieldPath.Last().GetProperty());
 	}
-	else if (GeneratedField.GetValue().Last().IsFunction())
+	else if (FieldPath.Last().IsFunction())
 	{
-		ObjectPropertyBase = CastField<const FObjectPropertyBase>(BindingHelper::GetReturnProperty(GeneratedField.GetValue().Last().GetFunction()));
+		ObjectPropertyBase = CastField<const FObjectPropertyBase>(BindingHelper::GetReturnProperty(FieldPath.Last().GetFunction()));
 	}
 
 	if (ObjectPropertyBase == nullptr)
@@ -370,7 +344,7 @@ TValueOrError<FCompiledBindingLibraryCompiler::FFieldPathHandle, FString> FCompi
 		return MakeError(FString::Printf(TEXT("The field does not return a '%s'."), *ExpectedType->GetName()));
 	}
 
-	return AddFieldPathImpl(MakeArrayView(GeneratedField.GetValue()), bInRead);
+	return AddFieldPathImpl(FieldPath, bInRead);
 }
 
 
diff --git a/Engine/Plugins/Runtime/ModelViewViewModel/Source/ModelViewViewModelBlueprint/Private/MVVMViewBlueprintCompiler.cpp b/Engine/Plugins/Runtime/ModelViewViewModel/Source/ModelViewViewModelBlueprint/Private/MVVMViewBlueprintCompiler.cpp
--- a/Engine/Plugins/Runtime/ModelViewViewModel/Source/ModelViewViewModelBlueprint/Private/MVVMViewBlueprintCompiler.cpp
+++ b/Engine/Plugins/Runtime/ModelViewViewModel/Source/ModelViewViewModelBlueprint/Private/MVVMViewBlueprintCompiler.cpp
@@ -74,6 +74,7 @@ void FMVVMViewBlueprintCompiler::CleanTemporaries(UWidgetBlueprintGeneratedClass
 			ensureMsgf(SourceCreatorContext.SetterGraph->HasAnyFlags(RF_Transient), TEXT("The graph should be temporary and should be generated automaticly."));
 			// GC may not have clean the graph (GC doesn't run when bRegenerateSkeletonOnly is on)
 			SourceCreatorContext.SetterGraph->Rename(nullptr, SourceCreatorContext.SetterGraph->GetOuter(), REN_DoNotDirty | REN_ForceNoResetLoaders);
+			SourceCreatorContext.SetterGraph = nullptr;
 		}
 	}
 }
@@ -243,11 +244,7 @@ void FMVVMViewBlueprintCompiler::CreateWidgetMap(const FWidgetBlueprintCompilerC
 void FMVVMViewBlueprintCompiler::CreateSourceLists(const FWidgetBlueprintCompilerContext::FCreateVariableContext& Context, UMVVMBlueprintView* BlueprintView)
 {
 	SourceContexts.Reset();
-
-	if (Context.GetCompileType() == EKismetCompileType::SkeletonOnly)
-	{
-		SourceCreatorContexts.Reset();
-	}
+	SourceCreatorContexts.Reset();
 
 	TSet<FGuid> ViewModelsGuid;
 	TSet<FName> WidgetSources;
@@ -280,7 +277,7 @@ void FMVVMViewBlueprintCompiler::CreateSourceLists(const FWidgetBlueprintCompile
 
 		const bool bCreateSetterFunction = ViewModelContext.bCreateSetterFunction || ViewModelContext.CreationType == EMVVMBlueprintViewModelContextCreationType::Manual;
 
-		int32 FoundSourceContextIndex = INDEX_NONE;
+		int32 FoundSourceCreatorContextIndex = INDEX_NONE;
 		if (Context.GetCompileType() == EKismetCompileType::SkeletonOnly)
 		{
 			FCompilerSourceCreatorContext SourceContext;
@@ -290,17 +287,17 @@ void FMVVMViewBlueprintCompiler::CreateSourceLists(const FWidgetBlueprintCompile
 			{
 				SourceContext.SetterFunctionName = TEXT("Set") + ViewModelContext.GetViewModelName().ToString();
 			}
-			FoundSourceContextIndex = SourceCreatorContexts.Emplace(MoveTemp(SourceContext));
+			FoundSourceCreatorContextIndex = SourceCreatorContexts.Emplace(MoveTemp(SourceContext));
 		}
 		else
 		{
 			FGuid ViewModelId = ViewModelContext.GetViewModelId();
-			FoundSourceContextIndex = SourceCreatorContexts.IndexOfByPredicate([ViewModelId](const FCompilerSourceCreatorContext& Other)
+			FoundSourceCreatorContextIndex = SourceCreatorContexts.IndexOfByPredicate([ViewModelId](const FCompilerSourceCreatorContext& Other)
 				{
 					return Other.ViewModelContext.GetViewModelId() == ViewModelId;
 				});
 		}
-		checkf(FoundSourceContextIndex != INDEX_NONE, TEXT("The viewmodel was added after the skeleton was created?"));
+		checkf(FoundSourceCreatorContextIndex != INDEX_NONE, TEXT("The viewmodel was added after the skeleton was created?"));
 
 		FCompilerSourceContext SourceVariable;
 		SourceVariable.Class = ViewModelContext.GetViewModelClass();
@@ -308,7 +305,7 @@ void FMVVMViewBlueprintCompiler::CreateSourceLists(const FWidgetBlueprintCompile
 		SourceVariable.DisplayName = ViewModelContext.GetDisplayName().ToString();
 		SourceVariable.CategoryName = TEXT("Viewmodel");
 		SourceVariable.bExposeOnSpawn = bCreateSetterFunction;
-		SourceVariable.BlueprintSetter = SourceCreatorContexts[FoundSourceContextIndex].SetterFunctionName;
+		SourceVariable.BlueprintSetter = SourceCreatorContexts[FoundSourceCreatorContextIndex].SetterFunctionName;
 		SourceContexts.Emplace(MoveTemp(SourceVariable));
 	}
 
@@ -331,56 +328,73 @@ void FMVVMViewBlueprintCompiler::CreateSourceLists(const FWidgetBlueprintCompile
 			continue;
 		}
 
-		if (IsForwardBinding(Binding.BindingType) || IsBackwardBinding(Binding.BindingType))
+		FMVVMViewBlueprintCompiler* Self = this;
+		auto GenerateCompilerSourceContext = [Self, Index, &ViewModelsGuid, &WidgetSources](const UWidgetBlueprintGeneratedClass* Class, const UMVVMBlueprintView* BlueprintView, const FMVVMBlueprintPropertyPath& PropertyPath) -> bool
 		{
-			// todo support any type of UObject
-			const FMVVMBlueprintViewModelContext* SourceViewModelContext = BlueprintView->FindViewModel(Binding.ViewModelPath.GetViewModelId());
-			if (SourceViewModelContext == nullptr)
+			if (PropertyPath.IsFromWidget())
 			{
-				WidgetBlueprintCompilerContext.MessageLog.Error(*FString::Printf(TEXT("The binding number %d has an invalid source."), Index));
-				bAreSourceContextsValid = false;
-				continue;
-			}
-
-			bool bFound = ViewModelsGuid.Contains(SourceViewModelContext->GetViewModelId());
-			if (!bFound)
-			{
-				WidgetBlueprintCompilerContext.MessageLog.Error(*FString::Printf(TEXT("The binding number %d has a source that need to be added automaticly. That is not supported yet."), Index));
-				bAreSourceContextsValid = false;
-			}
+				if (PropertyPath.GetWidgetName().IsNone())
+				{
+					Self->WidgetBlueprintCompilerContext.MessageLog.Error(*FString::Printf(TEXT("The binding number %d has an invalid widget name."), Index));
+					return false;
+				}
 
-			if (Binding.WidgetPath.GetWidgetName().IsNone())
-			{
-				WidgetBlueprintCompilerContext.MessageLog.Error(*FString::Printf(TEXT("The widget binding is invalid.")));
-				bAreSourceContextsValid = false;
-				continue;
-			}
+				if (PropertyPath.GetWidgetName() == Class->ClassGeneratedBy->GetFName())
+				{
+					// it's the userwidget
+					return true;
+				}
 
-			if (Binding.WidgetPath.GetWidgetName() != Context.GetWidgetBlueprint()->GetFName()) // is the userwidget itself
-			{
-				if (!WidgetSources.Contains(Binding.WidgetPath.GetWidgetName()))
+				// If the widget doesn't have a property, add one automatically.
+				if (!WidgetSources.Contains(PropertyPath.GetWidgetName()))
 				{
-					WidgetSources.Add(Binding.WidgetPath.GetWidgetName());
+					WidgetSources.Add(PropertyPath.GetWidgetName());
 
-					UWidget** WidgetPtr = WidgetNameToWidgetPointerMap.Find(Binding.WidgetPath.GetWidgetName());
+					UWidget** WidgetPtr = Self->WidgetNameToWidgetPointerMap.Find(PropertyPath.GetWidgetName());
 					if (WidgetPtr == nullptr || *WidgetPtr == nullptr)
 					{
-						WidgetBlueprintCompilerContext.MessageLog.Error(*FString::Printf(TEXT("The widget '%s' could not be found."), *Binding.WidgetPath.GetWidgetName().ToString()));
-						bAreSourceContextsValid = false;
-						continue;
+						Self->WidgetBlueprintCompilerContext.MessageLog.Error(*FString::Printf(TEXT("The binding number %d has an invalid path. The widget '%s' could not be found."), Index, *PropertyPath.GetWidgetName().ToString()));
+						return false;
 					}
 
 					UWidget* Widget = *WidgetPtr;
 					FCompilerSourceContext SourceVariable;
 					SourceVariable.Class = Widget->GetClass();
-					SourceVariable.PropertyName = Binding.WidgetPath.GetWidgetName();
+					SourceVariable.PropertyName = PropertyPath.GetWidgetName();
 					SourceVariable.DisplayName = Widget->GetDisplayLabel();
 					SourceVariable.CategoryName = TEXT("Widget");
-					SourceContexts.Emplace(MoveTemp(SourceVariable));
+					Self->SourceContexts.Emplace(MoveTemp(SourceVariable));
+				}
+			}
+			else if (PropertyPath.IsFromViewModel())
+			{
+				const FMVVMBlueprintViewModelContext* SourceViewModelContext = BlueprintView->FindViewModel(PropertyPath.GetViewModelId());
+				if (SourceViewModelContext == nullptr)
+				{
+					Self->WidgetBlueprintCompilerContext.MessageLog.Error(*FString::Printf(TEXT("The binding number %d has an invalid viewmodel."), Index));
+					return false;
+				}
+
+				bool bFound = ViewModelsGuid.Contains(SourceViewModelContext->GetViewModelId());
+				if (!bFound)
+				{
+					Self->WidgetBlueprintCompilerContext.MessageLog.Error(*FString::Printf(TEXT("The binding number %d has a viewmodel that was not added."), Index));
+					return false;
 				}
 			}
+			else
+			{
+				Self->WidgetBlueprintCompilerContext.MessageLog.Error(*FString::Printf(TEXT("The binding number %d has a unknown path type."), Index));
+				return false;
+			}
+			return true;
+		};
 
-			// todo do the same for the function arguments
+		if (IsForwardBinding(Binding.BindingType) || IsBackwardBinding(Binding.BindingType))
+		{
+			bAreSourceContextsValid = GenerateCompilerSourceContext(Context.GetSkeletonGeneratedClass(), BlueprintView, Binding.ViewModelPath)
+				&& GenerateCompilerSourceContext(Context.GetSkeletonGeneratedClass(), BlueprintView, Binding.WidgetPath)
+				&& bAreSourceContextsValid;
 		}
 	}
 }
@@ -390,7 +404,7 @@ void FMVVMViewBlueprintCompiler::CreateFunctionsDeclaration(const FWidgetBluepri
 {
 	for (FCompilerSourceCreatorContext& SourceCreator : SourceCreatorContexts)
 	{
-		if (!SourceCreator.SetterFunctionName.IsEmpty())
+		if (!SourceCreator.SetterFunctionName.IsEmpty() && SourceCreator.Type == ECompilerSourceCreatorType::ViewModel)
 		{
 			ensure(SourceCreator.SetterGraph == nullptr);
 			SourceCreator.SetterGraph = UE::MVVM::FunctionGraphHelper::CreateIntermediateFunctionGraph(
@@ -418,6 +432,11 @@ bool FMVVMViewBlueprintCompiler::PreCompile(UWidgetBlueprintGeneratedClass* Clas
 		return false;
 	}
 
+	const int32 NumBindings = BlueprintView->GetNumBindings();
+	Bindings.Reset(NumBindings*2);
+	BindingSourceContexts.Reset(NumBindings*2);
+
+	PreCompileBindingSources(Class, BlueprintView);
 	PreCompileSourceCreators(Class, BlueprintView);
 	PreCompileBindings(Class, BlueprintView);
 
@@ -451,6 +470,104 @@ bool FMVVMViewBlueprintCompiler::Compile(UWidgetBlueprintGeneratedClass* Class,
 }
 
 
+bool FMVVMViewBlueprintCompiler::PreCompileBindingSources(UWidgetBlueprintGeneratedClass* Class, UMVVMBlueprintView* BlueprintView)
+{
+	const int32 NumBindings = BlueprintView->GetNumBindings();
+	for (int32 Index = 0; Index < NumBindings; ++Index)
+	{
+		FMVVMBlueprintViewBinding* BindingPtr = BlueprintView->GetBindingAt(Index);
+		if (BindingPtr == nullptr)
+		{
+			WidgetBlueprintCompilerContext.MessageLog.Error(*FString::Printf(TEXT("The binding at index '%d' is invalid."), Index));
+			bIsBindingsValid = false;
+			continue;
+		}
+
+		FMVVMBlueprintViewBinding& Binding = *BindingPtr;
+		if (!Binding.bCompile)
+		{
+			continue;
+		}
+
+		if (IsForwardBinding(Binding.BindingType))
+		{
+			const TValueOrError<FBindingSourceContext, FString> CreatedBindingSourceContext = CreateBindingSourceContext(BlueprintView, Class, Binding.ViewModelPath);
+			if (CreatedBindingSourceContext.HasError())
+			{
+				AddErrorForBinding(Binding, BlueprintView, FString::Printf(TEXT("The property path '%s' is invalid. %s")
+					, *PropertyPathToString(BlueprintView, Binding.ViewModelPath)
+					, *CreatedBindingSourceContext.GetError()));
+				bIsBindingsValid = false;
+				continue;
+			}
+
+			FBindingSourceContext BindingSourceContext = CreatedBindingSourceContext.GetValue();
+			if (!IsPropertyPathValid(BindingSourceContext.PropertyPath))
+			{
+				AddErrorForBinding(Binding, BlueprintView, FString::Printf(TEXT("The property path '%s' is invalid."), *PropertyPathToString(BlueprintView, Binding.ViewModelPath)));
+				bIsBindingsValid = false;
+				continue;
+			}
+			if (BindingSourceContext.SourceClass == nullptr)
+			{
+				AddErrorForBinding(Binding, BlueprintView, TEXT("Internal error. The binding could not find its source."));
+				bIsBindingsValid = false;
+				continue;
+			}
+			if (!BindingSourceContext.bIsRootWidget && BindingSourceContext.CompilerSourceContextIndex == INDEX_NONE)
+			{
+				AddErrorForBinding(Binding, BlueprintView, TEXT("Internal error. The binding could not find its source."));
+				bIsBindingsValid = false;
+				continue;
+			}
+
+			BindingSourceContext.BindingIndex = Index;
+			BindingSourceContext.bIsForwardBinding = true;
+			BindingSourceContexts.Add(BindingSourceContext);
+		}
+
+		if (IsBackwardBinding(Binding.BindingType))
+		{
+			const TValueOrError<FBindingSourceContext, FString> CreatedBindingSourceContext = CreateBindingSourceContext(BlueprintView, Class, Binding.WidgetPath);
+			if (CreatedBindingSourceContext.HasError())
+			{
+				AddErrorForBinding(Binding, BlueprintView, FString::Printf(TEXT("The property path '%s' is invalid. %s")
+					, *PropertyPathToString(BlueprintView, Binding.ViewModelPath)
+					, *CreatedBindingSourceContext.GetError()));
+				bIsBindingsValid = false;
+				continue;
+			}
+
+			FBindingSourceContext BindingSourceContext = CreatedBindingSourceContext.GetValue();
+			if (!IsPropertyPathValid(BindingSourceContext.PropertyPath))
+			{
+				AddErrorForBinding(Binding, BlueprintView, FString::Printf(TEXT("The property path '%s' is invalid."), *PropertyPathToString(BlueprintView, Binding.WidgetPath)));
+				bIsBindingsValid = false;
+				continue;
+			}
+			if (BindingSourceContext.SourceClass == nullptr)
+			{
+				AddErrorForBinding(Binding, BlueprintView, TEXT("Internal error. The binding could not find its source."));
+				bIsBindingsValid = false;
+				continue;
+			}
+			if (!BindingSourceContext.bIsRootWidget && BindingSourceContext.CompilerSourceContextIndex == INDEX_NONE)
+			{
+				AddErrorForBinding(Binding, BlueprintView, TEXT("Internal error. The binding could not find its source."));
+				bIsBindingsValid = false;
+				continue;
+			}
+
+			BindingSourceContext.BindingIndex = Index;
+			BindingSourceContext.bIsForwardBinding = false;
+			BindingSourceContexts.Add(BindingSourceContext);
+		}
+	}
+
+	return bIsBindingsValid;
+}
+
+
 bool FMVVMViewBlueprintCompiler::PreCompileSourceCreators(UWidgetBlueprintGeneratedClass* Class, UMVVMBlueprintView* BlueprintView)
 {
 	if (!bAreSourcesCreatorValid)
@@ -465,11 +582,11 @@ bool FMVVMViewBlueprintCompiler::PreCompileSourceCreators(UWidgetBlueprintGenera
 		if (SourceCreatorContext.Type == ECompilerSourceCreatorType::ViewModel)
 		{
 			const FMVVMBlueprintViewModelContext& ViewModelContext = SourceCreatorContext.ViewModelContext;
-			checkf(ViewModelContext.GetViewModelClass(), TEXT("The ViewModel class is invalid. It was checked in CreateSourceList"));
+			checkf(ViewModelContext.GetViewModelClass(), TEXT("The viewmodel class is invalid. It was checked in CreateSourceList"));
 
 			if (ViewModelContext.GetViewModelClass()->HasAllClassFlags(CLASS_Deprecated))
 			{
-				WidgetBlueprintCompilerContext.MessageLog.Warning(*FString::Printf(TEXT("The ViewModel type '%s' is deprecated and should not be used for '%s'. Please update it in the View Binding panel under Manage ViewModels.")
+				WidgetBlueprintCompilerContext.MessageLog.Warning(*FString::Printf(TEXT("The viewmodel type '%s' is deprecated and should not be used for '%s'. Please update it in the View Binding panel under Manage ViewModels.")
 					, *ViewModelContext.GetViewModelClass()->GetDisplayNameText().ToString()
 					, *ViewModelContext.GetDisplayName().ToString()));
 			}
@@ -481,7 +598,7 @@ bool FMVVMViewBlueprintCompiler::PreCompileSourceCreators(UWidgetBlueprintGenera
 			{
 				if (ViewModelContext.GetViewModelClass()->HasAllClassFlags(CLASS_Abstract))
 				{
-					WidgetBlueprintCompilerContext.MessageLog.Error(*FString::Printf(TEXT("The ViewModel type '%s' is abstract and can't be created for '%s'. You can change it in the View Binding panel under Manage ViewModels.")
+					WidgetBlueprintCompilerContext.MessageLog.Error(*FString::Printf(TEXT("The viewmodel type '%s' is abstract and can't be created for '%s'. You can change it in the View Binding panel under Manage ViewModels.")
 						, *ViewModelContext.GetViewModelClass()->GetDisplayNameText().ToString()
 						, *ViewModelContext.GetDisplayName().ToString()));
 					bAreSourcesCreatorValid = false;
@@ -492,17 +609,28 @@ bool FMVVMViewBlueprintCompiler::PreCompileSourceCreators(UWidgetBlueprintGenera
 			{
 				if (ViewModelContext.ViewModelPropertyPath.IsEmpty())
 				{
-					WidgetBlueprintCompilerContext.MessageLog.Error(*FString::Printf(TEXT("The ViewModel '%s' has an invalid Getter. You can select a new one in the View Binding panel under Manage ViewModels.")
+					WidgetBlueprintCompilerContext.MessageLog.Error(*FString::Printf(TEXT("The viewmodel '%s' has an invalid Getter. You can select a new one in the View Binding panel under Manage ViewModels.")
 						, *ViewModelContext.GetDisplayName().ToString()));
 					bAreSourcesCreatorValid = true;
 					continue;
 				}
 
 				// Generate a path to read the value at runtime
-				TValueOrError<FCompiledBindingLibraryCompiler::FFieldPathHandle, FString> ReadFieldPathResult = BindingLibraryCompiler.AddObjectFieldPath(Class, ViewModelContext.ViewModelPropertyPath, ViewModelContext.GetViewModelClass(), true);
+				TValueOrError<TArray<FMVVMConstFieldVariant>, FString> GeneratedField = FieldPathHelper::GenerateFieldPathList(Class, ViewModelContext.ViewModelPropertyPath, true);
+				if (GeneratedField.HasError())
+				{
+					WidgetBlueprintCompilerContext.MessageLog.Error(*FString::Printf(TEXT("The viewmodel '%s' (%s) has an invalid Getter. %s")
+						, *ViewModelContext.GetDisplayName().ToString()
+						, *ViewModelContext.GetViewModelClass()->GetDisplayNameText().ToString()
+						, *GeneratedField.GetError()));
+					bAreSourcesCreatorValid = false;
+					continue;
+				}
+
+				TValueOrError<FCompiledBindingLibraryCompiler::FFieldPathHandle, FString> ReadFieldPathResult = BindingLibraryCompiler.AddObjectFieldPath(GeneratedField.GetValue(), ViewModelContext.GetViewModelClass(), true);
 				if (ReadFieldPathResult.HasError())
 				{
-					WidgetBlueprintCompilerContext.MessageLog.Error(*FString::Printf(TEXT("The ViewModel '%s' (%s) as an invalid Getter. You can select a new one in the View Binding panel under Manage ViewModels. Details: %s")
+					WidgetBlueprintCompilerContext.MessageLog.Error(*FString::Printf(TEXT("The viewmodel '%s' (%s) has an invalid Getter. %s")
 						, *ViewModelContext.GetDisplayName().ToString()
 						, *ViewModelContext.GetViewModelClass()->GetDisplayNameText().ToString()
 						, *ReadFieldPathResult.GetError()));
@@ -516,7 +644,7 @@ bool FMVVMViewBlueprintCompiler::PreCompileSourceCreators(UWidgetBlueprintGenera
 			{
 				if (ViewModelContext.GlobalViewModelIdentifier.IsNone())
 				{
-					WidgetBlueprintCompilerContext.MessageLog.Error(*FString::Printf(TEXT("The ViewModel '%s' (%s) doesn't have a valid Global identifier. You can specify a new one in the View Binding panel under Manage ViewModels.")
+					WidgetBlueprintCompilerContext.MessageLog.Error(*FString::Printf(TEXT("The viewmodel '%s' (%s) doesn't have a valid Global identifier. You can specify a new one in the View Binding panel under Manage ViewModels.")
 						, *ViewModelContext.GetDisplayName().ToString()
 						, *ViewModelContext.GetViewModelClass()->GetDisplayNameText().ToString()));
 					bAreSourcesCreatorValid = false;
@@ -525,7 +653,7 @@ bool FMVVMViewBlueprintCompiler::PreCompileSourceCreators(UWidgetBlueprintGenera
 			}
 			else
 			{
-				WidgetBlueprintCompilerContext.MessageLog.Error(*FString::Printf(TEXT("The ViewModel '%s' (%s) doesn't have a valid creation type. You can select one in the View Binding panel under Manage ViewModels.")
+				WidgetBlueprintCompilerContext.MessageLog.Error(*FString::Printf(TEXT("The viewmodel '%s' (%s) doesn't have a valid creation type. You can select one in the View Binding panel under Manage ViewModels.")
 					, *ViewModelContext.GetDisplayName().ToString()
 					, *ViewModelContext.GetViewModelClass()->GetDisplayNameText().ToString()));
 				bAreSourcesCreatorValid = true;
@@ -614,29 +742,16 @@ bool FMVVMViewBlueprintCompiler::CompileSourceCreators(const FCompiledBindingLib
 
 bool FMVVMViewBlueprintCompiler::PreCompileBindings(UWidgetBlueprintGeneratedClass* Class, UMVVMBlueprintView* BlueprintView)
 {
-	if (!bAreSourceContextsValid)
+	if (!bAreSourceContextsValid || !bIsBindingsValid)
 	{
 		return false;
 	}
 
-	const int32 NumBindings = BlueprintView->GetNumBindings();
-	Bindings.Reset(NumBindings);
-
-	for (int32 Index = 0; Index < BlueprintView->GetNumBindings(); ++Index)
+	for (const FBindingSourceContext& BindingSourceContext : BindingSourceContexts)
 	{
-		FMVVMBlueprintViewBinding* BindingPtr = BlueprintView->GetBindingAt(Index);
-		if (BindingPtr == nullptr)
-		{
-			WidgetBlueprintCompilerContext.MessageLog.Error(*FString::Printf(TEXT("The binding at index '%d' is invalid."), Index));
-			bAreSourceContextsValid = false;
-			continue;
-		}
-
+		FMVVMBlueprintViewBinding* BindingPtr = BlueprintView->GetBindingAt(BindingSourceContext.BindingIndex);
+		check(BindingPtr);
 		FMVVMBlueprintViewBinding& Binding = *BindingPtr;
-		if (!Binding.bCompile)
-		{
-			continue;
-		}
 
 		FMVVMViewBlueprintCompiler* Self = this;
 		auto AddBinding = [Self](UWidgetBlueprintGeneratedClass* Class, const TArrayView<TArray<UE::MVVM::FMVVMConstFieldVariant>> GetterFields, TArrayView<UE::MVVM::FMVVMConstFieldVariant> SetterFields, const UFunction* ConversionFunction) -> TValueOrError<FCompilerBinding, FString>
@@ -646,7 +761,15 @@ bool FMVVMViewBlueprintCompiler::PreCompileBindings(UWidgetBlueprintGeneratedCla
 			for (const TArray<UE::MVVM::FMVVMConstFieldVariant>& GetterField : GetterFields)
 			{
 				// Generate a path to read the value at runtime
-				TValueOrError<FCompiledBindingLibraryCompiler::FFieldPathHandle, FString> FieldPathResult = Self->BindingLibraryCompiler.AddFieldPath(GetterField, true);
+				TValueOrError<TArray<FMVVMConstFieldVariant>, FString> GeneratedField = FieldPathHelper::GenerateFieldPathList(GetterField, true);
+				if (GeneratedField.HasError())
+				{
+					return MakeError(FString::Printf(TEXT("Couldn't create the source field path '%s'. %s")
+						, *::UE::MVVM::FieldPathHelper::ToString(GetterField)
+						, *GeneratedField.GetError()));
+				}
+
+				TValueOrError<FCompiledBindingLibraryCompiler::FFieldPathHandle, FString> FieldPathResult = Self->BindingLibraryCompiler.AddFieldPath(GeneratedField.GetValue(), true);
 				if (FieldPathResult.HasError())
 				{
 					return MakeError(FString::Printf(TEXT("Couldn't create the source field path '%s'. %s")
@@ -657,7 +780,15 @@ bool FMVVMViewBlueprintCompiler::PreCompileBindings(UWidgetBlueprintGeneratedCla
 			}
 
 			{
-				TValueOrError<FCompiledBindingLibraryCompiler::FFieldPathHandle, FString> FieldPathResult = Self->BindingLibraryCompiler.AddFieldPath(SetterFields, false);
+				TValueOrError<TArray<FMVVMConstFieldVariant>, FString> GeneratedField = FieldPathHelper::GenerateFieldPathList(SetterFields, false);
+				if (GeneratedField.HasError())
+				{
+					return MakeError(FString::Printf(TEXT("Couldn't create the destination field path '%s'. %s")
+						, *::UE::MVVM::FieldPathHelper::ToString(SetterFields)
+						, *GeneratedField.GetError()));
+				}
+
+				TValueOrError<FCompiledBindingLibraryCompiler::FFieldPathHandle, FString> FieldPathResult = Self->BindingLibraryCompiler.AddFieldPath(GeneratedField.GetValue(), false);
 				if (FieldPathResult.HasError())
 				{
 					return MakeError(FString::Printf(TEXT("Couldn't create the destination field path '%s'. %s")
@@ -692,8 +823,6 @@ bool FMVVMViewBlueprintCompiler::PreCompileBindings(UWidgetBlueprintGeneratedCla
 
 		auto AddFieldId = [Self](UClass* SourceContextClass, bool bNotifyFieldValueChangedRequired, EMVVMBindingMode BindingMode, FName FieldToListenTo) -> TValueOrError<FCompiledBindingLibraryCompiler::FFieldIdHandle, FString>
 		{
-			UE::FieldNotification::FFieldId FoundFieldId;
-
 			if (!IsOneTimeBinding(BindingMode) && bNotifyFieldValueChangedRequired)
 			{
 				return Self->BindingLibraryCompiler.AddFieldId(SourceContextClass, FieldToListenTo);
@@ -701,157 +830,59 @@ bool FMVVMViewBlueprintCompiler::PreCompileBindings(UWidgetBlueprintGeneratedCla
 			return MakeValue(FCompiledBindingLibraryCompiler::FFieldIdHandle());
 		};
 
-		if (IsForwardBinding(Binding.BindingType))
+		TValueOrError<FCompiledBindingLibraryCompiler::FFieldIdHandle, FString> AddFieldResult = AddFieldId(BindingSourceContext.SourceClass, true, Binding.BindingType, BindingSourceContext.FieldId.GetFieldName());
+		if (AddFieldResult.HasError())
 		{
-			const FBindingSourceContext BindingSourceContext = CreateBindingSourceContext(BlueprintView, Class, Binding.ViewModelPath);
-			if (!IsPropertyPathValid(BindingSourceContext.PropertyPath))
-			{
-				AddErrorForBinding(Binding, BlueprintView, FString::Printf(TEXT("The property path '%s' is invalid."), *PropertyPathToString(BlueprintView, Binding.ViewModelPath)));
-				bIsBindingsValid = false;
-				continue;
-			}
-			if (BindingSourceContext.bIsRootWidget && BindingSourceContext.SourceClass == nullptr)
-			{
-				AddErrorForBinding(Binding, BlueprintView, TEXT("Internal error. The binding could not find its source."));
-				bIsBindingsValid = false;
-				continue;
-			}
-			else if (BindingSourceContext.bIsRootWidget && (BindingSourceContext.SourceClass == nullptr || BindingSourceContext.CompilerSourceContextIndex == INDEX_NONE))
-			{
-				AddErrorForBinding(Binding, BlueprintView, TEXT("Internal error. The binding could not find its source."));
-				bIsBindingsValid = false;
-				continue;
-			}
-
-			TValueOrError<FCompiledBindingLibraryCompiler::FFieldIdHandle, FString> AddFieldResult = AddFieldId(BindingSourceContext.SourceClass, true, Binding.BindingType, BindingSourceContext.FieldId);
-			if (AddFieldResult.HasError())
-			{
-				AddErrorForBinding(Binding, BlueprintView, *FString::Printf(TEXT("The binding could not create its source. %s"), *AddFieldResult.GetError()));
-				bIsBindingsValid = false;
-				continue;
-			}
-
-			TArray<TArray<UE::MVVM::FMVVMConstFieldVariant>> GetterFields;
-			//if (Binding.Conversion.HasConversionFunction && has read values defined (that are not the default one))
-			//{
-			//  Add the read values for the conversion function
-			//}
-			//else
-			{
-				GetterFields.Add(BindingSourceContext.PropertyPath);
-			}
-
-			TArray<UE::MVVM::FMVVMConstFieldVariant> SetterPath;
-			{
-				SetterPath = CreateBindingDestinationPath(BlueprintView, Class, Binding.WidgetPath);
-				if (!IsPropertyPathValid(SetterPath))
-				{
-					AddErrorForBinding(Binding, BlueprintView, FString::Printf(TEXT("The property path '%s' is invalid."), *PropertyPathToString(BlueprintView, Binding.WidgetPath)));
-					bIsBindingsValid = false;
-					continue;
-				}
-			}
-
-			FMemberReference ConversionFunctionRef = Binding.Conversion.SourceToDestinationFunction;
-			if (!Binding.Conversion.SourceToDestinationWrapper.IsNone())
-			{
-				ConversionFunctionRef.SetSelfMember(Binding.Conversion.SourceToDestinationWrapper);
-			}
-
-			const UFunction* ConversionFunction = ConversionFunctionRef.ResolveMember<UFunction>(Class);
-
-			TValueOrError<FCompilerBinding, FString> AddBindingResult = AddBinding(Class, GetterFields, SetterPath, ConversionFunction);
-			if (AddBindingResult.HasError())
-			{
-				AddErrorForBinding(Binding, BlueprintView, FString::Printf(TEXT("The binding could not be created. %s"), *AddBindingResult.GetError()));
-				bIsBindingsValid = false;
-				continue;
-			}
-
-			FCompilerBinding NewBinding = AddBindingResult.StealValue();
-			NewBinding.BindingIndex = Index;
-			NewBinding.SourceContextIndex = BindingSourceContext.CompilerSourceContextIndex;
-			NewBinding.FieldIdHandle = AddFieldResult.StealValue();
-			NewBinding.bSourceIsUserWidget = BindingSourceContext.bIsRootWidget;
-			NewBinding.bFieldIdNeeded = !IsOneTimeBinding(Binding.BindingType);
-			Bindings.Emplace(NewBinding);
+			AddErrorForBinding(Binding, BlueprintView, *FString::Printf(TEXT("The binding could not create its source. %s"), *AddFieldResult.GetError()));
+			bIsBindingsValid = false;
+			continue;
 		}
 
-		if (IsBackwardBinding(Binding.BindingType))
+		TArray<TArray<UE::MVVM::FMVVMConstFieldVariant>> GetterFields;
+		//if (Binding.Conversion.HasConversionFunction && has read values defined (that are not the default one))
+		//{
+		//  Add the read values for the conversion function
+		//}
+		//else
 		{
-			const FBindingSourceContext BindingSourceContext = CreateBindingSourceContext(BlueprintView, Class, Binding.WidgetPath);
-			if (!IsPropertyPathValid(BindingSourceContext.PropertyPath))
-			{
-				AddErrorForBinding(Binding, BlueprintView, FString::Printf(TEXT("The property path '%s' is invalid."), *PropertyPathToString(BlueprintView, Binding.WidgetPath)));
-				bIsBindingsValid = false;
-				continue;
-			}
-			if (BindingSourceContext.bIsRootWidget && BindingSourceContext.SourceClass == nullptr)
-			{
-				AddErrorForBinding(Binding, BlueprintView, TEXT("Internal error. The binding could not find its source."));
-				bIsBindingsValid = false;
-				continue;
-			}
-			else if (!BindingSourceContext.bIsRootWidget && (BindingSourceContext.SourceClass == nullptr || BindingSourceContext.CompilerSourceContextIndex == INDEX_NONE))
-			{
-				AddErrorForBinding(Binding, BlueprintView, TEXT("Internal error. The binding could not find its source."));
-				bIsBindingsValid = false;
-				continue;
-			}
+			GetterFields.Add(BindingSourceContext.PropertyPath);
+		}
 
-			TValueOrError<FCompiledBindingLibraryCompiler::FFieldIdHandle, FString> AddFieldResult = AddFieldId(BindingSourceContext.SourceClass, true, Binding.BindingType, BindingSourceContext.FieldId);
-			if (AddFieldResult.HasError())
+		TArray<UE::MVVM::FMVVMConstFieldVariant> SetterPath;
+		{
+			const FMVVMBlueprintPropertyPath& DestinationPath = BindingSourceContext.bIsForwardBinding ? Binding.WidgetPath : Binding.ViewModelPath;
+			SetterPath = CreateBindingDestinationPath(BlueprintView, Class, DestinationPath);
+			if (!IsPropertyPathValid(SetterPath))
 			{
-				AddErrorForBinding(Binding, BlueprintView, FString::Printf(TEXT("The binding could not create it's source. %s"), *AddFieldResult.GetError()));
+				AddErrorForBinding(Binding, BlueprintView, FString::Printf(TEXT("The property path '%s' is invalid."), *PropertyPathToString(BlueprintView, DestinationPath)));
 				bIsBindingsValid = false;
 				continue;
 			}
+		}
 
-			TArray<TArray<UE::MVVM::FMVVMConstFieldVariant>> GetterFields;
-			//if (Binding.Conversion.HasConversionFunction && has read values defined (that are not the default one))
-			//{
-			//  Add the read values for the conversion function
-			//}
-			//else
-			{
-				GetterFields.Add(BindingSourceContext.PropertyPath);
-			}
-
-			TArray<UE::MVVM::FMVVMConstFieldVariant> SetterPath;
-			{
-				SetterPath = CreateBindingDestinationPath(BlueprintView, Class, Binding.ViewModelPath);
-				if (!IsPropertyPathValid(SetterPath))
-				{
-					AddErrorForBinding(Binding, BlueprintView, FString::Printf(TEXT("The property path '%s' is invalid."), *PropertyPathToString(BlueprintView, Binding.ViewModelPath)));
-					bIsBindingsValid = false;
-					continue;
-				}
-			}
-
-			FMemberReference ConversionFunctionRef = Binding.Conversion.DestinationToSourceFunction;
-			if (!Binding.Conversion.DestinationToSourceWrapper.IsNone())
-			{
-				ConversionFunctionRef.SetSelfMember(Binding.Conversion.DestinationToSourceWrapper);
-			}
-
-			const UFunction* ConversionFunction = ConversionFunctionRef.ResolveMember<UFunction>(Class);
+		FMemberReference ConversionFunctionRef = Binding.Conversion.SourceToDestinationFunction;
+		if (!Binding.Conversion.SourceToDestinationWrapper.IsNone())
+		{
+			ConversionFunctionRef.SetSelfMember(Binding.Conversion.SourceToDestinationWrapper);
+		}
 
-			TValueOrError<FCompilerBinding, FString> AddBindingResult = AddBinding(Class, GetterFields, SetterPath, ConversionFunction);
-			if (AddBindingResult.HasError())
-			{
-				AddErrorForBinding(Binding, BlueprintView, FString::Printf(TEXT("The binding could not be created. %s"), *AddBindingResult.GetError()));
-				bIsBindingsValid = false;
-				continue;
-			}
+		const UFunction* ConversionFunction = ConversionFunctionRef.ResolveMember<UFunction>(Class);
 
-			FCompilerBinding NewBinding = AddBindingResult.StealValue();
-			NewBinding.BindingIndex = Index;
-			NewBinding.SourceContextIndex = BindingSourceContext.CompilerSourceContextIndex;
-			NewBinding.FieldIdHandle = AddFieldResult.StealValue();
-			NewBinding.bSourceIsUserWidget = BindingSourceContext.bIsRootWidget;
-			NewBinding.bFieldIdNeeded = true; // one time from destination to source doesn't exist
-			Bindings.Emplace(NewBinding);
+		TValueOrError<FCompilerBinding, FString> AddBindingResult = AddBinding(Class, GetterFields, SetterPath, ConversionFunction);
+		if (AddBindingResult.HasError())
+		{
+			AddErrorForBinding(Binding, BlueprintView, FString::Printf(TEXT("The binding could not be created. %s"), *AddBindingResult.GetError()));
+			bIsBindingsValid = false;
+			continue;
 		}
+
+		FCompilerBinding NewBinding = AddBindingResult.StealValue();
+		NewBinding.BindingIndex = BindingSourceContext.BindingIndex;
+		NewBinding.SourceContextIndex = BindingSourceContext.CompilerSourceContextIndex;
+		NewBinding.FieldIdHandle = AddFieldResult.StealValue();
+		NewBinding.bSourceIsUserWidget = BindingSourceContext.bIsRootWidget;
+		NewBinding.bFieldIdNeeded = !IsOneTimeBinding(Binding.BindingType);
+		Bindings.Emplace(NewBinding);
 	}
 
 	return bIsBindingsValid;
@@ -911,27 +942,22 @@ bool FMVVMViewBlueprintCompiler::CompileBindings(const FCompiledBindingLibraryCo
 
 const FMVVMViewBlueprintCompiler::FCompilerSourceCreatorContext* FMVVMViewBlueprintCompiler::FindViewModelSource(FGuid Id) const
 {
-	return SourceCreatorContexts.FindByPredicate([Id](const FCompilerSourceCreatorContext& Other) { return Other.Type == ECompilerSourceCreatorType::ViewModel ? Other.ViewModelContext.GetViewModelId() == Id : false; });
+	return SourceCreatorContexts.FindByPredicate([Id](const FCompilerSourceCreatorContext& Other)
+		{
+			return Other.Type == ECompilerSourceCreatorType::ViewModel ? Other.ViewModelContext.GetViewModelId() == Id : false;
+		});
 }
 
 
-FMVVMViewBlueprintCompiler::FBindingSourceContext FMVVMViewBlueprintCompiler::CreateBindingSourceContext(const UMVVMBlueprintView* BlueprintView, const UWidgetBlueprintGeneratedClass* Class, const FMVVMBlueprintPropertyPath& PropertyPath) const
+TValueOrError<FMVVMViewBlueprintCompiler::FBindingSourceContext, FString> FMVVMViewBlueprintCompiler::CreateBindingSourceContext(const UMVVMBlueprintView* BlueprintView, const UWidgetBlueprintGeneratedClass* Class, const FMVVMBlueprintPropertyPath& PropertyPath)
 {
-	FBindingSourceContext Result;
-
 	if (PropertyPath.IsEmpty())
 	{
-		ensureAlwaysMsgf(false, TEXT("Not supported yet."));
-		return Result;
-	}
-
-	// Todo the path may contains another INotifyFieldValueChanged
-	TArray<FName> Paths = PropertyPath.GetPaths();
-	if (Paths.Num() > 0)
-	{
-		Result.FieldId = Paths[0];
+		ensureAlways(false);
+		return MakeError(TEXT("Empty property path found. This is legal."));
 	}
 
+	FBindingSourceContext Result;
 	if (PropertyPath.IsFromViewModel())
 	{
 		Result.bIsRootWidget = false;
@@ -968,7 +994,57 @@ FMVVMViewBlueprintCompiler::FBindingSourceContext FMVVMViewBlueprintCompiler::Cr
 		ensureAlwaysMsgf(false, TEXT("Not supported yet."));
 	}
 
-	return Result;
+	// The path may contains another INotifyFieldValueChanged
+	TValueOrError<FieldPathHelper::FParsedBindingInfo, FString> BindingInfoResult = FieldPathHelper::GetBindingInfoFromFieldPath(Class, Result.PropertyPath);
+	if (BindingInfoResult.HasError())
+	{
+		return MakeError(BindingInfoResult.StealError());
+	}
+
+	const FieldPathHelper::FParsedBindingInfo& BindingInfo = BindingInfoResult.GetValue();
+	Result.FieldId = BindingInfo.NotifyFieldId;
+
+	const bool bCreateADynamicSourceEntry = BindingInfo.ViewModelIndex > 1
+		&& BindingInfo.NotifyFieldClass != nullptr;
+
+	// The INotifyFieldValueChanged/viewmodel is not the fist property. Creates the ViaBinding source entries.
+	if (bCreateADynamicSourceEntry)
+	{
+		Result.SourceClass = BindingInfo.NotifyFieldClass;
+		Result.CompilerSourceContextIndex = INDEX_NONE;
+		Result.bIsRootWidget = false;
+
+		FString InterfacePath = FieldPathHelper::ToString(BindingInfo.NotifyFieldInterfacePath);
+		int32 SourceCreatorContextIndex = SourceCreatorContexts.IndexOfByPredicate([&InterfacePath](const FCompilerSourceCreatorContext& Other)
+			{
+				return (Other.Type == ECompilerSourceCreatorType::ViewModel)
+					? (Other.ViewModelContext.CreationType == EMVVMBlueprintViewModelContextCreationType::PropertyPath
+						&& Other.ViewModelContext.ViewModelPropertyPath == InterfacePath)
+					: false;
+			});
+
+
+		// Add if the path doesn't already exists
+		if (!SourceCreatorContexts.IsValidIndex(SourceCreatorContextIndex))
+		{
+			return MakeError(TEXT("The field in the path is from a viewmodel that is not the root of the path. This is not yet supported."));
+		}
+
+		Result.PropertyPath.RemoveAt(0, BindingInfo.ViewModelIndex);
+		if (SourceCreatorContexts[SourceCreatorContextIndex].Type == ECompilerSourceCreatorType::ViewModel)
+		{
+			Result.CompilerSourceContextIndex = SourceCreatorContextIndex;
+		}
+	}
+	else
+	{
+		if (Result.SourceClass != BindingInfo.NotifyFieldClass)
+		{
+			return MakeError(TEXT("Internal Error. The class doesn't match"));
+		}
+	}
+
+	return MakeValue(Result);
 }
 
 
@@ -976,7 +1052,7 @@ TArray<FMVVMConstFieldVariant> FMVVMViewBlueprintCompiler::CreateBindingDestinat
 {
 	if (PropertyPath.IsEmpty())
 	{
-		ensureAlwaysMsgf(false, TEXT("Not supported yet."));
+		ensureAlwaysMsgf(false, TEXT("Empty property path found. This is legal."));
 		return TArray<FMVVMConstFieldVariant>();
 	}
 
diff --git a/Engine/Plugins/Runtime/ModelViewViewModel/Source/ModelViewViewModelBlueprint/Private/MVVMViewBlueprintCompiler.h b/Engine/Plugins/Runtime/ModelViewViewModel/Source/ModelViewViewModelBlueprint/Private/MVVMViewBlueprintCompiler.h
--- a/Engine/Plugins/Runtime/ModelViewViewModel/Source/ModelViewViewModelBlueprint/Private/MVVMViewBlueprintCompiler.h
+++ b/Engine/Plugins/Runtime/ModelViewViewModel/Source/ModelViewViewModelBlueprint/Private/MVVMViewBlueprintCompiler.h
@@ -60,6 +60,9 @@ private:
 	void CreateSourceLists(const FWidgetBlueprintCompilerContext::FCreateVariableContext& Context, UMVVMBlueprintView* BlueprintView);
 	void CreateFunctionsDeclaration(const FWidgetBlueprintCompilerContext::FCreateVariableContext& Context, UMVVMBlueprintView* BlueprintView);
 
+	bool PreCompileBindingSources(UWidgetBlueprintGeneratedClass* Class, UMVVMBlueprintView* BlueprintView);
+	bool CompileBindingSources(const FCompiledBindingLibraryCompiler::FCompileResult& CompileResult, UWidgetBlueprintGeneratedClass* Class, UMVVMBlueprintView* BlueprintView, UMVVMViewClass* ViewExtension);
+
 	bool PreCompileSourceCreators(UWidgetBlueprintGeneratedClass* Class, UMVVMBlueprintView* BlueprintView);
 	bool CompileSourceCreators(const FCompiledBindingLibraryCompiler::FCompileResult& CompileResult, UWidgetBlueprintGeneratedClass* Class, UMVVMBlueprintView* BlueprintView, UMVVMViewClass* ViewExtension);
 
@@ -70,7 +73,7 @@ private:
 
 	void AddErrorForBinding(FMVVMBlueprintViewBinding& Binding, const UMVVMBlueprintView* BlueprintView, const FString& Message) const;
 
-	FBindingSourceContext CreateBindingSourceContext(const UMVVMBlueprintView* BlueprintView, const UWidgetBlueprintGeneratedClass* Class, const FMVVMBlueprintPropertyPath& PropertyPath) const;
+	TValueOrError<FBindingSourceContext, FString> CreateBindingSourceContext(const UMVVMBlueprintView* BlueprintView, const UWidgetBlueprintGeneratedClass* Class, const FMVVMBlueprintPropertyPath& PropertyPath);
 	TArray<FMVVMConstFieldVariant> CreateBindingDestinationPath(const UMVVMBlueprintView* BlueprintView, const UWidgetBlueprintGeneratedClass* Class, const FMVVMBlueprintPropertyPath& PropertyPath) const;
 
 	FString PropertyPathToString(const UMVVMBlueprintView* BlueprintView, const FMVVMBlueprintPropertyPath& PropertyPath) const;
@@ -86,6 +89,7 @@ private:
 		FString CategoryName;
 		FString BlueprintSetter;
 		FMVVMConstFieldVariant Field;
+
 		bool bExposeOnSpawn = false;
 	};
 	TArray<FCompilerSourceContext> SourceContexts;
@@ -132,16 +136,23 @@ private:
 
 	struct FBindingSourceContext
 	{
+		int32 BindingIndex = INDEX_NONE;
+		bool bIsForwardBinding = false;
+
 		UClass* SourceClass = nullptr;
-		FName FieldId;
-		// ViewModel.Field.SubProperty.SubProperty
-		// or Widget.Field.SubProperty
-		// or Field.SubProperty (if the widget is the userwidget)
 		// if its the destination then the Field is not necessary
+		FFieldNotificationId FieldId;
+		// Viewmodel.Field.SubProperty.SubProperty
+		// or Widget.Field.SubProperty
+		// or Field.SubProperty (if the widget is the UserWidget)
 		TArray<UE::MVVM::FMVVMConstFieldVariant> PropertyPath;
+
+		// The source if it's a property
 		int32 CompilerSourceContextIndex = INDEX_NONE;
+		// The source is the UserWidget itself
 		bool bIsRootWidget = false;
 	};
+	TArray<FBindingSourceContext> BindingSourceContexts;
 
 	TMap<FName, UWidget*> WidgetNameToWidgetPointerMap;
 	FWidgetBlueprintCompilerContext& WidgetBlueprintCompilerContext;
diff --git a/Engine/Plugins/Runtime/ModelViewViewModel/Source/ModelViewViewModelBlueprint/Public/Bindings/MVVMCompiledBindingLibraryCompiler.h b/Engine/Plugins/Runtime/ModelViewViewModel/Source/ModelViewViewModelBlueprint/Public/Bindings/MVVMCompiledBindingLibraryCompiler.h
--- a/Engine/Plugins/Runtime/ModelViewViewModel/Source/ModelViewViewModelBlueprint/Public/Bindings/MVVMCompiledBindingLibraryCompiler.h
+++ b/Engine/Plugins/Runtime/ModelViewViewModel/Source/ModelViewViewModelBlueprint/Public/Bindings/MVVMCompiledBindingLibraryCompiler.h
@@ -157,15 +157,12 @@
 public:
 	/** */
 	TValueOrError<FFieldIdHandle, FString> AddFieldId(TSubclassOf<UObject> SourceClass, FName FieldName);
-
-	/** */
-	TValueOrError<FFieldPathHandle, FString> AddFieldPath(TSubclassOf<UObject> SourceClass, FStringView FieldPath, bool bRead);
 	
 	/** */
 	TValueOrError<FFieldPathHandle, FString> AddFieldPath(TArrayView<const UE::MVVM::FMVVMConstFieldVariant> FieldPath, bool bRead);
 
 	/** */
-	TValueOrError<FFieldPathHandle, FString> AddObjectFieldPath(TSubclassOf<UObject> SourceClass, FStringView FieldPath, UClass* ExpectedType, bool bRead);
+	TValueOrError<FFieldPathHandle, FString> AddObjectFieldPath(TArrayView<const UE::MVVM::FMVVMConstFieldVariant> FieldPath, UClass* ExpectedType, bool bRead);
 
 	/** */
 	TValueOrError<FFieldPathHandle, FString> AddConversionFunctionFieldPath(TSubclassOf<UObject> SourceClass, const UFunction* Function);
