MVVM - Added method to add a UWidgetBlueprint as a source to SSourceBindingList.

SSourceBindingList now exposes ShowSearchBox.

#rnx

diff --git a/Engine/Plugins/Runtime/ModelViewViewModel/Source/ModelViewViewModelEditor/Private/Widgets/SMVVMViewModelBindingListWidget.cpp b/Engine/Plugins/Runtime/ModelViewViewModel/Source/ModelViewViewModelEditor/Private/Widgets/SMVVMViewModelBindingListWidget.cpp
--- a/Engine/Plugins/Runtime/ModelViewViewModel/Source/ModelViewViewModelEditor/Private/Widgets/SMVVMViewModelBindingListWidget.cpp
+++ b/Engine/Plugins/Runtime/ModelViewViewModel/Source/ModelViewViewModelEditor/Private/Widgets/SMVVMViewModelBindingListWidget.cpp
@@ -119,23 +119,22 @@ void SSourceBindingList::Construct(const FArguments& InArgs, const UWidgetBluepr
 
 	OnDoubleClicked = InArgs._OnDoubleClicked;
 
-	PropertyViewer = SNew(SPropertyViewer)
+	ChildSlot
+	[
+		SAssignNew(PropertyViewer, SPropertyViewer)
+		.FieldIterator(FieldIterator.Get())
 		.PropertyVisibility(SPropertyViewer::EPropertyVisibility::Hidden)
 		.bShowFieldIcon(true)
 		.bSanitizeName(true)
-		.FieldIterator(FieldIterator.Get())
-		.bShowSearchBox(false)
+		.bShowSearchBox(InArgs._ShowSearchBox)
 		.OnSelectionChanged(this, &SSourceBindingList::HandleSelectionChanged)
-		.OnDoubleClicked(this, &SSourceBindingList::HandleDoubleClicked);
-
-	ChildSlot
-	[
-		PropertyViewer.ToSharedRef()
+		.OnDoubleClicked(this, &SSourceBindingList::HandleDoubleClicked)
 	];
 }
 
 void SSourceBindingList::Clear()
 {
+	Sources.Reset();
 	if (PropertyViewer)
 	{
 		PropertyViewer->RemoveAll();
@@ -150,6 +149,16 @@ void SSourceBindingList::AddSource(UClass* Class, FName Name, FGuid Guid)
 	Source.ViewModelId = Guid;
 
 	AddSources(MakeArrayView(&Source, 1));
+} 
+
+void SSourceBindingList::AddWidgetBlueprint(const UWidgetBlueprint* InWidgetBlueprint)
+{
+	FBindingSource Source;
+	Source.Class = InWidgetBlueprint->GeneratedClass;
+	Source.Name = InWidgetBlueprint->GetFName();
+	Source.DisplayName = FText::FromString(InWidgetBlueprint->GetName());
+
+	AddSources(MakeArrayView(&Source, 1 ));
 }
 
 void SSourceBindingList::AddWidgets(TArrayView<const UWidget*> InWidgets)
diff --git a/Engine/Plugins/Runtime/ModelViewViewModel/Source/ModelViewViewModelEditor/Private/Widgets/SMVVMViewModelBindingListWidget.h b/Engine/Plugins/Runtime/ModelViewViewModel/Source/ModelViewViewModelEditor/Private/Widgets/SMVVMViewModelBindingListWidget.h
--- a/Engine/Plugins/Runtime/ModelViewViewModel/Source/ModelViewViewModelEditor/Private/Widgets/SMVVMViewModelBindingListWidget.h
+++ b/Engine/Plugins/Runtime/ModelViewViewModel/Source/ModelViewViewModelEditor/Private/Widgets/SMVVMViewModelBindingListWidget.h
@@ -53,10 +53,9 @@ namespace UE::MVVM
 
 		DECLARE_DELEGATE_OneParam(FOnDoubleClicked, const FMVVMBlueprintPropertyPath&);
 	
-		SLATE_BEGIN_ARGS(SSourceBindingList) : 
-			_FieldVisibilityFlags(EFieldVisibility::All)
-			{}
-			SLATE_ARGUMENT(EFieldVisibility, FieldVisibilityFlags)
+		SLATE_BEGIN_ARGS(SSourceBindingList) {}
+			SLATE_ARGUMENT_DEFAULT(bool, ShowSearchBox) = false;
+			SLATE_ARGUMENT_DEFAULT(EFieldVisibility, FieldVisibilityFlags) = EFieldVisibility::All;
 			SLATE_EVENT(FOnDoubleClicked, OnDoubleClicked)
 		SLATE_END_ARGS()
 
@@ -67,6 +66,7 @@ namespace UE::MVVM
 		void AddSource(UClass* Class, FName Name, FGuid Guid);
 		void AddSources(TArrayView<const FBindingSource> InSources);
 
+		void AddWidgetBlueprint(const UWidgetBlueprint* WidgetBlueprint);
 		void AddWidgets(TArrayView<const UWidget*> Widgets);
 		void AddViewModels(TArrayView<const FMVVMBlueprintViewModelContext> ViewModels);
 
