MVVM - Made the FViewModelFilter class viewer filter from MVVMSelectViewModel.h public so it can be reused elsewhere.

Removed unnecessary cast to SClassViewer when it can be exposed as just SWidget.

#rnx

diff --git a/Engine/Plugins/Runtime/ModelViewViewModel/Source/ModelViewViewModelEditor/Private/Widgets/SMVVMSelectViewModel.cpp b/Engine/Plugins/Runtime/ModelViewViewModel/Source/ModelViewViewModelEditor/Private/Widgets/SMVVMSelectViewModel.cpp
--- a/Engine/Plugins/Runtime/ModelViewViewModel/Source/ModelViewViewModelEditor/Private/Widgets/SMVVMSelectViewModel.cpp
+++ b/Engine/Plugins/Runtime/ModelViewViewModel/Source/ModelViewViewModelEditor/Private/Widgets/SMVVMSelectViewModel.cpp
@@ -19,7 +19,6 @@
 #include "Widgets/SMVVMViewModelBindingListWidget.h"
 
 #include "ClassViewerModule.h"
-#include "ClassViewerFilter.h"
 #include "SClassViewer.h"
 
 #define LOCTEXT_NAMESPACE "SMVVMSelectViewModel"
@@ -29,28 +28,7 @@ namespace UE::MVVM
 
 namespace Private
 {
-	bool IsValidViewModel(const UClass * InClass);
-
-	class FSelectViewModelClassFilter : public IClassViewerFilter
-	{
-	public:
-		FSelectViewModelClassFilter() = default;
-
-		static const EClassFlags DisallowedClassFlags = CLASS_HideDropDown | CLASS_Hidden | CLASS_Deprecated | CLASS_Abstract | CLASS_NotPlaceable;
-		virtual bool IsClassAllowed(const FClassViewerInitializationOptions& InInitOptions, const UClass* InClass, TSharedRef<FClassViewerFilterFuncs> InFilterFuncs) override
-		{
-			return IsValidViewModel(InClass);
-		}
-
-		virtual bool IsUnloadedClassAllowed(const FClassViewerInitializationOptions& InInitOptions, const TSharedRef<const IUnloadedBlueprintData> InUnloadedClassData, TSharedRef<FClassViewerFilterFuncs> InFilterFuncs) override
-		{
-			if (InUnloadedClassData->IsChildOf(UWidget::StaticClass()))
-			{
-				return false;
-			}
-			return InUnloadedClassData->ImplementsInterface(UNotifyFieldValueChanged::StaticClass()) && !InUnloadedClassData->HasAnyClassFlags(DisallowedClassFlags);
-		}
-	};
+	static const EClassFlags DisallowedClassFlags = CLASS_HideDropDown | CLASS_Hidden | CLASS_Deprecated | CLASS_Abstract | CLASS_NotPlaceable;
 
 	bool IsValidViewModel(const UClass* InClass)
 	{
@@ -58,9 +36,24 @@ namespace Private
 		{
 			return false;
 		}
-		return InClass->ImplementsInterface(UNotifyFieldValueChanged::StaticClass()) && !InClass->HasAnyClassFlags(FSelectViewModelClassFilter::DisallowedClassFlags);
+		return InClass->ImplementsInterface(UNotifyFieldValueChanged::StaticClass()) && !InClass->HasAnyClassFlags(DisallowedClassFlags);
 	}
-} //namespace
+}
+
+bool FViewModelClassFilter::IsClassAllowed(const FClassViewerInitializationOptions& InInitOptions, const UClass* InClass, TSharedRef<FClassViewerFilterFuncs> InFilterFuncs)
+{
+	return Private::IsValidViewModel(InClass);
+}
+
+bool FViewModelClassFilter::IsUnloadedClassAllowed(const FClassViewerInitializationOptions& InInitOptions, const TSharedRef<const IUnloadedBlueprintData> InUnloadedClassData, TSharedRef<FClassViewerFilterFuncs> InFilterFuncs)
+{
+	if (InUnloadedClassData->IsChildOf(UWidget::StaticClass()))
+	{
+		return false;
+	}
+	return InUnloadedClassData->ImplementsInterface(UNotifyFieldValueChanged::StaticClass()) && !InUnloadedClassData->HasAnyClassFlags(Private::DisallowedClassFlags);
+}
+
 
 void SMVVMSelectViewModel::Construct(const FArguments& InArgs, const UWidgetBlueprint* WidgetBlueprint)
 {
@@ -72,9 +65,9 @@ void SMVVMSelectViewModel::Construct(const FArguments& InArgs, const UWidgetBlue
 		+ SUniformGridPanel::Slot(0, 0)
 		[
 			SNew(SPrimaryButton)
-			.Text(LOCTEXT("ViewModelAddButtonText", "Ok"))
+			.Text(LOCTEXT("ViewModelAddButtonText", "OK"))
 			.OnClicked(this, &SMVVMSelectViewModel::HandleAccepted)
-			.IsEnabled(this, &SMVVMSelectViewModel::HandelIsSelectionEnabled)
+			.IsEnabled(this, &SMVVMSelectViewModel::HandleIsSelectionEnabled)
 		]
 		+ SUniformGridPanel::Slot(1, 0)
 		[
@@ -87,9 +80,10 @@ void SMVVMSelectViewModel::Construct(const FArguments& InArgs, const UWidgetBlue
 	FClassViewerInitializationOptions ClassViewerOptions;
 	ClassViewerOptions.DisplayMode = EClassViewerDisplayMode::TreeView;
 	ClassViewerOptions.Mode = EClassViewerMode::ClassPicker;
-	ClassViewerOptions.ClassFilters.Add(MakeShared<Private::FSelectViewModelClassFilter>());
+	ClassViewerOptions.ClassFilters.Add(MakeShared<FViewModelClassFilter>());
 
-	ClassViewer = StaticCastSharedRef<SClassViewer>(FModuleManager::LoadModuleChecked<FClassViewerModule>("ClassViewer").CreateClassViewer(ClassViewerOptions, FOnClassPicked::CreateSP(this, &SMVVMSelectViewModel::HandleClassPicked)));
+	ClassViewer = FModuleManager::LoadModuleChecked<FClassViewerModule>("ClassViewer")
+		.CreateClassViewer(ClassViewerOptions, FOnClassPicked::CreateSP(this, &SMVVMSelectViewModel::HandleClassPicked));
 
 	ChildSlot
 	[
@@ -135,8 +129,14 @@ void SMVVMSelectViewModel::Construct(const FArguments& InArgs, const UWidgetBlue
 
 void SMVVMSelectViewModel::HandleClassPicked(UClass* ClassPicked)
 {
-	SelectedClass = Private::IsValidViewModel(ClassPicked) ? ClassPicked : nullptr;
-	BindingListWidget->AddSource(ClassPicked, FName(), FGuid());
+	BindingListWidget->Clear();
+	SelectedClass.Reset();
+
+	if (Private::IsValidViewModel(ClassPicked))
+	{
+		SelectedClass = ClassPicked;
+		BindingListWidget->AddSource(ClassPicked, FName(), FGuid());
+	}
 }
 
 
@@ -154,7 +154,7 @@ FReply SMVVMSelectViewModel::HandleCancel()
 }
 
 
-bool SMVVMSelectViewModel::HandelIsSelectionEnabled() const
+bool SMVVMSelectViewModel::HandleIsSelectionEnabled() const
 {
 	return SelectedClass.Get() != nullptr;
 }
diff --git a/Engine/Plugins/Runtime/ModelViewViewModel/Source/ModelViewViewModelEditor/Private/Widgets/SMVVMSelectViewModel.h b/Engine/Plugins/Runtime/ModelViewViewModel/Source/ModelViewViewModelEditor/Private/Widgets/SMVVMSelectViewModel.h
--- a/Engine/Plugins/Runtime/ModelViewViewModel/Source/ModelViewViewModelEditor/Private/Widgets/SMVVMSelectViewModel.h
+++ b/Engine/Plugins/Runtime/ModelViewViewModel/Source/ModelViewViewModelEditor/Private/Widgets/SMVVMSelectViewModel.h
@@ -3,6 +3,7 @@
 #pragma once
 
 #include "CoreMinimal.h"
+#include "ClassViewerFilter.h"
 #include "Widgets/SCompoundWidget.h"
 
 class SClassViewer;
@@ -12,6 +13,15 @@ namespace UE::MVVM
 {
 class SSourceBindingList;
 
+class FViewModelClassFilter : public IClassViewerFilter
+{
+public:
+	FViewModelClassFilter() = default;
+
+	virtual bool IsClassAllowed(const FClassViewerInitializationOptions& InInitOptions, const UClass* InClass, TSharedRef<FClassViewerFilterFuncs> InFilterFuncs) override;
+	virtual bool IsUnloadedClassAllowed(const FClassViewerInitializationOptions& InInitOptions, const TSharedRef<const IUnloadedBlueprintData> InUnloadedClassData, TSharedRef<FClassViewerFilterFuncs> InFilterFuncs) override;
+};
+
 class SMVVMSelectViewModel : public SCompoundWidget
 {
 public:
@@ -29,10 +39,10 @@ private:
 	void HandleClassPicked(UClass* ClassPicked);
 	FReply HandleAccepted();
 	FReply HandleCancel();
-	bool HandelIsSelectionEnabled() const;
+	bool HandleIsSelectionEnabled() const;
 
 private:
-	TSharedPtr<SClassViewer> ClassViewer;
+	TSharedPtr<SWidget> ClassViewer;
 	TSharedPtr<SSourceBindingList> BindingListWidget;
 	TWeakObjectPtr<const UClass> SelectedClass;
 
