MVVM - Added widget clipping to SMVVMFieldEntry and SMVVMSourceEntry.

#rnx

diff --git a/Engine/Plugins/Runtime/ModelViewViewModel/Source/ModelViewViewModelEditor/Private/Widgets/SMVVMFieldEntry.cpp b/Engine/Plugins/Runtime/ModelViewViewModel/Source/ModelViewViewModelEditor/Private/Widgets/SMVVMFieldEntry.cpp
--- a/Engine/Plugins/Runtime/ModelViewViewModel/Source/ModelViewViewModelEditor/Private/Widgets/SMVVMFieldEntry.cpp
+++ b/Engine/Plugins/Runtime/ModelViewViewModel/Source/ModelViewViewModelEditor/Private/Widgets/SMVVMFieldEntry.cpp
@@ -49,6 +49,7 @@ void SFieldEntry::Construct(const FArguments& InArgs)
 		[
 			SAssignNew(Label, STextBlock)
 			.TextStyle(InArgs._TextStyle)
+			.Clipping(EWidgetClipping::OnDemand)
 		]
 	];
 
diff --git a/Engine/Plugins/Runtime/ModelViewViewModel/Source/ModelViewViewModelEditor/Private/Widgets/SMVVMSourceEntry.cpp b/Engine/Plugins/Runtime/ModelViewViewModel/Source/ModelViewViewModelEditor/Private/Widgets/SMVVMSourceEntry.cpp
--- a/Engine/Plugins/Runtime/ModelViewViewModel/Source/ModelViewViewModelEditor/Private/Widgets/SMVVMSourceEntry.cpp
+++ b/Engine/Plugins/Runtime/ModelViewViewModel/Source/ModelViewViewModelEditor/Private/Widgets/SMVVMSourceEntry.cpp
@@ -71,6 +71,7 @@ void SSourceEntry::Construct(const FArguments& InArgs)
 		[
 			SAssignNew(Label, STextBlock)
 			.TextStyle(InArgs._TextStyle)
+			.Clipping(EWidgetClipping::OnDemand)
 		]
 	];
 
