MVVM - Fixed crash when binding a conversion function to a struct property.

#rnx

diff --git a/Engine/Plugins/Runtime/ModelViewViewModel/Source/ModelViewViewModelEditor/Private/MVVMEditorSubsystem.cpp b/Engine/Plugins/Runtime/ModelViewViewModel/Source/ModelViewViewModelEditor/Private/MVVMEditorSubsystem.cpp
--- a/Engine/Plugins/Runtime/ModelViewViewModel/Source/ModelViewViewModelEditor/Private/MVVMEditorSubsystem.cpp
+++ b/Engine/Plugins/Runtime/ModelViewViewModel/Source/ModelViewViewModelEditor/Private/MVVMEditorSubsystem.cpp
@@ -844,22 +844,29 @@ void UMVVMEditorSubsystem::SetPathForConversionFunctionArgument(UWidgetBlueprint
 	// create all the subsequent nodes in the path
 	for (int32 Index = 0; Index < Fields.Num(); ++Index)
 	{
-		if (!ensureMsgf(PreviousClass != nullptr, TEXT("Path is too long!")))
+		if (!ensureMsgf(PreviousClass != nullptr, TEXT("Previous class not set!")))
 		{
 			return;
 		}
 
 		UEdGraphNode* NewNode = nullptr;
-
 		const UE::MVVM::FMVVMConstFieldVariant& Field = Fields[Index];
+		const bool bLastField = (Index == Fields.Num() - 1);
+
 		if (Field.IsProperty())
 		{
 			const FProperty* Property = Field.GetProperty();
 
-			if (const FObjectProperty* ObjectProperty = CastField<FObjectProperty>(Property))
+			// for structs in the middle of a path, we need to use a break node
+			const FStructProperty* StructProperty = CastField<FStructProperty>(Property);
+			if (!bLastField && StructProperty != nullptr)
+			{
+				// TODO: Add struct break node
+				ensureAlwaysMsgf(false, TEXT("Struct nodes in the middle of the path are not yet supported. Path: %s"), *Path.GetBasePropertyPath());
+				PreviousClass = nullptr;
+			}
+			else if (PreviousClass != nullptr)
 			{
-				PreviousClass = ObjectProperty->PropertyClass;
-
 				FGraphNodeCreator<UK2Node_VariableGet> GetterCreator(*ConversionFunctionGraph);
 				UK2Node_VariableGet* GetterNode = GetterCreator.CreateNode();
 				GetterNode->NodePosX = PosX;
@@ -869,16 +876,15 @@ void UMVVMEditorSubsystem::SetPathForConversionFunctionArgument(UWidgetBlueprint
 				GetterCreator.Finalize();
 
 				NewNode = GetterNode;
-			}
-			else if (const FStructProperty* StructProperty = CastField<FStructProperty>(Property))
-			{
-				// TODO (sebastiann): StructProperty should be supported as well, but that's not a VariableGet
-				ensureMsgf(false, TEXT("Invalid path, structs are not supported yet."));
-				return;
-			}
-			else
-			{
-				PreviousClass = nullptr;
+
+				if (const FObjectProperty* ObjectProperty = CastField<FObjectProperty>(Property))
+				{
+					PreviousClass = ObjectProperty->PropertyClass;
+				}
+				else
+				{
+					PreviousClass = nullptr;
+				}
 			}
 		}
 		else if (Field.IsFunction())
@@ -898,12 +904,6 @@ void UMVVMEditorSubsystem::SetPathForConversionFunctionArgument(UWidgetBlueprint
 			{
 				PreviousClass = ObjectProperty->PropertyClass;
 			}
-			else if (const FStructProperty* StructProperty = CastField<FStructProperty>(ReturnProperty))
-			{
-				// TODO (sebastiann): StructProperty should be supported as well, but that's not a VariableGet
-				ensureMsgf(false, TEXT("Invalid path, structs are not supported yet."));
-				return;
-			}
 			else
 			{
 				PreviousClass = nullptr;
