MVVM - Field selector is now a split pane viewer with viewmodels/widgets on the left and properties on the right.

#rnx

diff --git a/Engine/Plugins/Runtime/ModelViewViewModel/Source/ModelViewViewModelEditor/Private/Styling/MVVMEditorStyle.cpp b/Engine/Plugins/Runtime/ModelViewViewModel/Source/ModelViewViewModelEditor/Private/Styling/MVVMEditorStyle.cpp
--- a/Engine/Plugins/Runtime/ModelViewViewModel/Source/ModelViewViewModelEditor/Private/Styling/MVVMEditorStyle.cpp
+++ b/Engine/Plugins/Runtime/ModelViewViewModel/Source/ModelViewViewModelEditor/Private/Styling/MVVMEditorStyle.cpp
@@ -69,8 +69,6 @@ FMVVMEditorStyle::FMVVMEditorStyle()
 		FComboButtonStyle(FAppStyle::Get().GetWidgetStyle<FComboButtonStyle>("ComboButton"))
 		.SetDownArrowImage(*FCoreStyle::Get().GetBrush("Icons.Edit")));
 
-	Set("FieldSelector.MenuBorderBrush", new FSlateRoundedBoxBrush(FStyleColors::Dropdown, 0.0f, FStyleColors::DropdownOutline, 1.0f));
-
 	FSlateStyleRegistry::RegisterSlateStyle(*this);
 }
 
diff --git a/Engine/Plugins/Runtime/ModelViewViewModel/Source/ModelViewViewModelEditor/Private/Widgets/SMVVMFieldSelector.cpp b/Engine/Plugins/Runtime/ModelViewViewModel/Source/ModelViewViewModelEditor/Private/Widgets/SMVVMFieldSelector.cpp
--- a/Engine/Plugins/Runtime/ModelViewViewModel/Source/ModelViewViewModelEditor/Private/Widgets/SMVVMFieldSelector.cpp
+++ b/Engine/Plugins/Runtime/ModelViewViewModel/Source/ModelViewViewModelEditor/Private/Widgets/SMVVMFieldSelector.cpp
@@ -2,13 +2,19 @@
 
 #include "Widgets/SMVVMFieldSelector.h"
  
+#include "ClassViewerModule.h"
 #include "Editor.h"
+#include "Hierarchy/SReadOnlyHierarchyView.h"
+#include "Modules/ModuleManager.h"
 #include "MVVMEditorSubsystem.h"
+#include "SMVVMSelectViewModel.h"
 #include "SNegativeActionButton.h"
 #include "SPositiveActionButton.h"
 #include "SPrimaryButton.h"
 #include "SSimpleButton.h"
 #include "Styling/MVVMEditorStyle.h"
+#include "Styling/SlateIconFinder.h"
+#include "WidgetBlueprint.h"
 #include "Widgets/SMVVMFieldEntry.h"
 #include "Widgets/SMVVMSourceEntry.h"
 #include "Widgets/Images/SImage.h"
@@ -48,6 +54,7 @@ void SFieldSelector::Construct(const FArguments& InArgs, const UWidgetBlueprint*
 
 	SelectedField = InArgs._SelectedField;
 	check(SelectedField.IsSet());
+	CachedSelectedField = SelectedField.Get();
 
 	BindingMode = InArgs._BindingMode;
 	check(BindingMode.IsSet());
@@ -75,12 +82,14 @@ void SFieldSelector::Construct(const FArguments& InArgs, const UWidgetBlueprint*
 					.Visibility_Lambda([this]() { return CachedSelectedField.IsEmpty() ? EVisibility::Collapsed : EVisibility::Visible; })
 					+ SHorizontalBox::Slot()
 					.Padding(8, 0, 0, 0)
-					.AutoWidth()
+					.HAlign(HAlign_Fill)
+					.FillWidth(1)
 					[
 						SAssignNew(SelectedSourceWidget, SSourceEntry)
 					]
 					+ SHorizontalBox::Slot()
 					.Padding(8, 0)
+					.VAlign(VAlign_Center)
 					.AutoWidth()
 					[
 						SNew(SImage)
@@ -88,7 +97,8 @@ void SFieldSelector::Construct(const FArguments& InArgs, const UWidgetBlueprint*
 					]
 					+ SHorizontalBox::Slot()
 					.Padding(0, 0, 8, 0)
-					.AutoWidth()
+					.HAlign(HAlign_Fill)
+					.FillWidth(1)
 					[
 						SAssignNew(SelectedEntryWidget, SFieldEntry)
 						.TextStyle(TextStyle)
@@ -113,13 +123,6 @@ void SFieldSelector::Construct(const FArguments& InArgs, const UWidgetBlueprint*
 	Refresh();
 }
 
-void SFieldSelector::OnFieldSelected(FMVVMBlueprintPropertyPath Selected)
-{
-	OnSelectionChangedDelegate.ExecuteIfBound(Selected);
-	
-	Refresh();
-}
-
 void SFieldSelector::Refresh()
 {
 	CachedSelectedField = SelectedField.Get();
@@ -184,11 +187,6 @@ FReply SFieldSelector::OnClearBinding()
 	return FReply::Handled();
 }
 
-void SFieldSelector::HandleSearchChanged(const FText& InFilterText)
-{
-	BindingList->SetRawFilterText(InFilterText);
-}
-
 FReply SFieldSelector::OnCancel()
 {
 	if (ComboButton.IsValid())
@@ -210,84 +208,185 @@ void SFieldSelector::SetSelection(const FMVVMBlueprintPropertyPath& SelectedPath
 	Refresh();
 }
 
-TSharedRef<SWidget> SFieldSelector::OnGetMenuContent()
+void SFieldSelector::OnSearchTextChanged(const FText& NewText)
 {
-	TSharedRef<SWidget> MenuWidget = 
-		SNew(SBorder)
-		.BorderImage(FMVVMEditorStyle::Get().GetBrush("FieldSelector.MenuBorderBrush"))
-		.Padding(FMargin(8, 1, 8, 3))
+	BindingList->SetRawFilterText(NewText);
+}
+
+void SFieldSelector::OnViewModelSelected(FBindingSource Source, ESelectInfo::Type)
+{
+	TArray<FBindingSource> Selection = ViewModelList->GetSelectedItems();
+	if (BindingList.IsValid())
+	{
+		BindingList->Clear();
+		BindingList->AddSources(Selection);
+	}
+}
+
+TSharedRef<ITableRow> SFieldSelector::GenerateRowForViewModel(FBindingSource ViewModel, const TSharedRef<STableViewBase>& OwnerTable) const
+{
+	return SNew(STableRow<FBindingSource>, OwnerTable)
 		[
-			SNew(SVerticalBox)
-			+ SVerticalBox::Slot()
-			.Padding(4)
-			.AutoHeight()
+			SNew(SHorizontalBox)
+			+ SHorizontalBox::Slot()
+			.Padding(0, 0, 2, 0)
+			.HAlign(HAlign_Left)
+			.VAlign(VAlign_Center)
+			.AutoWidth()
 			[
-				SAssignNew(SearchBox, SSearchBox)
-				.HintText(LOCTEXT("SearchViewmodel", "Search"))
-				.OnTextChanged(this, &SFieldSelector::HandleSearchChanged)
+				SNew(SImage)
+				.Image(FSlateIconFinder::FindIconBrushForClass(ViewModel.Class))
+				.ColorAndOpacity(FSlateColor::UseForeground())
 			]
-			+ SVerticalBox::Slot()
-			.FillHeight(1.0f)
+			+ SHorizontalBox::Slot()
+			.HAlign(HAlign_Left)
+			.VAlign(VAlign_Center)
+			.AutoWidth()
 			[
-				SNew(SBorder)
-				.BorderImage(FAppStyle::Get().GetBrush("Brushes.Dropdown"))
-				.Padding(0)
-				[
-					SAssignNew(BindingList, SSourceBindingList, WidgetBlueprint)
-					.OnDoubleClicked(this, &SFieldSelector::SetSelection)
-					.FieldVisibilityFlags(GetFieldVisibilityFlags())
-				]
+				SNew(STextBlock)
+				.Text(ViewModel.DisplayName)
+				.ColorAndOpacity(FSlateColor::UseForeground())
 			]
-			+ SVerticalBox::Slot()
-			.Padding(4, 4, 4, 0)
-			.AutoHeight()
+		];
+}
+
+TSharedRef<SWidget> SFieldSelector::OnGetMenuContent()
+{
+	ViewModelSources.Reset();
+	ViewModelList.Reset();
+
+	BindingList = SNew(SSourceBindingList, WidgetBlueprint)
+		.ShowSearchBox(false)
+		.OnDoubleClicked(this, &SFieldSelector::SetSelection)
+		.FieldVisibilityFlags(GetFieldVisibilityFlags());
+
+	TSharedPtr<SWidget> SourcePicker;
+	if (bViewModelProperty)
+	{
+		FBindingSource SelectedSource;
+
+		UMVVMEditorSubsystem* EditorSubsystem = GEditor->GetEditorSubsystem<UMVVMEditorSubsystem>();
+		if (const UMVVMBlueprintView* View = EditorSubsystem->GetView(WidgetBlueprint))
+		{
+			const TArrayView<const FMVVMBlueprintViewModelContext> ViewModels = View->GetViewModels();
+			ViewModelSources.Reserve(ViewModels.Num());
+
+			for (const FMVVMBlueprintViewModelContext& ViewModel : ViewModels)
+			{
+				FBindingSource& Source = ViewModelSources.AddDefaulted_GetRef();
+				Source.Class = ViewModel.GetViewModelClass();
+				Source.Name = ViewModel.GetViewModelName();
+				Source.DisplayName = ViewModel.GetDisplayName();
+				Source.ViewModelId = ViewModel.GetViewModelId();
+
+				if (Source.ViewModelId == CachedSelectedField.GetViewModelId())
+				{
+					SelectedSource = Source;
+				}
+			}
+		}
+
+		ViewModelList = SNew(SListView<FBindingSource>)
+			.OnGenerateRow(this, &SFieldSelector::GenerateRowForViewModel)
+			.SelectionMode(ESelectionMode::Multi)
+			.ListItemsSource(&ViewModelSources)
+			.OnSelectionChanged(this, &SFieldSelector::OnViewModelSelected);
+
+		ViewModelList->SetItemSelection(SelectedSource, true);
+
+		SourcePicker = ViewModelList;
+	}
+	else
+	{
+		auto OnWidgetSelected = [this](FName WidgetName, ESelectInfo::Type)
+		{
+			BindingList->Clear();
+
+			FBindingSource Source = FBindingSource::CreateForWidget(WidgetBlueprint, WidgetName);
+			if (Source.IsValid())
+			{
+				BindingList->AddSources(MakeArrayView(&Source, 1));
+			}
+		};
+
+		SourcePicker = SNew(SReadOnlyHierarchyView, WidgetBlueprint)
+			.OnSelectionChanged_Lambda(OnWidgetSelected)
+			.ShowSearch(false);
+	}
+
+	TSharedRef<SWidget> MenuWidget = 
+		SNew(SBox)
+		.MinDesiredWidth(400)
+		.MinDesiredHeight(200)
+		[
+			SNew(SBorder)
+			.BorderImage(FAppStyle::Get().GetBrush("ToolPanel.GroupBorder"))
+			.Padding(FMargin(8, 2, 8, 3))
 			[
-				SNew(SHorizontalBox)
-				+ SHorizontalBox::Slot()
+				SNew(SVerticalBox)
+				+ SVerticalBox::Slot()
+				.HAlign(HAlign_Fill)
+				.VAlign(VAlign_Top)
+				.AutoHeight()
 				[
-					SNew(SButton)
-					.OnClicked(this, &SFieldSelector::OnCancel)
-					.HAlign(HAlign_Center)
-					[
-						SNew(STextBlock)
-						.Text(LOCTEXT("Cancel", "Cancel"))
-					]
+					SNew(SSearchBox)
+					.OnTextChanged(this, &SFieldSelector::OnSearchTextChanged)
 				]
-				+ SHorizontalBox::Slot()
-				.Padding(4, 0, 0, 0)
+				+ SVerticalBox::Slot()
+				.FillHeight(1.0f)
 				[
-					SNew(SButton)
-					.OnClicked(this, &SFieldSelector::OnClearBinding)
-					.HAlign(HAlign_Center)
+					SNew(SSplitter)
+					.PhysicalSplitterHandleSize(4.0f)
+					+ SSplitter::Slot()
+					.Value(0.5f)
 					[
-						SNew(STextBlock)
-						.Text(LOCTEXT("Clear", "Clear"))
+						SourcePicker.ToSharedRef()
+					]
+					+ SSplitter::Slot()
+					.Value(0.5f)
+					[
+						BindingList.ToSharedRef()
 					]
 				]
-				+ SHorizontalBox::Slot()
-				.Padding(4, 0, 0, 0)
+				+ SVerticalBox::Slot()
+				.Padding(4, 4, 4, 0)
+				.AutoHeight()
 				[
-					SNew(SPrimaryButton)
-					.OnClicked(this, &SFieldSelector::OnSelectProperty)
-					.Text(LOCTEXT("Select", "Select"))
+					SNew(SHorizontalBox)
+					+ SHorizontalBox::Slot()
+					[
+						SNew(SButton)
+						.OnClicked(this, &SFieldSelector::OnCancel)
+						.HAlign(HAlign_Center)
+						[
+							SNew(STextBlock)
+							.Text(LOCTEXT("Cancel", "Cancel"))
+						]
+					]
+					+ SHorizontalBox::Slot()
+					.Padding(4, 0, 0, 0)
+					[
+						SNew(SButton)
+						.OnClicked(this, &SFieldSelector::OnClearBinding)
+						.IsEnabled(this, &SFieldSelector::IsClearEnabled)
+						.HAlign(HAlign_Center)
+						[
+							SNew(STextBlock)
+							.Text(LOCTEXT("Clear", "Clear"))
+						]
+					]
+					+ SHorizontalBox::Slot()
+					.Padding(4, 0, 0, 0)
+					[
+						SNew(SPrimaryButton)
+						.OnClicked(this, &SFieldSelector::OnSelectProperty)
+						.IsEnabled(this, &SFieldSelector::IsSelectEnabled)
+						.Text(LOCTEXT("Select", "Select"))
+					]
 				]
 			]
 		];
 
-	UMVVMEditorSubsystem* EditorSubsystem = GEditor->GetEditorSubsystem<UMVVMEditorSubsystem>();
-	if (bViewModelProperty)
-	{
-		if (UMVVMBlueprintView* View = EditorSubsystem->GetView(WidgetBlueprint))
-		{
-			BindingList->AddViewModels(View->GetViewModels());
-		}
-	}
-	else
-	{
-		TArray<FBindingSource> BindableWidgets = EditorSubsystem->GetBindableWidgets(WidgetBlueprint);
-		BindingList->AddSources(BindableWidgets);
-	}
-
 	return MenuWidget;
 }
 
diff --git a/Engine/Plugins/Runtime/ModelViewViewModel/Source/ModelViewViewModelEditor/Private/Widgets/SMVVMFieldSelector.h b/Engine/Plugins/Runtime/ModelViewViewModel/Source/ModelViewViewModelEditor/Private/Widgets/SMVVMFieldSelector.h
--- a/Engine/Plugins/Runtime/ModelViewViewModel/Source/ModelViewViewModelEditor/Private/Widgets/SMVVMFieldSelector.h
+++ b/Engine/Plugins/Runtime/ModelViewViewModel/Source/ModelViewViewModelEditor/Private/Widgets/SMVVMFieldSelector.h
@@ -18,6 +18,7 @@
 
 class SComboButton;
 class SSearchBox;
+class SReadOnlyHierarchyView;
 
 namespace UE::MVVM
 {
@@ -28,7 +29,7 @@ class SSourceBindingList;
 class SFieldSelector : public SCompoundWidget
 {
 public:
-	DECLARE_DELEGATE_OneParam(FSelectionChanged, FMVVMBlueprintPropertyPath);
+	DECLARE_DELEGATE_OneParam(FOnSelectionChanged, FMVVMBlueprintPropertyPath);
 
 	SLATE_BEGIN_ARGS(SFieldSelector) :
 		_TextStyle(&FCoreStyle::Get().GetWidgetStyle<FTextBlockStyle>("NormalText"))
@@ -36,7 +37,7 @@ public:
 		}
 		SLATE_STYLE_ARGUMENT(FTextBlockStyle, TextStyle)
 		SLATE_ATTRIBUTE(FMVVMBlueprintPropertyPath, SelectedField)
-		SLATE_EVENT(FSelectionChanged, OnSelectionChanged)
+		SLATE_EVENT(FOnSelectionChanged, OnSelectionChanged)
 		SLATE_ATTRIBUTE(EMVVMBindingMode, BindingMode)
 
 		/**
@@ -51,14 +52,10 @@ public:
 
 private:
 	TSharedRef<SWidget> OnGenerateFieldWidget(FMVVMBlueprintPropertyPath Path) const;
-
 	TSharedRef<SWidget> OnGetMenuContent();
-	void OnFieldSelected(FMVVMBlueprintPropertyPath Selected);
 
 	TValueOrError<bool, FString> ValidateField(FMVVMBlueprintPropertyPath Field) const;
 
-	void HandleSearchChanged(const FText& InFilterText);
-
 	bool IsClearEnabled() const;
 	FReply OnClearBinding();
 
@@ -70,6 +67,11 @@ private:
 	void SetSelection(const FMVVMBlueprintPropertyPath& SelectedPath);
 	EFieldVisibility GetFieldVisibilityFlags() const;
 
+	void OnViewModelSelected(FBindingSource ViewModel, ESelectInfo::Type);
+	TSharedRef<ITableRow> GenerateRowForViewModel(MVVM::FBindingSource ViewModel, const TSharedRef<STableViewBase>& OwnerTable) const;
+
+	void OnSearchTextChanged(const FText& NewText);
+
 private:
 	TAttribute<FMVVMBlueprintPropertyPath> SelectedField;
 	TAttribute<EMVVMBindingMode> BindingMode;
@@ -79,7 +81,7 @@ private:
 
 	const FTextBlockStyle* TextStyle = nullptr;
 
-	FSelectionChanged OnSelectionChangedDelegate;
+	FOnSelectionChanged OnSelectionChangedDelegate;
 	FIsFieldValid OnValidateFieldDelegate;
 
 	FMVVMBlueprintPropertyPath CachedSelectedField;
@@ -90,6 +92,9 @@ private:
 
 	const UWidgetBlueprint* WidgetBlueprint = nullptr;
 	bool bViewModelProperty = false;
+
+	TArray<FBindingSource> ViewModelSources;
+	TSharedPtr<SListView<FBindingSource>> ViewModelList;
 }; 
 
 } // namespace UE::MVVM
\ No newline at end of file
