MVVM - Bindings list changed to tree view as per latest design.

SMVVMViewBindingsListView is now UE::MVVM::SBindingsList.

FBindingEntry now can represent either a widget or a binding row, which generate an SWidgetRow or SBindingRow depending on which is set.

#rnx

diff --git a/Engine/Plugins/Runtime/ModelViewViewModel/Source/ModelViewViewModelEditor/Private/Tabs/MVVMBindingSummoner.cpp b/Engine/Plugins/Runtime/ModelViewViewModel/Source/ModelViewViewModelEditor/Private/Tabs/MVVMBindingSummoner.cpp
--- a/Engine/Plugins/Runtime/ModelViewViewModel/Source/ModelViewViewModelEditor/Private/Tabs/MVVMBindingSummoner.cpp
+++ b/Engine/Plugins/Runtime/ModelViewViewModel/Source/ModelViewViewModelEditor/Private/Tabs/MVVMBindingSummoner.cpp
@@ -28,7 +28,7 @@ FMVVMBindingSummoner::FMVVMBindingSummoner(TSharedPtr<FWidgetBlueprintEditor> Bl
 
 TSharedRef<SWidget> FMVVMBindingSummoner::CreateTabBody(const FWorkflowTabSpawnInfo& Info) const
 {
-	return SNew(SMVVMViewBindingPanel, WeakWidgetBlueprintEditor.Pin(), bIsDrawerTab);
+	return SNew(UE::MVVM::SBindingsPanel, WeakWidgetBlueprintEditor.Pin(), bIsDrawerTab);
 }
 
 void FMVVMBindingSummoner::ToggleMVVMDrawer()
diff --git a/Engine/Plugins/Runtime/ModelViewViewModel/Source/ModelViewViewModelEditor/Private/Widgets/SMVVMViewBindingListView.cpp b/Engine/Plugins/Runtime/ModelViewViewModel/Source/ModelViewViewModelEditor/Private/Widgets/SMVVMViewBindingListView.cpp
--- a/Engine/Plugins/Runtime/ModelViewViewModel/Source/ModelViewViewModelEditor/Private/Widgets/SMVVMViewBindingListView.cpp
+++ b/Engine/Plugins/Runtime/ModelViewViewModel/Source/ModelViewViewModelEditor/Private/Widgets/SMVVMViewBindingListView.cpp
@@ -18,6 +18,7 @@
 #include "MVVMSubsystem.h"
 #include "MVVMWidgetBlueprintExtension_View.h"
 #include "SEnumCombo.h"
+#include "SPrimaryButton.h"
 #include "SSimpleButton.h" 
 #include "Styling/MVVMEditorStyle.h"
 #include "Styling/StyleColors.h"
@@ -32,135 +33,263 @@
 #include "Widgets/Input/SComboBox.h"
 #include "Widgets/Input/SEditableTextBox.h"
 #include "Widgets/Text/STextBlock.h"
-#include "Widgets/Views/SListView.h"
+#include "Widgets/Views/STreeView.h"
 #include "Widgets/Views/STableViewBase.h"
 
 #define LOCTEXT_NAMESPACE "BindingListView"
 
-struct FMVVMViewBindingListEntry
+namespace UE::MVVM
 {
-	FMVVMViewBindingListEntry(int32 InIndex) : Index(InIndex) {}
-	int32 Index;
-};
 
-class SMVVMViewBindingListEntryRow : public SMultiColumnTableRow<FMVVMViewBindingListEntryPtr>
+// a wrapper around either a widget row or a binding row
+struct FBindingEntry
 {
-public:
-	static FName EnabledColumnName;
-	static FName CompileColumnName;
-	static FName ErrorColumnName;
-	static FName ViewModelColumnName;
-	static FName ModeColumnName;
-	static FName WidgetColumnName;
-	static FName WidgetColumName;
-	static FName UpdateColumnName;
-	static FName ConversionFunctionColumnName;
-	static FName DropDownOptionsColumnName;
-
-public:
-	SLATE_BEGIN_ARGS(SMVVMViewBindingListEntryRow) {}
-	SLATE_END_ARGS()
+	FBindingEntry(FName InWidgetName) : WidgetName(InWidgetName) {}
+	FBindingEntry(int32 InIndex) : BindingIndex(InIndex) { check(InIndex != INDEX_NONE); }
+	
+	FMVVMBlueprintViewBinding* GetBinding(UMVVMBlueprintView* View) const
+	{
+		return View->GetBindingAt(BindingIndex);
+	}
 
-	void Construct(const FArguments& Args, const TSharedRef<STableViewBase>& OwnerTableView, const FMVVMViewBindingListEntryPtr& InEntry, UMVVMBlueprintView* InBlueprintView, UWidgetBlueprint* InWidgetBlueprint)
+	const FMVVMBlueprintViewBinding* GetBinding(const UMVVMBlueprintView* View) const
 	{
-		Entry = InEntry;
-		BlueprintView = InBlueprintView;
-		WidgetBlueprint = InWidgetBlueprint;
+		return View->GetBindingAt(BindingIndex);
+	}
 
-		OnBlueprintChangedHandle = WidgetBlueprint->OnChanged().AddSP(this, &SMVVMViewBindingListEntryRow::HandleBlueprintChanged);
+	int32 GetBindingIndex() const
+	{
+		return BindingIndex;
+	}
 
-		FMVVMBlueprintViewBinding* ViewModelBinding = GetThisViewBinding();
+	FName GetWidgetName() const
+	{
+		return WidgetName;
+	}
 
-		SMultiColumnTableRow<FMVVMViewBindingListEntryPtr>::Construct(
-			FSuperRowType::FArguments()
-			.Padding(1.0f),
-			OwnerTableView
-		);
+	TConstArrayView<TSharedPtr<FBindingEntry>> GetChildren() const
+	{
+		return Children;
 	}
 
-	~SMVVMViewBindingListEntryRow()
+	void AddChild(TSharedRef<FBindingEntry> Child)
 	{
-		WidgetBlueprint->OnChanged().Remove(OnBlueprintChangedHandle);
+		Children.Add(Child);
 	}
 
-	FMVVMBlueprintViewBinding* GetThisViewBinding() const
+private:
+	FName WidgetName;
+	int32 BindingIndex = INDEX_NONE;
+	TArray<TSharedPtr<FBindingEntry>> Children;
+};
+
+namespace Private
+{
+	TArray<FName>* GetBindingModeNames()
 	{
-		if (UMVVMBlueprintView* BlueprintViewPtr = BlueprintView.Get())
+		static TArray<FName> BindingModeNames;
+
+		if (BindingModeNames.IsEmpty())
 		{
-			FMVVMBlueprintViewBinding* ViewBinding = BlueprintViewPtr->GetBindingAt(Entry->Index);
-			return ViewBinding;
+			UEnum* ModeEnum = StaticEnum<EMVVMBindingMode>();
+
+			BindingModeNames.Reserve(ModeEnum->NumEnums());
+
+			for (int32 BindingIndex = 0; BindingIndex < ModeEnum->NumEnums() - 1; ++BindingIndex)
+			{
+				const bool bIsHidden = ModeEnum->HasMetaData(TEXT("Hidden"), BindingIndex);
+				if (!bIsHidden)
+				{
+					BindingModeNames.Add(ModeEnum->GetNameByIndex(BindingIndex));
+				}
+			}
 		}
 
-		return nullptr;
+		return &BindingModeNames;
 	}
 
-	TArray<FMVVMBlueprintViewBinding*> GetThisViewBindingAsArray() const
+	void ExpandAll(const TSharedPtr<STreeView<TSharedPtr<FBindingEntry>>>& TreeView, const TSharedPtr<FBindingEntry>& Entry)
 	{
-		TArray<FMVVMBlueprintViewBinding*> Result;
+		TreeView->SetItemExpansion(Entry, true);
 
-		FMVVMBlueprintViewBinding* ViewBinding = GetThisViewBinding();
-		if (ViewBinding != nullptr)
+		for (const TSharedPtr<FBindingEntry>& Child : Entry->GetChildren())
 		{
-			Result.Add(ViewBinding);
+			ExpandAll(TreeView, Child);
 		}
+	}
+}
 
-		return Result;
+class SWidgetRow : public STableRow<TSharedPtr<FBindingEntry>>
+{
+public:
+	SLATE_BEGIN_ARGS(SWidgetRow) {}
+	SLATE_END_ARGS()
+
+	void Construct(const FArguments& Args, const TSharedRef<STableViewBase>& OwnerTableView, const TSharedPtr<FBindingEntry>& InEntry, UWidgetBlueprint* InWidgetBlueprint)
+	{
+		Entry = InEntry;
+		WidgetBlueprintWeak = InWidgetBlueprint;
+
+		STableRow<TSharedPtr<FBindingEntry>>::Construct(
+			STableRow<TSharedPtr<FBindingEntry>>::FArguments()
+			.Padding(1.0f)
+			[
+				SNew(SBox)
+				.HeightOverride(30)
+				[
+					SNew(SHorizontalBox)
+					+ SHorizontalBox::Slot()
+					.Padding(2, 1)
+					.HAlign(HAlign_Left)
+					.VAlign(VAlign_Fill)
+					.AutoWidth()
+					[
+						SNew(SSourceSelector, InWidgetBlueprint)
+						.ShowClear(false)
+						.AutoRefresh(true)
+						.ViewModels(false)
+						.SelectedSource(this, &SWidgetRow::GetSelectedWidget)
+					]
+					+ SHorizontalBox::Slot()
+					.HAlign(HAlign_Left)
+					.VAlign(VAlign_Center)
+					.AutoWidth()
+					[
+						SNew(SSimpleButton)
+						.Icon(FAppStyle::Get().GetBrush("Icons.Plus"))
+						.IsEnabled_Lambda([this]() { return !Entry->GetWidgetName().IsNone(); })
+						.OnClicked(this, &SWidgetRow::AddBinding)
+
+					]
+				]
+			],
+			OwnerTableView
+		);
 	}
 
-	/** Overridden from SMultiColumnTableRow.  Generates a widget for this column of the list view. */
-	virtual TSharedRef<SWidget> GenerateWidgetForColumn(const FName& ColumnName) override
+private:
+
+	FBindingSource GetSelectedWidget() const
 	{
-		FMVVMBlueprintViewBinding* ViewModelBinding = GetThisViewBinding();
+		return FBindingSource::CreateForWidget(WidgetBlueprintWeak.Get(), Entry->GetWidgetName());
+	}
 
-		if (ColumnName == CompileColumnName)
-		{
-			return SNew(SCheckBox)
-				.IsChecked(this, &SMVVMViewBindingListEntryRow::IsBindingCompiled)
-				.OnCheckStateChanged(this, &SMVVMViewBindingListEntryRow::OnIsBindingCompileChanged);
-		}
-		else if (ColumnName == ErrorColumnName)
+	FReply AddBinding() const
+	{
+		if (UWidgetBlueprint* WidgetBlueprint = WidgetBlueprintWeak.Get())
 		{
-			return SNew(SSimpleButton)
-				.Icon(FAppStyle::Get().GetBrush("Icons.Error"))
-				.Visibility(this, &SMVVMViewBindingListEntryRow::GetErrorVisibility)
-				.ToolTipText(this, &SMVVMViewBindingListEntryRow::GetErrorToolTip)
-				.OnClicked(this, &SMVVMViewBindingListEntryRow::OnErrorButtonClicked);
+			UMVVMEditorSubsystem* EditorSubsystem = GEditor->GetEditorSubsystem<UMVVMEditorSubsystem>();
+			FMVVMBlueprintViewBinding& Binding = EditorSubsystem->AddBinding(WidgetBlueprint);
+			FMVVMBlueprintPropertyPath Path;
+			Path.SetWidgetName(Entry->GetWidgetName());
+			EditorSubsystem->SetWidgetPropertyForBinding(WidgetBlueprint, Binding, Path);
 		}
-		if (ColumnName == ViewModelColumnName)
-		{
 
-			return SNew(SBox)
-				.Padding(FMargin(2, 0))
+		return FReply::Handled();
+	}
+
+private:
+	TSharedPtr<FBindingEntry> Entry;
+	TWeakObjectPtr<UWidgetBlueprint> WidgetBlueprintWeak;
+};
+
+class SBindingRow : public STableRow<TSharedPtr<FBindingEntry>>
+{
+public:
+	SLATE_BEGIN_ARGS(SBindingRow) {}
+	SLATE_END_ARGS()
+
+	void Construct(const FArguments& Args, const TSharedRef<STableViewBase>& OwnerTableView, const TSharedPtr<FBindingEntry>& InEntry, UMVVMBlueprintView* InBlueprintView, UWidgetBlueprint* InWidgetBlueprint)
+	{
+		Entry = InEntry;
+		WidgetBlueprintWeak = InWidgetBlueprint;
+		BlueprintView = InBlueprintView;
+
+		OnBlueprintChangedHandle = InWidgetBlueprint->OnChanged().AddSP(this, &SBindingRow::HandleBlueprintChanged);
+
+		FMVVMBlueprintViewBinding* ViewModelBinding = GetThisViewBinding();
+
+		STableRow<TSharedPtr<FBindingEntry>>::Construct(
+			STableRow<TSharedPtr<FBindingEntry>>::FArguments()
+			.Padding(1.0f)
+			[
+				SNew(SHorizontalBox)
+
+				+ SHorizontalBox::Slot()
+				.Padding(2, 0)
 				.VAlign(VAlign_Center)
+				.HAlign(HAlign_Left) 
+				.AutoWidth()
 				[
-					SAssignNew(ViewModelFieldSelector, UE::MVVM::SFieldSelector, WidgetBlueprint, true)
-					.SelectedField(this, &SMVVMViewBindingListEntryRow::GetSelectedViewModelProperty)
-					.BindingMode(this, &SMVVMViewBindingListEntryRow::GetCurrentBindingMode)
-					.OnSelectionChanged(this, &SMVVMViewBindingListEntryRow::OnViewModelPropertySelectionChanged)
-				];
-		}
-		else if (ColumnName == ModeColumnName)
-		{
-			UEnum* ModeEnum = StaticEnum<EMVVMBindingMode>();
-			for (int32 Index = 0; Index < ModeEnum->NumEnums() - 1; ++Index)
-			{
-				const bool bIsHidden = ModeEnum->HasMetaData(TEXT("Hidden"), Index);
-				if (!bIsHidden)
-				{
-					ModeNames.Add(ModeEnum->GetNameByIndex(Index));
-				}
-			}
+					SNew(SCheckBox)
+					.IsChecked(this, &SBindingRow::IsBindingCompiled)
+					.OnCheckStateChanged(this, &SBindingRow::OnIsBindingCompileChanged)
+				]
 
-			return 
-				SNew(SBox)
-				.Padding(FMargin(2, 0))
+				+ SHorizontalBox::Slot()
+				.Padding(2, 0)
 				.VAlign(VAlign_Center)
+				.HAlign(HAlign_Left) 
+				.AutoWidth()
+				[
+					SNew(SSimpleButton)
+					.Icon(FAppStyle::Get().GetBrush("Icons.Error"))
+					.Visibility(this, &SBindingRow::GetErrorVisibility)
+					.ToolTipText(this, &SBindingRow::GetErrorToolTip)
+					.OnClicked(this, &SBindingRow::OnErrorButtonClicked)
+				]
+
+				+ SHorizontalBox::Slot()
+				.Padding(4, 0)
+				.VAlign(VAlign_Center)
+				.HAlign(HAlign_Fill)
+				.AutoWidth()
+				[
+					SNew(SBox)
+					.HeightOverride(28)
+					[
+						SAssignNew(WidgetMenuAnchor, SMenuAnchor)
+						.OnGetMenuContent(this, &SBindingRow::OnGetWidgetMenuContent)
+						[
+							SNew(SButton)
+							.ButtonStyle(&FAppStyle::Get().GetWidgetStyle<FComboButtonStyle>("ComboButton").ButtonStyle)
+							.OnClicked(this, &SBindingRow::OnWidgetMenuAnchorClicked)
+							[
+								SNew(SHorizontalBox)
+								+ SHorizontalBox::Slot()
+								.HAlign(HAlign_Left)
+								.VAlign(VAlign_Fill)
+								.AutoWidth()
+								[
+									SAssignNew(SelectedWidget, SFieldEntry)
+									.Field(GetSelectedWidgetProperty())
+								]
+								+ SHorizontalBox::Slot()
+								.Padding(4, 0, 0, 0)
+								.HAlign(HAlign_Right)
+								.VAlign(VAlign_Center)
+								.AutoWidth()
+								[
+									SNew(SImage)
+									.Image(FAppStyle::Get().GetBrush("Icons.ChevronDown"))
+									.ColorAndOpacity(FSlateColor::UseForeground())
+								]
+							]
+						]
+					]
+				]
+
+				+ SHorizontalBox::Slot()
+				.Padding(2, 1)
+				.VAlign(VAlign_Fill)
+				.HAlign(HAlign_Left) 
+				.AutoWidth()
 				[
 					SNew(SComboBox<FName>)
-					.OptionsSource(&ModeNames)
-					.InitiallySelectedItem(ModeEnum->GetNameByValue((int64) ViewModelBinding->BindingType))
-					.OnSelectionChanged(this, &SMVVMViewBindingListEntryRow::OnModeSelectionChanged)
-					.OnGenerateWidget(this, &SMVVMViewBindingListEntryRow::GenerateModeWidget)
+					.OptionsSource(Private::GetBindingModeNames())
+					.InitiallySelectedItem(StaticEnum<EMVVMBindingMode>()->GetNameByValue((int64) ViewModelBinding->BindingType))
+					.OnSelectionChanged(this, &SBindingRow::OnModeSelectionChanged)
+					.OnGenerateWidget(this, &SBindingRow::GenerateModeWidget)
 					.Content()
 					[
 						SNew(SBox)
@@ -170,76 +299,109 @@ public:
 						.HeightOverride(16)
 						[
 							SNew(SImage)
-							.Image(this, &SMVVMViewBindingListEntryRow::GetCurrentModeBrush)
+							.Image(this, &SBindingRow::GetCurrentModeBrush)
 						]
 					]
-				];
-		}
-		else if (ColumnName == WidgetColumName)
-		{
-			return SNew(SBox)
-				.Padding(FMargin(2, 0))
+				]
+				
+				+ SHorizontalBox::Slot()
+				.Padding(2, 0)
 				.VAlign(VAlign_Center)
+				.HAlign(HAlign_Fill) 
+				.AutoWidth()
 				[
-					SAssignNew(WidgetFieldSelector, UE::MVVM::SFieldSelector, WidgetBlueprint, false)
-					.SelectedField(this, &SMVVMViewBindingListEntryRow::GetSelectedWidgetProperty)
-					.BindingMode(this, &SMVVMViewBindingListEntryRow::GetCurrentBindingMode)
-					.OnSelectionChanged(this, &SMVVMViewBindingListEntryRow::OnWidgetPropertySelectionChanged)
-				];
-		}
-		else if (ColumnName == UpdateColumnName)
-		{
-			const UEnum* UpdateModeEnum = StaticEnum<EMVVMViewBindingUpdateMode>();
-			return SNew(SBox)
-				.Padding(FMargin(2, 0))
-				.VAlign(VAlign_Center)
+					SAssignNew(ViewModelFieldSelector, SFieldSelector, InWidgetBlueprint, true)
+					.SelectedField(this, &SBindingRow::GetSelectedViewModelProperty)
+					.BindingMode(this, &SBindingRow::GetCurrentBindingMode)
+					.OnSelectionChanged(this, &SBindingRow::OnViewModelPropertySelectionChanged)
+				]
+
+				+ SHorizontalBox::Slot()
+				.Padding(2, 1)
+				.VAlign(VAlign_Fill)
+				.HAlign(HAlign_Left)
+				.AutoWidth()
 				[
-					SNew(SEnumComboBox, UpdateModeEnum)
+					SNew(SEnumComboBox, StaticEnum<EMVVMViewBindingUpdateMode>())
 					.ContentPadding(FMargin(4, 0))
-					.OnEnumSelectionChanged(this, &SMVVMViewBindingListEntryRow::OnUpdateModeSelectionChanged)
-					.CurrentValue(this, &SMVVMViewBindingListEntryRow::GetUpdateModeValue)
-				];
-		}
-		else if (ColumnName == ConversionFunctionColumnName)
-		{
-			return SNew(SHorizontalBox)
+					.OnEnumSelectionChanged(this, &SBindingRow::OnUpdateModeSelectionChanged)
+					.CurrentValue(this, &SBindingRow::GetUpdateModeValue)
+				]
+		
 				+ SHorizontalBox::Slot()
+				.Padding(2, 0)
 				.VAlign(VAlign_Center)
 				.HAlign(HAlign_Left)
 				.AutoWidth()
 				[
-					SNew(SMVVMConversionPath, WidgetBlueprint, false)
-					.Bindings(this, &SMVVMViewBindingListEntryRow::GetThisViewBindingAsArray)
-					.OnFunctionChanged(this, &SMVVMViewBindingListEntryRow::OnConversionFunctionChanged, false)
+					SNew(SMVVMConversionPath, InWidgetBlueprint, false)
+					.Bindings(this, &SBindingRow::GetThisViewBindingAsArray)
+					.OnFunctionChanged(this, &SBindingRow::OnConversionFunctionChanged, false)
 				]
+
 				+ SHorizontalBox::Slot()
+				.Padding(2, 0)
 				.VAlign(VAlign_Center)
 				.HAlign(HAlign_Left)
 				.AutoWidth()
 				[
-					SNew(SMVVMConversionPath, WidgetBlueprint, true)
-					.Bindings(this, &SMVVMViewBindingListEntryRow::GetThisViewBindingAsArray)
-					.OnFunctionChanged(this, &SMVVMViewBindingListEntryRow::OnConversionFunctionChanged, true)
-				];
-		}
-		else if (ColumnName == DropDownOptionsColumnName)
-		{
-			return SAssignNew(ContextMenuOptionHelper, SButton)
-				.ToolTipText(LOCTEXT("DropDownOptionsToolTip", "Context Menu for Binding"))
-				.ButtonStyle(&FAppStyle::Get().GetWidgetStyle<FButtonStyle>("SimpleButton"))
-				.OnClicked(this, &SMVVMViewBindingListEntryRow::HandleDropDownOptionsPressed)
+					SNew(SMVVMConversionPath, InWidgetBlueprint, true)
+					.Bindings(this, &SBindingRow::GetThisViewBindingAsArray)
+					.OnFunctionChanged(this, &SBindingRow::OnConversionFunctionChanged, true)
+				]
+
+				+ SHorizontalBox::Slot() 
+				.Padding(2, 0)
+				.VAlign(VAlign_Center)
+				.HAlign(HAlign_Right)
+				.AutoWidth()
 				[
-					SNew(SBox)
-					.Padding(FMargin(3, 0))
+					SAssignNew(ContextMenuOptionHelper, SButton)
+					.ToolTipText(LOCTEXT("DropDownOptionsToolTip", "Context Menu for Binding"))
+					.ButtonStyle(&FAppStyle::Get().GetWidgetStyle<FButtonStyle>("SimpleButton"))
+					.OnClicked(this, &SBindingRow::HandleDropDownOptionsPressed)
 					[
-						SNew(SImage)
-						.Image(&FAppStyle::Get().GetWidgetStyle<FComboButtonStyle>("SegmentedCombo.Right").DownArrowImage)
-						.ColorAndOpacity(FSlateColor::UseForeground())
+						SNew(SBox)
+						.Padding(FMargin(3, 0))
+						[
+							SNew(SImage)
+							.Image(&FAppStyle::Get().GetWidgetStyle<FComboButtonStyle>("SegmentedCombo.Right").DownArrowImage)
+							.ColorAndOpacity(FSlateColor::UseForeground())
+						]
 					]
-				];
+				]
+			],
+			OwnerTableView
+		);
+	}
+
+	~SBindingRow()
+	{
+		if (UWidgetBlueprint* WidgetBlueprint = WidgetBlueprintWeak.Get())
+		{
+			WidgetBlueprint->OnChanged().Remove(OnBlueprintChangedHandle);
 		}
-		
-		return SNullWidget::NullWidget;
+	}
+
+	FMVVMBlueprintViewBinding* GetThisViewBinding() const
+	{
+		if (UMVVMBlueprintView* BlueprintViewPtr = BlueprintView.Get())
+		{
+			FMVVMBlueprintViewBinding* ViewBinding = Entry->GetBinding(BlueprintViewPtr);
+			return ViewBinding;
+		}
+
+		return nullptr;
+	}
+
+	TArray<FMVVMBlueprintViewBinding*> GetThisViewBindingAsArray() const
+	{
+		TArray<FMVVMBlueprintViewBinding*> Result;
+		if (FMVVMBlueprintViewBinding* ViewBinding = GetThisViewBinding())
+		{
+			Result.Add(ViewBinding);
+		}
+		return Result;
 	}
 
 private:
@@ -296,7 +458,7 @@ private:
 				[
 					SNew(SListView<TSharedPtr<FText>>)
 					.ListItemsSource(&ErrorItems)
-					.OnGenerateRow(this, &SMVVMViewBindingListEntryRow::OnGenerateErrorRow)
+					.OnGenerateRow(this, &SBindingRow::OnGenerateErrorRow)
 				];
 
 			ErrorDialog->Show();
@@ -311,23 +473,86 @@ private:
 		return ViewModelBinding->BindingType;
 	}
 
-	UMVVMSubsystem::FConstDirectionalBindingArgs GetCurrentBindingArgs(bool bIsGetter) const
+	FReply OnWidgetMenuAnchorClicked()
 	{
-		const FMVVMBlueprintViewBinding* ViewModelBinding = GetThisViewBinding();
+		WidgetMenuAnchor->SetIsOpen(!WidgetMenuAnchor->IsOpen());
+		return FReply::Handled();
+	}
 
-		UMVVMSubsystem::FConstDirectionalBindingArgs Args;
-		if (bIsGetter)
+	TSharedRef<SWidget> OnGetWidgetMenuContent()
+	{
+		FMVVMBlueprintViewBinding* ViewBinding = GetThisViewBinding();
+
+		UWidgetBlueprint* WidgetBlueprint = WidgetBlueprintWeak.Get();
+
+		EFieldVisibility Flags;
+		if (UE::MVVM::IsForwardBinding(ViewBinding->BindingType))
 		{
-			Args.SourceBinding = ViewModelBinding->ViewModelPath.GetFields().Last();
-			Args.DestinationBinding = ViewModelBinding->WidgetPath.GetFields().Last();
+			Flags |= EFieldVisibility::Writable;
 		}
-		else
+
+		if (UE::MVVM::IsBackwardBinding(ViewBinding->BindingType))
 		{
-			Args.SourceBinding = ViewModelBinding->WidgetPath.GetFields().Last();
-			Args.DestinationBinding = ViewModelBinding->ViewModelPath.GetFields().Last();
+			Flags |= EFieldVisibility::Readable;
+
+			if (!UE::MVVM::IsOneTimeBinding(ViewBinding->BindingType))
+			{
+				Flags |= EFieldVisibility::Notify;
+			}
 		}
 
-		return MoveTemp(Args);
+		FBindingSource Source = FBindingSource::CreateForWidget(WidgetBlueprint, ViewBinding->WidgetPath.GetWidgetName());
+
+		WidgetFieldIterator = MakeShared<UE::MVVM::FFieldIterator_Bindable>(WidgetBlueprint, Flags);
+		
+		return SNew(SBorder)
+			.BorderImage(FAppStyle::Get().GetBrush("ToolPanel.GroupBorder"))
+			.Padding(FMargin(8, 8, 8, 4))
+			[
+				SNew(SVerticalBox)
+				+ SVerticalBox::Slot()
+				[
+					SAssignNew(WidgetBindingList, SSourceBindingList, WidgetBlueprint)
+					.InitialSource(Source)
+					.ShowSearchBox(false)
+					.OnDoubleClicked(this, &SBindingRow::OnWidgetPropertySelectionChanged)
+					.FieldVisibilityFlags(Flags)
+				]
+				+ SVerticalBox::Slot()
+				.Padding(4, 4, 4, 0)
+				.AutoHeight()
+				[
+					SNew(SHorizontalBox)
+					+ SHorizontalBox::Slot()
+					[
+						SNew(SButton)
+						.OnClicked(this, &SBindingRow::OnWidgetButtonClicked, false)
+						.HAlign(HAlign_Center)
+						[
+							SNew(STextBlock)
+							.Text(LOCTEXT("Cancel", "Cancel"))
+						]
+					]
+					+ SHorizontalBox::Slot()
+					.Padding(4, 0, 0, 0)
+					[
+						SNew(SButton)
+						.OnClicked(this, &SBindingRow::OnWidgetClearButtonClicked)
+						.HAlign(HAlign_Center)
+						[
+							SNew(STextBlock)
+							.Text(LOCTEXT("Clear", "Clear"))
+						]
+					]
+					+ SHorizontalBox::Slot()
+					.Padding(4, 0, 0, 0)
+					[
+						SNew(SPrimaryButton)
+						.OnClicked(this, &SBindingRow::OnWidgetButtonClicked, true)
+						.Text(LOCTEXT("Select", "Select"))
+					]
+				]
+			];
 	}
 
 	TSharedRef<ITableRow> OnGenerateErrorRow(TSharedPtr<FText> Text, const TSharedRef<STableViewBase>& TableView) const
@@ -349,19 +574,19 @@ private:
 
 			if (bSourceToDest)
 			{
-				EditorSubsystem->SetSourceToDestinationConversionFunction(WidgetBlueprint, *ViewBinding, Function);
+				EditorSubsystem->SetSourceToDestinationConversionFunction(WidgetBlueprintWeak.Get(), *ViewBinding, Function);
 			} 
 			else
 			{
-				EditorSubsystem->SetDestinationToSourceConversionFunction(WidgetBlueprint, *ViewBinding, Function);
+				EditorSubsystem->SetDestinationToSourceConversionFunction(WidgetBlueprintWeak.Get(), *ViewBinding, Function);
 			}
 		}
 	}
 
-	TArray<UE::MVVM::FBindingSource> GetAvailableViewModels() const
+	TArray<FBindingSource> GetAvailableViewModels() const
 	{
 		UMVVMEditorSubsystem* EditorSubsystem = GEditor->GetEditorSubsystem<UMVVMEditorSubsystem>();
-		return EditorSubsystem->GetAllViewModels(WidgetBlueprint);
+		return EditorSubsystem->GetAllViewModels(WidgetBlueprintWeak.Get());
 	}
 
 	FMVVMBlueprintPropertyPath GetSelectedViewModelProperty() const
@@ -391,39 +616,58 @@ private:
 			if (ViewBinding->ViewModelPath != SelectedField)
 			{
 				UMVVMEditorSubsystem* Subsystem = GEditor->GetEditorSubsystem<UMVVMEditorSubsystem>();
-				Subsystem->SetViewModelPropertyForBinding(WidgetBlueprint, *ViewBinding, SelectedField);
-
-				if (WidgetFieldSelector.IsValid())
-				{
-					WidgetFieldSelector->Refresh();
-				}
+				Subsystem->SetViewModelPropertyForBinding(WidgetBlueprintWeak.Get(), *ViewBinding, SelectedField);
 			}
 		}
 	}
 
-	void OnWidgetPropertySelectionChanged(FMVVMBlueprintPropertyPath SelectedField)
+	void OnWidgetPropertySelectionChanged(const FMVVMBlueprintPropertyPath& SelectedField)
 	{
 		if (FMVVMBlueprintViewBinding* ViewBinding = GetThisViewBinding())
 		{
 			if (ViewBinding->WidgetPath != SelectedField)
 			{
 				UMVVMEditorSubsystem* Subsystem = GEditor->GetEditorSubsystem<UMVVMEditorSubsystem>();
-				Subsystem->SetWidgetPropertyForBinding(WidgetBlueprint, *ViewBinding, SelectedField);
+				Subsystem->SetWidgetPropertyForBinding(WidgetBlueprintWeak.Get(), *ViewBinding, SelectedField);
 
 				if (ViewModelFieldSelector.IsValid())
 				{
 					ViewModelFieldSelector->Refresh();
 				}
+
+				if (SelectedWidget.IsValid())
+				{
+					SelectedWidget->SetField(SelectedField);
+				}
 			}
 		}
 	}
 
+	FReply OnWidgetButtonClicked(bool bSelect)
+	{
+		if (bSelect)
+		{
+			OnWidgetPropertySelectionChanged(WidgetBindingList->GetSelectedProperty());
+		}
+
+		WidgetMenuAnchor->SetIsOpen(false);
+		return FReply::Handled();
+	}
+
+	FReply OnWidgetClearButtonClicked()
+	{
+		FMVVMBlueprintPropertyPath Path;
+		Path.SetWidgetName(Entry->GetWidgetName());
+		OnWidgetPropertySelectionChanged(Path);
+		return FReply::Handled();
+	}
+
 	void OnUpdateModeSelectionChanged(int32 Value, ESelectInfo::Type)
 	{
 		if (FMVVMBlueprintViewBinding* ViewBinding = GetThisViewBinding())
 		{
 			UMVVMEditorSubsystem* Subsystem = GEditor->GetEditorSubsystem<UMVVMEditorSubsystem>();
-			Subsystem->SetUpdateModeForBinding(WidgetBlueprint, *ViewBinding, (EMVVMViewBindingUpdateMode) Value);
+			Subsystem->SetUpdateModeForBinding(WidgetBlueprintWeak.Get(), *ViewBinding, (EMVVMViewBindingUpdateMode) Value);
 		}			
 	}
 
@@ -442,7 +686,7 @@ private:
 		if (FMVVMBlueprintViewBinding* ViewBinding = GetThisViewBinding())
 		{
 			UMVVMEditorSubsystem* Subsystem = GEditor->GetEditorSubsystem<UMVVMEditorSubsystem>();
-			Subsystem->SetEnabledForBinding(WidgetBlueprint, *ViewBinding, NewState == ECheckBoxState::Checked);
+			Subsystem->SetEnabledForBinding(WidgetBlueprintWeak.Get(), *ViewBinding, NewState == ECheckBoxState::Checked);
 		}
 	}
 
@@ -456,7 +700,7 @@ private:
 		if (FMVVMBlueprintViewBinding* ViewBinding = GetThisViewBinding())
 		{
 			UMVVMEditorSubsystem* Subsystem = GEditor->GetEditorSubsystem<UMVVMEditorSubsystem>();
-			Subsystem->SetCompileForBinding(WidgetBlueprint, *ViewBinding, NewState == ECheckBoxState::Checked);
+			Subsystem->SetCompileForBinding(WidgetBlueprintWeak.Get(), *ViewBinding, NewState == ECheckBoxState::Checked);
 		}
 	}
 
@@ -546,25 +790,20 @@ private:
 			EMVVMBindingMode NewMode = (EMVVMBindingMode) ModeEnum->GetValueByName(ValueName);
 
 			UMVVMEditorSubsystem* Subsystem = GEditor->GetEditorSubsystem<UMVVMEditorSubsystem>();
-			Subsystem->SetBindingTypeForBinding(WidgetBlueprint, *ViewBinding, NewMode);
+			Subsystem->SetBindingTypeForBinding(WidgetBlueprintWeak.Get(), *ViewBinding, NewMode);
 
 			if (ViewModelFieldSelector.IsValid())
 			{
 				ViewModelFieldSelector->Refresh();
 			}
-
-			if (WidgetFieldSelector.IsValid())
-			{
-				WidgetFieldSelector->Refresh();
-			}
 		}
 	}
 
 	FReply HandleDropDownOptionsPressed() 
 	{
-		if (TSharedPtr<ITypedTableView<FMVVMViewBindingListEntryPtr>> ListView = OwnerTablePtr.Pin())
+		if (TSharedPtr<ITypedTableView<TSharedPtr<FBindingEntry>>> ListView = OwnerTablePtr.Pin())
 		{
-			if (TSharedPtr<SMVVMViewBindingListView> ParentList = StaticCastSharedPtr<SMVVMViewBindingListView>(ListView->AsWidget()->GetParentWidget()))
+			if (TSharedPtr<SBindingsList> ParentList = StaticCastSharedPtr<SBindingsList>(ListView->AsWidget()->GetParentWidget()))
 			{
 				// Get the context menu content. If invalid, don't open a menu.
 				ListView->Private_SetItemSelection(Entry, true);
@@ -585,98 +824,49 @@ private:
 	void HandleBlueprintChanged(UBlueprint* Blueprint)
 	{
 		ViewModelFieldSelector->Refresh();
-		WidgetFieldSelector->Refresh();
 	}
 
 private:
-	FMVVMViewBindingListEntryPtr Entry;
+	TSharedPtr<FBindingEntry> Entry;
+	TWeakObjectPtr<UWidgetBlueprint> WidgetBlueprintWeak;
 	TWeakObjectPtr<UMVVMBlueprintView> BlueprintView;
-	UWidgetBlueprint* WidgetBlueprint = nullptr;
-	TSharedPtr<UE::MVVM::SFieldSelector> ViewModelFieldSelector;
-	TSharedPtr<UE::MVVM::SFieldSelector> WidgetFieldSelector;
+	TSharedPtr<SFieldSelector> ViewModelFieldSelector;
 	TSharedPtr<SWidget> ContextMenuOptionHelper;
 	TSharedPtr<SCustomDialog> ErrorDialog;
 	TArray<TSharedPtr<FText>> ErrorItems;
-	TArray<FName> ModeNames;
 	FDelegateHandle OnBlueprintChangedHandle;
+	TSharedPtr<UE::MVVM::FFieldIterator_Bindable> WidgetFieldIterator;
+	TSharedPtr<SMenuAnchor> WidgetMenuAnchor;
+	TSharedPtr<SSourceBindingList> WidgetBindingList;
+	TSharedPtr<SFieldEntry> SelectedWidget;
 };
 
-FName SMVVMViewBindingListEntryRow::EnabledColumnName = "Enabled";
-FName SMVVMViewBindingListEntryRow::CompileColumnName = "Compile";
-FName SMVVMViewBindingListEntryRow::ErrorColumnName = "Error";
-FName SMVVMViewBindingListEntryRow::ViewModelColumnName = "ViewModelProperty";
-FName SMVVMViewBindingListEntryRow::ModeColumnName = "Mode";
-FName SMVVMViewBindingListEntryRow::WidgetColumName = "WidgetProperty";
-FName SMVVMViewBindingListEntryRow::UpdateColumnName = "Update";
-FName SMVVMViewBindingListEntryRow::ConversionFunctionColumnName = "ConversionFunction";
-FName SMVVMViewBindingListEntryRow::DropDownOptionsColumnName = "DropdownOptions";
-
-/** */
-void SMVVMViewBindingListView::Construct(const FArguments& InArgs, TSharedPtr<SMVVMViewBindingPanel> Owner, UMVVMWidgetBlueprintExtension_View* InMVVMExtension)
+void SBindingsList::Construct(const FArguments& InArgs, TSharedPtr<SBindingsPanel> Owner, UMVVMWidgetBlueprintExtension_View* InMVVMExtension)
 {
 	BindingPanel = Owner;
 	MVVMExtension = InMVVMExtension;
 	check(InMVVMExtension);
 
-	MVVMExtension->OnBlueprintViewChangedDelegate().AddSP(this, &SMVVMViewBindingListView::RequestListRefresh);
-	MVVMExtension->GetBlueprintView()->OnBindingsUpdated.AddSP(this, &SMVVMViewBindingListView::RequestListRefresh);
-	MVVMExtension->GetBlueprintView()->OnViewModelsUpdated.AddSP(this, &SMVVMViewBindingListView::RequestListRefresh);
-
-	RequestListRefresh();
+	MVVMExtension->OnBlueprintViewChangedDelegate().AddSP(this, &SBindingsList::Refresh);
+	MVVMExtension->GetBlueprintView()->OnBindingsUpdated.AddSP(this, &SBindingsList::Refresh);
+	MVVMExtension->GetBlueprintView()->OnViewModelsUpdated.AddSP(this, &SBindingsList::Refresh);
 
 	ChildSlot
 	[
-		SAssignNew(ListView, SListView<FMVVMViewBindingListEntryPtr>)
-		.ListItemsSource(&SourceData)
+		SAssignNew(TreeView, STreeView<TSharedPtr<FBindingEntry>>)
+		.TreeItemsSource(&RootWidgets)
 		.SelectionMode(ESelectionMode::Single)
-		.OnGenerateRow(this, &SMVVMViewBindingListView::MakeSourceListViewWidget)
-		.OnContextMenuOpening(this, &SMVVMViewBindingListView::OnSourceConstructContextMenu)
-		.OnSelectionChanged(this, &SMVVMViewBindingListView::OnSourceListSelectionChanged)
-		.HeaderRow
-		(
-			SNew(SHeaderRow)
-			+ SHeaderRow::Column(SMVVMViewBindingListEntryRow::CompileColumnName)
-			.DefaultLabel(FText::GetEmpty())
-			.FixedWidth(25.f)
-			+ SHeaderRow::Column(SMVVMViewBindingListEntryRow::ErrorColumnName)
-			.DefaultLabel(FText::GetEmpty())
-			.FixedWidth(25.f)
-			.HeaderContent()
-			[
-				SNew(SBox)
-				.WidthOverride(16)
-				.HeightOverride(16)
-				.VAlign(VAlign_Center)
-				.HAlign(HAlign_Center)
-				[
-					SNew(SImage)
-					.Image(FAppStyle::Get().GetBrush("Icons.Error"))
-				]
-			]
-			+ SHeaderRow::Column(SMVVMViewBindingListEntryRow::ViewModelColumnName)
-			.FillWidth(0.25f)
-			.DefaultLabel(LOCTEXT("ViewModelProperty", "View Model Property"))
-			+ SHeaderRow::Column(SMVVMViewBindingListEntryRow::ModeColumnName)
-			.FixedWidth(52)
-			.DefaultLabel(LOCTEXT("Mode", "Mode"))
-			+ SHeaderRow::Column(SMVVMViewBindingListEntryRow::WidgetColumName)
-			.FillWidth(0.25f)
-			.DefaultLabel(LOCTEXT("WidgetProperty", "Widget Property"))
-			+ SHeaderRow::Column(SMVVMViewBindingListEntryRow::UpdateColumnName)
-			.FillWidth(0.05f)
-			.DefaultLabel(LOCTEXT("Update", "Update"))
-			+ SHeaderRow::Column(SMVVMViewBindingListEntryRow::ConversionFunctionColumnName)
-			.FillWidth(0.05f)
-			.DefaultLabel(LOCTEXT("Conversion", "Conversion"))
-			+ SHeaderRow::Column(SMVVMViewBindingListEntryRow::DropDownOptionsColumnName)
-			.FixedWidth(25.f)
-			.DefaultLabel(FText::GetEmpty())
-		)
+		.OnGenerateRow(this, &SBindingsList::GenerateEntryRow)
+		.OnGetChildren(this, &SBindingsList::GetChildrenOfEntry)
+		.OnContextMenuOpening(this, &SBindingsList::OnSourceConstructContextMenu)
+		.OnSelectionChanged(this, &SBindingsList::OnSourceListSelectionChanged)
+		.ItemHeight(32)
 	];
-}
 
+	Refresh();
+}
 
-SMVVMViewBindingListView::~SMVVMViewBindingListView()
+SBindingsList::~SBindingsList()
 {
 	if (UMVVMWidgetBlueprintExtension_View* MVVMExtensionPtr = MVVMExtension.Get())
 	{
@@ -686,79 +876,135 @@ SMVVMViewBindingListView::~SMVVMViewBindingListView()
 	}
 }
 
-void SMVVMViewBindingListView::RequestListRefresh()
+void SBindingsList::GetChildrenOfEntry(TSharedPtr<FBindingEntry> Entry, TArray<TSharedPtr<FBindingEntry>>& OutChildren) const
+{
+	TConstArrayView<TSharedPtr<FBindingEntry>> Children = Entry->GetChildren();
+	OutChildren.Append(Children.GetData(), Children.Num());
+}
+
+void SBindingsList::Refresh()
 {
-	int32 SelectedIndex = -1;
-	if (ListView.IsValid())
+	UMVVMWidgetBlueprintExtension_View* MVVMExtensionPtr = MVVMExtension.Get();
+	UMVVMBlueprintView* BlueprintView = MVVMExtensionPtr ? MVVMExtensionPtr->GetBlueprintView() : nullptr;
+
+	// store the current binding index
+	TArray<FMVVMBlueprintViewBinding*, TInlineAllocator<8>> SelectedBindings;
+
+	if (TreeView.IsValid() && BlueprintView)
 	{
-		if (const UMVVMWidgetBlueprintExtension_View* MVVMExtensionPtr = MVVMExtension.Get())
+		TArray<TSharedPtr<FBindingEntry>> SelectedItems = TreeView->GetSelectedItems();
+		for (const TSharedPtr<FBindingEntry>& Entry : SelectedItems)
 		{
-			if (const UMVVMBlueprintView* BlueprintView = MVVMExtensionPtr->GetBlueprintView())
+			FMVVMBlueprintViewBinding* Binding = Entry->GetBinding(BlueprintView);
+			if (Binding != nullptr)
 			{
-				TArray<FMVVMViewBindingListEntryPtr> SelectedItems = ListView->GetSelectedItems();
-				for (const FMVVMViewBindingListEntryPtr& Entry : SelectedItems)
-				{
-					SelectedIndex = Entry->Index;
-					break;
-				}
+				SelectedBindings.Add(Binding);
 			}
 		}
 	}
 
-	SourceData.Reset();
-	if (const UMVVMWidgetBlueprintExtension_View* MVVMExtensionPtr = MVVMExtension.Get())
+	RootWidgets.Reset();
+
+	TArray<TSharedPtr<FBindingEntry>> SelectedEntries;
+
+	// generate our entries
+	// for each widget with bindings, create an entry at the root level
+	// then add all bindings that reference that widget as its children
+	if (BlueprintView)
 	{
-		if (const UMVVMBlueprintView* BlueprintView = MVVMExtensionPtr->GetBlueprintView())
+		TArrayView<FMVVMBlueprintViewBinding> Bindings = BlueprintView->GetBindings();
+		for (int32 BindingIndex = 0; BindingIndex < Bindings.Num(); ++BindingIndex)
 		{
-			const int32 Count = BlueprintView->GetBindings().Num();
-			for (int32 Index = 0; Index < Count; ++Index)
+			const FMVVMBlueprintViewBinding& Binding = Bindings[BindingIndex];
+			
+			FName WidgetName = Binding.WidgetPath.GetWidgetName();
+
+			TSharedPtr<FBindingEntry> ExistingWidget;
+			for (TSharedPtr<FBindingEntry> Widget : RootWidgets)
+			{
+				if (WidgetName == Widget->GetWidgetName())
+				{
+					ExistingWidget = Widget;
+					break;
+				}
+			}
+
+			if (!ExistingWidget.IsValid())
+			{
+				ExistingWidget = RootWidgets.Add_GetRef(MakeShared<FBindingEntry>(WidgetName));
+			}
+
+			TSharedRef<FBindingEntry> NewBindingEntry = MakeShared<FBindingEntry>(BindingIndex);
+			ExistingWidget->AddChild(NewBindingEntry);
+
+			if (SelectedBindings.Contains(&Binding))
 			{
-				SourceData.Add(MakeShared<FMVVMViewBindingListEntry>(Index));
+				SelectedEntries.Add(NewBindingEntry);
 			}
 		}
 	}
 
-	if (ListView.IsValid())
+	if (TreeView.IsValid())
 	{
-		ListView->RequestListRefresh();
+		TreeView->RequestTreeRefresh();
 
-		if (SourceData.IsValidIndex(SelectedIndex))
+		TreeView->SetItemSelection(SelectedEntries, true);
+
+		for (const TSharedPtr<FBindingEntry>& Entry : RootWidgets)
 		{
-			ListView->SetItemSelection(SourceData[SelectedIndex], true);
+			Private::ExpandAll(TreeView, Entry);
 		}
 	}
 }
 
 
-TSharedRef<ITableRow> SMVVMViewBindingListView::MakeSourceListViewWidget(FMVVMViewBindingListEntryPtr Entry, const TSharedRef<STableViewBase>& OwnerTable) const
+TSharedRef<ITableRow> SBindingsList::GenerateEntryRow(TSharedPtr<FBindingEntry> Entry, const TSharedRef<STableViewBase>& OwnerTable) const
 {
+	TSharedPtr<ITableRow> Row;
+
 	if (UMVVMWidgetBlueprintExtension_View* MVVMExtensionPtr = MVVMExtension.Get())
 	{
-		return SNew(SMVVMViewBindingListEntryRow, OwnerTable, Entry, MVVMExtensionPtr->GetBlueprintView(), MVVMExtensionPtr->GetWidgetBlueprint());
+		if (Entry->GetBinding(MVVMExtensionPtr->GetBlueprintView()) != nullptr)
+		{
+			Row = SNew(SBindingRow, OwnerTable, Entry, MVVMExtensionPtr->GetBlueprintView(), MVVMExtensionPtr->GetWidgetBlueprint());
+		}
+		else
+		{
+			Row = SNew(SWidgetRow, OwnerTable, Entry, MVVMExtensionPtr->GetWidgetBlueprint());
+		}
+
+		return Row.ToSharedRef();
 	}
-	return SNew(SMVVMViewBindingListEntryRow, OwnerTable, Entry, nullptr, nullptr);
+
+	ensureMsgf(false, TEXT("Failed to create binding or widget row."));
+	return SNew(STableRow<TSharedPtr<FBindingEntry>>, OwnerTable);
 }
 
-TSharedPtr<SWidget> SMVVMViewBindingListView::OnSourceConstructContextMenu()
+TSharedPtr<SWidget> SBindingsList::OnSourceConstructContextMenu()
 {
 	const bool bShouldCloseWindowAfterMenuSelection = true;
 	FMenuBuilder MenuBuilder(bShouldCloseWindowAfterMenuSelection, nullptr);
 
-	TArray<FMVVMViewBindingListEntryPtr> Selection = ListView->GetSelectedItems();
+	TArray<TSharedPtr<FBindingEntry>> Selection = TreeView->GetSelectedItems();
 	if (Selection.Num() > 0)
 	{
 		FUIAction RemoveAction;
-		RemoveAction.ExecuteAction = FExecuteAction::CreateLambda([this, ToRemove = Selection[0]->Index]()
+		RemoveAction.ExecuteAction = FExecuteAction::CreateLambda([this, Selection]()
 			{
 				if (UMVVMWidgetBlueprintExtension_View* MVVMExtensionPtr = MVVMExtension.Get())
 				{
-					if (TSharedPtr<SMVVMViewBindingPanel> BindingPanelPtr = BindingPanel.Pin())
+					if (TSharedPtr<SBindingsPanel> BindingPanelPtr = BindingPanel.Pin())
 					{
-						BindingPanelPtr->OnBindingListSelectionChanged(INDEX_NONE);
+						BindingPanelPtr->OnBindingListSelectionChanged(TConstArrayView<FMVVMBlueprintViewBinding*>());
 					}
 
-					UMVVMBlueprintView* BlueprintView = MVVMExtensionPtr->GetBlueprintView();
-					BlueprintView->RemoveBindingAt(ToRemove);
+					if (UMVVMBlueprintView* BlueprintView = MVVMExtensionPtr->GetBlueprintView())
+					{
+						for (const TSharedPtr<FBindingEntry>& Entry : Selection)
+						{
+							BlueprintView->RemoveBinding(Entry->GetBinding(BlueprintView));
+						}
+					}
 				}
 			});
 		MenuBuilder.AddMenuEntry(LOCTEXT("RemoveBinding", "Remove Binding"), LOCTEXT("RemoveBindingTooltip", "Remove this binding."), FSlateIcon(FAppStyle::GetAppStyleSetName(), "Icons.Delete"), RemoveAction);
@@ -767,7 +1013,7 @@ TSharedPtr<SWidget> SMVVMViewBindingListView::OnSourceConstructContextMenu()
 	return MenuBuilder.MakeWidget();
 }
 
-void SMVVMViewBindingListView::OnSourceListSelectionChanged(FMVVMViewBindingListEntryPtr Entry, ESelectInfo::Type SelectionType) const
+void SBindingsList::OnSourceListSelectionChanged(TSharedPtr<FBindingEntry> Entry, ESelectInfo::Type SelectionType) const
 {
 	if (bSelectionChangedGuard)
 	{
@@ -775,11 +1021,29 @@ void SMVVMViewBindingListView::OnSourceListSelectionChanged(FMVVMViewBindingList
 	}
 	TGuardValue<bool> ReentrantGuard(bSelectionChangedGuard, true);
 
-	const int32 SelectionIndex = Entry ? Entry->Index : INDEX_NONE;
-	if (TSharedPtr<SMVVMViewBindingPanel> BindingPanelPtr = BindingPanel.Pin())
+	if (TSharedPtr<SBindingsPanel> BindingPanelPtr = BindingPanel.Pin())
 	{
-		BindingPanelPtr->OnBindingListSelectionChanged(SelectionIndex);
+		if (UMVVMWidgetBlueprintExtension_View* MVVMExtensionPtr = MVVMExtension.Get())
+		{
+			if (UMVVMBlueprintView* View = MVVMExtensionPtr->GetBlueprintView())
+			{
+				TArray<TSharedPtr<FBindingEntry>> SelectedEntries = TreeView->GetSelectedItems();
+				TArray<FMVVMBlueprintViewBinding*> SelectedBindings;
+
+				for (const TSharedPtr<FBindingEntry>& SelectedEntry : SelectedEntries)
+				{
+					if (FMVVMBlueprintViewBinding* SelectedBinding = Entry->GetBinding(View))
+					{
+						SelectedBindings.Add(SelectedBinding);
+					}
+				}
+
+				BindingPanelPtr->OnBindingListSelectionChanged(SelectedBindings);
+			}
+		}
 	}
 }
 
+} // namespace UE::MVVM
+
 #undef LOCTEXT_NAMESPACE
diff --git a/Engine/Plugins/Runtime/ModelViewViewModel/Source/ModelViewViewModelEditor/Private/Widgets/SMVVMViewBindingListView.h b/Engine/Plugins/Runtime/ModelViewViewModel/Source/ModelViewViewModelEditor/Private/Widgets/SMVVMViewBindingListView.h
--- a/Engine/Plugins/Runtime/ModelViewViewModel/Source/ModelViewViewModelEditor/Private/Widgets/SMVVMViewBindingListView.h
+++ b/Engine/Plugins/Runtime/ModelViewViewModel/Source/ModelViewViewModelEditor/Private/Widgets/SMVVMViewBindingListView.h
@@ -8,36 +8,41 @@
 #include "Widgets/SCompoundWidget.h"
 
 class ITableRow;
-template <typename ItemType> class SListView;
+template <typename ItemType> class STreeView;
 class STableViewBase;
-class SMVVMViewBindingPanel;
 class UMVVMWidgetBlueprintExtension_View;
 
-struct FMVVMViewBindingListEntry;
-using FMVVMViewBindingListEntryPtr = TSharedPtr<FMVVMViewBindingListEntry>;
+namespace UE::MVVM
+{
+class SBindingsPanel;
+
+struct FBindingEntry;
 
-class SMVVMViewBindingListView : public SCompoundWidget
+class SBindingsList : public SCompoundWidget
 {
 public:
-	SLATE_BEGIN_ARGS(SMVVMViewBindingListView) {}
+	SLATE_BEGIN_ARGS(SBindingsList) {}
 	SLATE_END_ARGS()
 
-	void Construct(const FArguments& InArgs, TSharedPtr<SMVVMViewBindingPanel> Owner, UMVVMWidgetBlueprintExtension_View* MVVMExtension);
-	~SMVVMViewBindingListView();
+	void Construct(const FArguments& InArgs, TSharedPtr<SBindingsPanel> Owner, UMVVMWidgetBlueprintExtension_View* MVVMExtension);
+	~SBindingsList();
 
-	void RequestListRefresh();
+	void Refresh();
 
 	/** Constructs context menu used for right click and dropdown button */
 	TSharedPtr<SWidget> OnSourceConstructContextMenu();
 
 private:
-	TSharedRef<ITableRow> MakeSourceListViewWidget(FMVVMViewBindingListEntryPtr Entry, const TSharedRef<STableViewBase>& OwnerTable) const;
-	void OnSourceListSelectionChanged(FMVVMViewBindingListEntryPtr Entry, ESelectInfo::Type SelectionType) const;
+	TSharedRef<ITableRow> GenerateEntryRow(TSharedPtr<FBindingEntry> Entry, const TSharedRef<STableViewBase>& OwnerTable) const;
+	void OnSourceListSelectionChanged(TSharedPtr<FBindingEntry> Entry, ESelectInfo::Type SelectionType) const;
+	void GetChildrenOfEntry(TSharedPtr<FBindingEntry> Entry, TArray<TSharedPtr<FBindingEntry>>& OutChildren) const;
 
 private:
-	TWeakPtr<SMVVMViewBindingPanel> BindingPanel;
-	TSharedPtr<SListView<FMVVMViewBindingListEntryPtr>> ListView;
-	TArray<FMVVMViewBindingListEntryPtr> SourceData;
+	TWeakPtr<SBindingsPanel> BindingPanel;
+	TSharedPtr<STreeView<TSharedPtr<FBindingEntry>>> TreeView;
+	TArray<TSharedPtr<FBindingEntry>> RootWidgets;
 	TWeakObjectPtr<UMVVMWidgetBlueprintExtension_View> MVVMExtension;
 	mutable bool bSelectionChangedGuard = false;
-};
\ No newline at end of file
+};
+
+} // namespace UE::MVVM
diff --git a/Engine/Plugins/Runtime/ModelViewViewModel/Source/ModelViewViewModelEditor/Private/Widgets/SMVVMViewBindingPanel.cpp b/Engine/Plugins/Runtime/ModelViewViewModel/Source/ModelViewViewModelEditor/Private/Widgets/SMVVMViewBindingPanel.cpp
--- a/Engine/Plugins/Runtime/ModelViewViewModel/Source/ModelViewViewModelEditor/Private/Widgets/SMVVMViewBindingPanel.cpp
+++ b/Engine/Plugins/Runtime/ModelViewViewModel/Source/ModelViewViewModelEditor/Private/Widgets/SMVVMViewBindingPanel.cpp
@@ -34,9 +34,11 @@
 
 #define LOCTEXT_NAMESPACE "BindingPanel"
 
+namespace UE::MVVM
+{
 
 /** */
-void SMVVMViewBindingPanel::Construct(const FArguments& InArgs, TSharedPtr<FWidgetBlueprintEditor> WidgetBlueprintEditor, bool bInIsDrawerTab)
+void SBindingsPanel::Construct(const FArguments& InArgs, TSharedPtr<FWidgetBlueprintEditor> WidgetBlueprintEditor, bool bInIsDrawerTab)
 {
 	WeakBlueprintEditor = WidgetBlueprintEditor;
 	bIsDrawerTab = bInIsDrawerTab;
@@ -48,11 +50,11 @@ void SMVVMViewBindingPanel::Construct(const FArguments& InArgs, TSharedPtr<FWidg
 	MVVMExtension = MVVMExtensionPtr;
 	if (MVVMExtensionPtr)
 	{
-		BlueprintViewChangedDelegateHandle = MVVMExtensionPtr->OnBlueprintViewChangedDelegate().AddSP(this, &SMVVMViewBindingPanel::HandleBlueprintViewChangedDelegate);
+		BlueprintViewChangedDelegateHandle = MVVMExtensionPtr->OnBlueprintViewChangedDelegate().AddSP(this, &SBindingsPanel::HandleBlueprintViewChangedDelegate);
 	}
 	else
 	{
-		WidgetBlueprint->OnExtensionAdded.AddRaw(this, &SMVVMViewBindingPanel::HandleExtensionAdded);
+		WidgetBlueprint->OnExtensionAdded.AddRaw(this, &SBindingsPanel::HandleExtensionAdded);
 	}
 
 	{
@@ -67,7 +69,6 @@ void SMVVMViewBindingPanel::Construct(const FArguments& InArgs, TSharedPtr<FWidg
 		DetailsViewArgs.ViewIdentifier = NAME_None;
 		DetailsView = PropertyEditorModule.CreateDetailView(DetailsViewArgs);
 		
-		//DetailsView->OnFinishedChangingProperties().AddSP(this, &SLiveLinkClientPanel::OnPropertyChanged);
 		FStructureDetailsViewArgs StructureDetailsViewArgs;
 		StructDetailsView = PropertyEditorModule.CreateStructureDetailView(DetailsViewArgs, StructureDetailsViewArgs, TSharedPtr<FStructOnScope>());
 		StructDetailsView->GetDetailsView()->RegisterInstancedCustomPropertyTypeLayout(FMVVMBlueprintPropertyPath::StaticStruct()->GetFName(), FOnGetPropertyTypeCustomizationInstance::CreateStatic(&UE::MVVM::FPropertyPathCustomization::MakeInstance, WidgetBlueprint));
@@ -78,7 +79,7 @@ void SMVVMViewBindingPanel::Construct(const FArguments& InArgs, TSharedPtr<FWidg
 }
 
 
-SMVVMViewBindingPanel::~SMVVMViewBindingPanel()
+SBindingsPanel::~SBindingsPanel()
 {
 	if (UMVVMWidgetBlueprintExtension_View* Extension = MVVMExtension.Get())
 	{
@@ -87,7 +88,6 @@ SMVVMViewBindingPanel::~SMVVMViewBindingPanel()
 	if (TSharedPtr<FWidgetBlueprintEditor> WidgetEditor = WeakBlueprintEditor.Pin())
 	{
 		UWidgetBlueprint* WidgetBlueprint = WidgetEditor->GetWidgetBlueprintObj();
-
 		if (WidgetBlueprint)
 		{
 			WidgetBlueprint->OnExtensionAdded.RemoveAll(this);
@@ -95,7 +95,7 @@ SMVVMViewBindingPanel::~SMVVMViewBindingPanel()
 	}
 }
 
-FReply SMVVMViewBindingPanel::OnKeyDown(const FGeometry& MyGeometry, const FKeyEvent& InKeyEvent)
+FReply SBindingsPanel::OnKeyDown(const FGeometry& MyGeometry, const FKeyEvent& InKeyEvent)
 {
 	FReply Reply = FReply::Unhandled();
 	if (TSharedPtr<FWidgetBlueprintEditor> WidgetEditor = WeakBlueprintEditor.Pin())
@@ -109,12 +109,12 @@ FReply SMVVMViewBindingPanel::OnKeyDown(const FGeometry& MyGeometry, const FKeyE
 	return Reply;
 }
 
-bool SMVVMViewBindingPanel::SupportsKeyboardFocus() const
+bool SBindingsPanel::SupportsKeyboardFocus() const
 {
 	return true;
 }
 
-void SMVVMViewBindingPanel::HandleBlueprintViewChangedDelegate()
+void SBindingsPanel::HandleBlueprintViewChangedDelegate()
 {
 	UMVVMWidgetBlueprintExtension_View* MVVMExtensionPtr = MVVMExtension.Get();
 
@@ -139,44 +139,35 @@ void SMVVMViewBindingPanel::HandleBlueprintViewChangedDelegate()
 }
 
 
-void SMVVMViewBindingPanel::OnBindingListSelectionChanged(int32 Index)
+void SBindingsPanel::OnBindingListSelectionChanged(TConstArrayView<FMVVMBlueprintViewBinding*> Selection)
 {
-	bool bDefault = true;
-	if (UMVVMWidgetBlueprintExtension_View* MVVMExtensionPtr = MVVMExtension.Get())
+	if (FMVVMBlueprintViewBinding* Binding = Selection.Num() > 0 ? Selection.Last() : nullptr)
 	{
-		if (UMVVMBlueprintView* BlueprintView = MVVMExtensionPtr->GetBlueprintView())
-		{
-			FMVVMBlueprintViewBinding* Binding = BlueprintView->GetBindingAt(Index);
-			if (Binding != nullptr)
-			{
-				TSharedRef<FStructOnScope> StructScope = MakeShared<FStructOnScope>(FMVVMBlueprintViewBinding::StaticStruct(), reinterpret_cast<uint8*>(Binding));
-				StructDetailsView->SetStructureData(StructScope);
-				DetailContainer->SetContent(StructDetailsView->GetWidget().ToSharedRef());
-				bDefault = false;
-			}
-		}
+		TSharedRef<FStructOnScope> StructScope = MakeShared<FStructOnScope>(FMVVMBlueprintViewBinding::StaticStruct(), reinterpret_cast<uint8*>(Binding));
+		StructDetailsView->SetStructureData(StructScope);
+		DetailContainer->SetContent(StructDetailsView->GetWidget().ToSharedRef());
 	}
-
-	if (bDefault)
+	else
 	{
+		// show empty details view
 		DetailsView->SetObject(nullptr);
 		DetailContainer->SetContent(DetailsView.ToSharedRef());
 	}
 }
 
-void SMVVMViewBindingPanel::AddDefaultBinding()
+void SBindingsPanel::AddDefaultBinding()
 {
 	if (UMVVMWidgetBlueprintExtension_View* MVVMExtensionPtr = MVVMExtension.Get())
 	{
 		if (UMVVMBlueprintView* EditorData = MVVMExtensionPtr->GetBlueprintView())
 		{
 			EditorData->AddDefaultBinding();
-			ListView->RequestListRefresh();
+			BindingsList->Refresh();
 		}
 	}
 }
 
-bool SMVVMViewBindingPanel::CanAddBinding() const
+bool SBindingsPanel::CanAddBinding() const
 {
 	if (UMVVMWidgetBlueprintExtension_View* MVVMExtensionPtr = MVVMExtension.Get())
 	{
@@ -189,7 +180,7 @@ bool SMVVMViewBindingPanel::CanAddBinding() const
 	return false;
 }
 
-FText SMVVMViewBindingPanel::GetAddBindingToolTip() const
+FText SBindingsPanel::GetAddBindingToolTip() const
 {
 	if (CanAddBinding())
 	{
@@ -201,54 +192,16 @@ FText SMVVMViewBindingPanel::GetAddBindingToolTip() const
 	}
 }
 
-void SMVVMViewBindingPanel::ShowManageViewModelsWindow()
-{
-	if (UMVVMWidgetBlueprintExtension_View* MVVMExtensionPtr = MVVMExtension.Get())
-	{
-		if (UMVVMBlueprintView* BlueprintView = MVVMExtensionPtr->GetBlueprintView())
-		{
-			TWeakObjectPtr<UMVVMWidgetBlueprintExtension_View> LocalMVVMExtension = MVVMExtension;
-			FOnViewModelContextsPicked ViewModelContextsPickedDelegate = FOnViewModelContextsPicked::CreateLambda
-			(
-				[LocalMVVMExtension](TArray<FMVVMBlueprintViewModelContext> UpdatedContexts)
-				{
-					UMVVMWidgetBlueprintExtension_View* MVVMExtensionPtr = LocalMVVMExtension.Get();
-					if (ensure(MVVMExtensionPtr))
-					{
-						if (UMVVMBlueprintView* BlueprintView = MVVMExtensionPtr->GetBlueprintView())
-						{
-							BlueprintView->SetViewModels(UpdatedContexts);
-						}
-					}
-				}
-			);
-
-			TSharedRef<SWindow> ManageViewModelsWindow = SNew(SWindow)
-				.Title(LOCTEXT("MVVMManageViewMoedlsWindowHeader", "Manage ViewModels"))
-				.SupportsMaximize(false)
-				.ClientSize(FVector2D(800.0f, 600.0f));
-
-			TSharedRef<SMVVMManageViewModelsWidget> ManageViewModelsWidget = SNew(SMVVMManageViewModelsWidget)
-				.ParentWindow(ManageViewModelsWindow)
-				.OnViewModelContextsPickedDelegate(ViewModelContextsPickedDelegate)
-				.WidgetBlueprint(WeakBlueprintEditor.Pin()->GetWidgetBlueprintObj());
-
-			ManageViewModelsWindow->SetContent(ManageViewModelsWidget);
-			GEditor->EditorAddModalWindow(ManageViewModelsWindow);
-		}
-	}
-}
-
-TSharedRef<SWidget> SMVVMViewBindingPanel::CreateDrawerDockButton()
+TSharedRef<SWidget> SBindingsPanel::CreateDrawerDockButton()
 {
-	if(bIsDrawerTab)
+	if (bIsDrawerTab)
 	{
 		return 
 			SNew(SButton)
 			.ButtonStyle(FAppStyle::Get(), "SimpleButton")
 			.ToolTipText(LOCTEXT("DockInLayout_Tooltip", "Docks MVVM drawer in tab."))
 			.ContentPadding(FMargin(1, 0))
-			.OnClicked(this, &SMVVMViewBindingPanel::CreateDrawerDockButtonClicked)
+			.OnClicked(this, &SBindingsPanel::CreateDrawerDockButtonClicked)
 			[
 				SNew(SHorizontalBox)
 				+ SHorizontalBox::Slot()
@@ -274,7 +227,7 @@ TSharedRef<SWidget> SMVVMViewBindingPanel::CreateDrawerDockButton()
 	return SNullWidget::NullWidget;
 }
 
-FReply SMVVMViewBindingPanel::CreateDrawerDockButtonClicked()
+FReply SBindingsPanel::CreateDrawerDockButtonClicked()
 {
 	if (TSharedPtr<FWidgetBlueprintEditor> WidgetEditor = WeakBlueprintEditor.Pin())
 	{
@@ -289,7 +242,7 @@ FReply SMVVMViewBindingPanel::CreateDrawerDockButtonClicked()
 	return FReply::Handled();
 }
 
-void SMVVMViewBindingPanel::HandleExtensionAdded(UBlueprintExtension* NewExtension)
+void SBindingsPanel::HandleExtensionAdded(UBlueprintExtension* NewExtension)
 {
 	if (UMVVMWidgetBlueprintExtension_View* MVVMExtensionPtr = Cast<UMVVMWidgetBlueprintExtension_View>(NewExtension))
 	{
@@ -303,7 +256,7 @@ void SMVVMViewBindingPanel::HandleExtensionAdded(UBlueprintExtension* NewExtensi
 
 			if (!BlueprintViewChangedDelegateHandle.IsValid())
 			{
-				BlueprintViewChangedDelegateHandle = MVVMExtensionPtr->OnBlueprintViewChangedDelegate().AddSP(this, &SMVVMViewBindingPanel::HandleBlueprintViewChangedDelegate);
+				BlueprintViewChangedDelegateHandle = MVVMExtensionPtr->OnBlueprintViewChangedDelegate().AddSP(this, &SBindingsPanel::HandleBlueprintViewChangedDelegate);
 			}
 
 			if (MVVMExtensionPtr->GetBlueprintView() == nullptr)
@@ -316,7 +269,7 @@ void SMVVMViewBindingPanel::HandleExtensionAdded(UBlueprintExtension* NewExtensi
 	}
 }
 
-TSharedRef<SWidget> SMVVMViewBindingPanel::GenerateSettingsMenu()
+TSharedRef<SWidget> SBindingsPanel::GenerateSettingsMenu()
 {
 	UMVVMWidgetBlueprintExtension_View* MVVMExtensionPtr = MVVMExtension.Get();
 	DetailsView->SetObject(MVVMExtensionPtr && MVVMExtensionPtr->GetBlueprintView() ? MVVMExtensionPtr->GetBlueprintView() : nullptr);
@@ -330,13 +283,12 @@ TSharedRef<SWidget> SMVVMViewBindingPanel::GenerateSettingsMenu()
 }
 
 
-void SMVVMViewBindingPanel::RegisterSettingsMenu()
+void SBindingsPanel::RegisterSettingsMenu()
 {
 	UToolMenu* Menu = UToolMenus::Get()->RegisterMenu("MVVMEditor.Panel.Toolbar.Settings");
 }
 
-
-TSharedRef<SWidget> SMVVMViewBindingPanel::GenerateCreateViewWidget()
+TSharedRef<SWidget> SBindingsPanel::GenerateCreateViewWidget()
 {
 	return SNew(SBorder)
 		[
@@ -347,50 +299,32 @@ TSharedRef<SWidget> SMVVMViewBindingPanel::GenerateCreateViewWidget()
 			.VAlign(EVerticalAlignment::VAlign_Center)
 			[
 				SNew(SButton)
-				.OnClicked(this, &SMVVMViewBindingPanel::HandleCreateViewClicked)
+				.OnClicked(this, &SBindingsPanel::HandleCreateViewClicked)
 				.Text(LOCTEXT("AddViewModelButtonText", "Add View Model"))
 			]
 		];
 }
 
-
-TSharedRef<SWidget> SMVVMViewBindingPanel::GenerateEditViewWidget()
+TSharedRef<SWidget> SBindingsPanel::GenerateEditViewWidget()
 {
 	FSlimHorizontalToolBarBuilder ToolbarBuilderGlobal(TSharedPtr<const FUICommandList>(), FMultiBoxCustomization::None);
 	{
 		ToolbarBuilderGlobal.BeginSection("Binding");
 		ToolbarBuilderGlobal.AddToolBarButton(
 			FUIAction(
-				FExecuteAction::CreateSP(this, &SMVVMViewBindingPanel::AddDefaultBinding),
-				FCanExecuteAction::CreateSP(this, &SMVVMViewBindingPanel::CanAddBinding),
+				FExecuteAction::CreateSP(this, &SBindingsPanel::AddDefaultBinding),
+				FCanExecuteAction::CreateSP(this, &SBindingsPanel::CanAddBinding),
 				FGetActionCheckState()
 			),
 			NAME_None,
 			LOCTEXT("AddBinding", "Add Binding"),
-			TAttribute<FText>(this, &SMVVMViewBindingPanel::GetAddBindingToolTip),
+			TAttribute<FText>(this, &SBindingsPanel::GetAddBindingToolTip),
 			FSlateIcon(FMVVMEditorStyle::Get().GetStyleSetName(), "BindingView.AddBinding"),
 			EUserInterfaceActionType::Button
 		);
 		ToolbarBuilderGlobal.EndSection();
 	}
 
-	{
-		ToolbarBuilderGlobal.BeginSection("ManageViewModels");
-		ToolbarBuilderGlobal.AddToolBarButton(
-			FUIAction(
-				FExecuteAction::CreateSP(this, &SMVVMViewBindingPanel::ShowManageViewModelsWindow),
-				FCanExecuteAction(),
-				FGetActionCheckState()
-			),
-			NAME_None,
-			LOCTEXT("ManageViewModels", "Manage ViewModels"),
-			LOCTEXT("ManageViewModelsTooltip", "Manage viewmodels for this widget."),
-			FSlateIcon(FMVVMEditorStyle::Get().GetStyleSetName(), "BindingView.ManageViewModels"),
-			EUserInterfaceActionType::Button
-		);
-		ToolbarBuilderGlobal.EndSection();
-	}
-
 	{
 		ToolbarBuilderGlobal.AddWidget(SNew(SSpacer), NAME_None, true, HAlign_Right);
 	}
@@ -403,7 +337,7 @@ TSharedRef<SWidget> SMVVMViewBindingPanel::GenerateEditViewWidget()
 				FCanExecuteAction(),
 				FGetActionCheckState()
 			),
-			FOnGetContent::CreateSP(this, &SMVVMViewBindingPanel::GenerateSettingsMenu),
+			FOnGetContent::CreateSP(this, &SBindingsPanel::GenerateSettingsMenu),
 			LOCTEXT("Settings", "Settings"),
 			LOCTEXT("SettingsTooltip", "ModelView Settings"),
 			FSlateIcon(FAppStyle::GetAppStyleSetName(), "LevelEditor.GameSettings"),
@@ -434,7 +368,7 @@ TSharedRef<SWidget> SMVVMViewBindingPanel::GenerateEditViewWidget()
 				+ SSplitter::Slot()
 				.Value(0.75f)
 				[
-					SAssignNew(ListView, SMVVMViewBindingListView, SharedThis(this), MVVMExtension.Get())
+					SAssignNew(BindingsList, SBindingsList, StaticCastSharedRef<SBindingsPanel>(AsShared()), MVVMExtension.Get())
 				]
 
 				+ SSplitter::Slot()
@@ -450,7 +384,7 @@ TSharedRef<SWidget> SMVVMViewBindingPanel::GenerateEditViewWidget()
 }
 
 
-FReply SMVVMViewBindingPanel::HandleCreateViewClicked()
+FReply SBindingsPanel::HandleCreateViewClicked()
 {
 	if (TSharedPtr<FWidgetBlueprintEditor> WidgetBlueprintEditor = WeakBlueprintEditor.Pin())
 	{
@@ -464,7 +398,7 @@ FReply SMVVMViewBindingPanel::HandleCreateViewClicked()
 
 		if (!BlueprintViewChangedDelegateHandle.IsValid())
 		{
-			BlueprintViewChangedDelegateHandle = MVVMExtensionPtr->OnBlueprintViewChangedDelegate().AddSP(this, &SMVVMViewBindingPanel::HandleBlueprintViewChangedDelegate);
+			BlueprintViewChangedDelegateHandle = MVVMExtensionPtr->OnBlueprintViewChangedDelegate().AddSP(this, &SBindingsPanel::HandleBlueprintViewChangedDelegate);
 		}
 
 		if (MVVMExtensionPtr->GetBlueprintView() == nullptr)
@@ -480,4 +414,44 @@ FReply SMVVMViewBindingPanel::HandleCreateViewClicked()
 	return FReply::Handled();
 }
 
-#undef LOCTEXT_NAMESPACE
\ No newline at end of file
+void SBindingsPanel::ShowManageViewModelsWindow()
+{
+	if (UMVVMWidgetBlueprintExtension_View* MVVMExtensionPtr = MVVMExtension.Get())
+	{
+		if (UMVVMBlueprintView* BlueprintView = MVVMExtensionPtr->GetBlueprintView())
+		{
+			TWeakObjectPtr<UMVVMWidgetBlueprintExtension_View> LocalMVVMExtension = MVVMExtension;
+			FOnViewModelContextsPicked ViewModelContextsPickedDelegate = FOnViewModelContextsPicked::CreateLambda
+			(
+				[LocalMVVMExtension](TArray<FMVVMBlueprintViewModelContext> UpdatedContexts)
+				{
+					UMVVMWidgetBlueprintExtension_View* MVVMExtensionPtr = LocalMVVMExtension.Get();
+					if (ensure(MVVMExtensionPtr))
+					{
+						if (UMVVMBlueprintView* BlueprintView = MVVMExtensionPtr->GetBlueprintView())
+						{
+							BlueprintView->SetViewModels(UpdatedContexts);
+						}
+					}
+				}
+			);
+
+			TSharedRef<SWindow> ManageViewModelsWindow = SNew(SWindow)
+				.Title(LOCTEXT("MVVMManageViewMoedlsWindowHeader", "Manage ViewModels"))
+				.SupportsMaximize(false)
+				.ClientSize(FVector2D(800.0f, 600.0f));
+
+			TSharedRef<SMVVMManageViewModelsWidget> ManageViewModelsWidget = SNew(SMVVMManageViewModelsWidget)
+				.ParentWindow(ManageViewModelsWindow)
+				.OnViewModelContextsPickedDelegate(ViewModelContextsPickedDelegate)
+				.WidgetBlueprint(WeakBlueprintEditor.Pin()->GetWidgetBlueprintObj());
+
+			ManageViewModelsWindow->SetContent(ManageViewModelsWidget);
+			GEditor->EditorAddModalWindow(ManageViewModelsWindow);
+		}
+	}
+}
+
+} // namespace UE::MVVM
+
+#undef LOCTEXT_NAMESPACE
diff --git a/Engine/Plugins/Runtime/ModelViewViewModel/Source/ModelViewViewModelEditor/Private/Widgets/SMVVMViewBindingPanel.h b/Engine/Plugins/Runtime/ModelViewViewModel/Source/ModelViewViewModelEditor/Private/Widgets/SMVVMViewBindingPanel.h
--- a/Engine/Plugins/Runtime/ModelViewViewModel/Source/ModelViewViewModelEditor/Private/Widgets/SMVVMViewBindingPanel.h
+++ b/Engine/Plugins/Runtime/ModelViewViewModel/Source/ModelViewViewModelEditor/Private/Widgets/SMVVMViewBindingPanel.h
@@ -9,29 +9,33 @@
 #include "Layout/Visibility.h"
 #include "Input/Reply.h"
 
+struct FMVVMBlueprintViewBinding;
 class FWidgetBlueprintEditor;
 class IDetailsView;
 class IStructureDetailsView;
 class SBorder;
-class SMVVMViewBindingListView;
 class UBlueprintExtension;
 class UMVVMWidgetBlueprintExtension_View;
 
-class SMVVMViewBindingPanel : public SCompoundWidget
+namespace UE::MVVM
+{
+class SBindingsList;
+
+class SBindingsPanel : public SCompoundWidget
 {
 public:
-	SLATE_BEGIN_ARGS(SMVVMViewBindingPanel) {}
+	SLATE_BEGIN_ARGS(SBindingsPanel) {}
 	SLATE_END_ARGS()
 
 	void Construct(const FArguments& InArgs, TSharedPtr<FWidgetBlueprintEditor>, bool bInIsDrawerTab);
-	virtual ~SMVVMViewBindingPanel();
+	virtual ~SBindingsPanel();
 
 	//~ Begin SWidget Interface
 	virtual FReply OnKeyDown(const FGeometry& MyGeometry, const FKeyEvent& InKeyEvent) override;
 	virtual bool SupportsKeyboardFocus() const override;
 	//~ End SWidget Interface
 
-	void OnBindingListSelectionChanged(int32 Index);
+	void OnBindingListSelectionChanged(TConstArrayView<FMVVMBlueprintViewBinding*> Selection);
 
 	static void RegisterSettingsMenu();
 
@@ -49,16 +53,15 @@ private:
 
 	TSharedRef<SWidget> GenerateSettingsMenu();
 
-	void ShowManageViewModelsWindow();
-
 	TSharedRef<SWidget> CreateDrawerDockButton();
 	FReply CreateDrawerDockButtonClicked();
 
 	void HandleExtensionAdded(UBlueprintExtension* NewExtension);
+	void ShowManageViewModelsWindow();
 
 private:
 	TWeakPtr<FWidgetBlueprintEditor> WeakBlueprintEditor;
-	TSharedPtr<SMVVMViewBindingListView> ListView;
+	TSharedPtr<SBindingsList> BindingsList;
 	TSharedPtr<SBorder> DetailContainer;
 	TSharedPtr<IDetailsView> DetailsView;
 	TSharedPtr<IStructureDetailsView> StructDetailsView;
@@ -66,3 +69,5 @@ private:
 	FDelegateHandle BlueprintViewChangedDelegateHandle;
 	bool bIsDrawerTab;
 };
+
+} // namespace UE::MVVM
