MVVM - Added InitialSource to SSourceBindingsList.

#rnx

diff --git a/Engine/Plugins/Runtime/ModelViewViewModel/Source/ModelViewViewModelEditor/Private/Widgets/SMVVMViewModelBindingListWidget.cpp b/Engine/Plugins/Runtime/ModelViewViewModel/Source/ModelViewViewModelEditor/Private/Widgets/SMVVMViewModelBindingListWidget.cpp
--- a/Engine/Plugins/Runtime/ModelViewViewModel/Source/ModelViewViewModelEditor/Private/Widgets/SMVVMViewModelBindingListWidget.cpp
+++ b/Engine/Plugins/Runtime/ModelViewViewModel/Source/ModelViewViewModelEditor/Private/Widgets/SMVVMViewModelBindingListWidget.cpp
@@ -130,6 +130,11 @@ void SSourceBindingList::Construct(const FArguments& InArgs, const UWidgetBluepr
 		.OnSelectionChanged(this, &SSourceBindingList::HandleSelectionChanged)
 		.OnDoubleClicked(this, &SSourceBindingList::HandleDoubleClicked)
 	];
+
+	if (InArgs._InitialSource.IsValid())
+	{
+		AddSources(MakeArrayView(&InArgs._InitialSource, 1));
+	}
 }
 
 void SSourceBindingList::Clear()
@@ -192,6 +197,11 @@ void SSourceBindingList::AddViewModels(TArrayView<const FMVVMBlueprintViewModelC
 	AddSources(NewSources);
 }
 
+void SSourceBindingList::AddSource(const FBindingSource& InSource)
+{
+	AddSources(MakeArrayView(&InSource, 1));
+}
+
 void SSourceBindingList::AddSources(TArrayView<const FBindingSource> InSources)
 {
 	Algo::Transform(InSources, Sources, [](const FBindingSource& Source)
diff --git a/Engine/Plugins/Runtime/ModelViewViewModel/Source/ModelViewViewModelEditor/Private/Widgets/SMVVMViewModelBindingListWidget.h b/Engine/Plugins/Runtime/ModelViewViewModel/Source/ModelViewViewModelEditor/Private/Widgets/SMVVMViewModelBindingListWidget.h
--- a/Engine/Plugins/Runtime/ModelViewViewModel/Source/ModelViewViewModelEditor/Private/Widgets/SMVVMViewModelBindingListWidget.h
+++ b/Engine/Plugins/Runtime/ModelViewViewModel/Source/ModelViewViewModelEditor/Private/Widgets/SMVVMViewModelBindingListWidget.h
@@ -7,6 +7,7 @@
 #include "Misc/EnumClassFlags.h"
 #include "MVVMPropertyPath.h"
 #include "Templates/SubclassOf.h"
+#include "Types/MVVMBindingSource.h"
 #include "UObject/WeakObjectPtr.h"
 #include "Widgets/SCompoundWidget.h"
 #include "Widgets/PropertyViewer/SPropertyViewer.h"
@@ -17,8 +18,6 @@ class UWidgetBlueprint;
 
 namespace UE::MVVM
 {
-	struct FBindingSource;
-
 	enum class EFieldVisibility : uint8
 	{
 		None = 0,
@@ -54,6 +53,7 @@ namespace UE::MVVM
 		DECLARE_DELEGATE_OneParam(FOnDoubleClicked, const FMVVMBlueprintPropertyPath&);
 	
 		SLATE_BEGIN_ARGS(SSourceBindingList) {}
+			SLATE_ARGUMENT(FBindingSource, InitialSource)
 			SLATE_ARGUMENT_DEFAULT(bool, ShowSearchBox) = false;
 			SLATE_ARGUMENT_DEFAULT(EFieldVisibility, FieldVisibilityFlags) = EFieldVisibility::All;
 			SLATE_EVENT(FOnDoubleClicked, OnDoubleClicked)
@@ -64,6 +64,7 @@ namespace UE::MVVM
 		void Clear();
 
 		void AddSource(UClass* Class, FName Name, FGuid Guid);
+		void AddSource(const FBindingSource& InSource);
 		void AddSources(TArrayView<const FBindingSource> InSources);
 
 		void AddWidgetBlueprint(const UWidgetBlueprint* WidgetBlueprint);
