MVVM - SSourceSelector now displays the widget hierarchy view for widgets.

#rnx

diff --git a/Engine/Plugins/Runtime/ModelViewViewModel/Source/ModelViewViewModelEditor/Private/Widgets/SMVVMSourceSelector.cpp b/Engine/Plugins/Runtime/ModelViewViewModel/Source/ModelViewViewModelEditor/Private/Widgets/SMVVMSourceSelector.cpp
--- a/Engine/Plugins/Runtime/ModelViewViewModel/Source/ModelViewViewModelEditor/Private/Widgets/SMVVMSourceSelector.cpp
+++ b/Engine/Plugins/Runtime/ModelViewViewModel/Source/ModelViewViewModelEditor/Private/Widgets/SMVVMSourceSelector.cpp
@@ -3,7 +3,10 @@
 #include "SMVVMSourceSelector.h"
 
 #include "Algo/Transform.h"
+#include "Editor.h"
+#include "MVVMEditorSubsystem.h"
 #include "SSimpleButton.h"
+#include "Widgets/Layout/SBox.h"
 #include "Widgets/SMVVMFieldIcon.h"
 #include "Widgets/SMVVMSourceEntry.h"
 
@@ -12,74 +15,185 @@
 namespace UE::MVVM
 {
 
-void SSourceSelector::Construct(const FArguments& Args)
+void SSourceSelector::Construct(const FArguments& Args, const UWidgetBlueprint* InWidgetBlueprint)
 {
 	TextStyle = Args._TextStyle;
-	AvailableSourcesAttribute = Args._AvailableSources;
-	check(AvailableSourcesAttribute.IsSet());
 	SelectedSourceAttribute = Args._SelectedSource;
 	check(SelectedSourceAttribute.IsSet());
 
+	WidgetBlueprint = InWidgetBlueprint;
+	check(InWidgetBlueprint);
+
+	bAutoRefresh = Args._AutoRefresh;
+	bViewModels = Args._ViewModels;
+	bShowClear = Args._ShowClear;
+	
 	Refresh();
 
-	ChildSlot
-	[
-		SNew(SHorizontalBox)
+	TSharedRef<SHorizontalBox> HorizontalBox = SNew(SHorizontalBox)
 		+ SHorizontalBox::Slot()
-		.HAlign(HAlign_Fill)
-		.VAlign(VAlign_Center)
 		[
-			SAssignNew(SourceComboBox, SComboBox<FBindingSource>)
-			.OptionsSource(&AvailableSources)
-			.OnGenerateWidget_Lambda([this](FBindingSource Source)
-			{
-				return SNew(SSourceEntry)
-					.TextStyle(TextStyle)
-					.Source(Source);
-			})
-			.OnSelectionChanged(this, &SSourceSelector::OnComboBoxSelectionChanged)
-			.InitiallySelectedItem(SelectedSource)
+			SAssignNew(MenuAnchor, SMenuAnchor)
+			.OnGetMenuContent(this, &SSourceSelector::OnGetMenuContent)
 			[
-				SAssignNew(SelectedSourceWidget, SSourceEntry)
-				.TextStyle(TextStyle)
-				.Source(SelectedSource)
+				SNew(SButton)
+				.ButtonStyle(&FAppStyle::Get().GetWidgetStyle<FComboButtonStyle>("ComboButton").ButtonStyle)
+				.OnClicked_Lambda([this]() 
+				{
+					MenuAnchor->SetIsOpen(true); 
+					return FReply::Handled(); 
+				})
+				[
+					SNew(SHorizontalBox)
+					+ SHorizontalBox::Slot()
+					[
+						SAssignNew(SelectedSourceWidget, SSourceEntry)
+						.Source(SelectedSource)
+						.TextStyle(TextStyle)
+					]
+					+ SHorizontalBox::Slot()
+					.Padding(2, 0)
+					.HAlign(HAlign_Right)
+					.VAlign(VAlign_Center)
+					.AutoWidth()
+					[
+						SNew(SImage)
+						.Image(FAppStyle::Get().GetBrush("Icons.ChevronDown"))
+					]
+				]
 			]
-		]
-		+ SHorizontalBox::Slot()
-		.HAlign(HAlign_Right)
-		.VAlign(VAlign_Center)
-		.AutoWidth()
-		[
-			SNew(SSimpleButton)
-			.Icon(FAppStyle::Get().GetBrush("Icons.X"))
-			.ToolTipText(LOCTEXT("ClearField", "Clear source selection."))
-			.Visibility(this, &SSourceSelector::GetClearVisibility)
-			.OnClicked(this, &SSourceSelector::OnClearSource)
-		]
+		];
+
+	if (bShowClear)
+	{
+		HorizontalBox->AddSlot()
+			.HAlign(HAlign_Right)
+			.VAlign(VAlign_Center)
+			.AutoWidth()
+			[
+				SNew(SSimpleButton)
+				.Icon(FAppStyle::Get().GetBrush("Icons.X"))
+				.ToolTipText(LOCTEXT("ClearField", "Clear source selection."))
+				.Visibility(this, &SSourceSelector::GetClearVisibility)
+				.OnClicked(this, &SSourceSelector::OnClearSource)
+			];
+	}
+
+	ChildSlot
+	[
+		HorizontalBox
 	];
 
-	OnSelectionChanged = Args._OnSelectionChanged;
+	OnSelectionChangedDelegate = Args._OnSelectionChanged;
+}
+
+TSharedRef<SWidget> SSourceSelector::OnGetMenuContent()
+{
+	TSharedRef<SBox> TopLevelBox = SNew(SBox)
+		.MinDesiredWidth(200)
+		.MinDesiredHeight(400);
+
+	if (bViewModels)
+	{
+		ViewModelList = SNew(SListView<FBindingSource>)
+			.OnGenerateRow_Lambda([this](FBindingSource Source, TSharedRef<STableViewBase> OwnerTable)
+				{
+					return SNew(STableRow<FBindingSource>, OwnerTable)
+					[
+						SNew(SSourceEntry)
+						.Source(Source)
+						.TextStyle(TextStyle)
+					];
+				})
+			.OnSelectionChanged(this, &SSourceSelector::OnViewModelSelectionChanged)
+			.SelectionMode(ESelectionMode::Single)
+			.ListItemsSource(&ViewModelSources);
+
+		TopLevelBox->SetContent(ViewModelList.ToSharedRef());
+	}
+	else
+	{
+		WidgetHierarchy = SNew(SReadOnlyHierarchyView, WidgetBlueprint.Get())
+			.OnSelectionChanged(this, &SSourceSelector::OnWidgetSelectionChanged)
+			.ShowSearch(false);
+
+		TopLevelBox->SetContent(WidgetHierarchy.ToSharedRef());
+	}
+
+	return TopLevelBox;
 }
 
-void SSourceSelector::OnComboBoxSelectionChanged(FBindingSource Selected, ESelectInfo::Type SelectionType)
+void SSourceSelector::OnViewModelSelectionChanged(FBindingSource Selected, ESelectInfo::Type SelectionType)
 {
 	SelectedSource = Selected;
 
-	SelectedSourceWidget->RefreshSource(Selected);
+	if (SelectedSourceWidget.IsValid())
+	{
+		SelectedSourceWidget->RefreshSource(Selected);
+	}
+
+	OnSelectionChangedDelegate.ExecuteIfBound(Selected);
+}
+
+void SSourceSelector::OnWidgetSelectionChanged(FName SelectedName, ESelectInfo::Type SelectionType)
+{
+	FBindingSource Source = FBindingSource::CreateForWidget(WidgetBlueprint.Get(), SelectedName);
+
+	if (SelectedSourceWidget.IsValid())
+	{
+		SelectedSourceWidget->RefreshSource(Source);
+	}
+
+	OnSelectionChangedDelegate.ExecuteIfBound(Source);
 
-	OnSelectionChanged.ExecuteIfBound(Selected);
 }
 
 void SSourceSelector::Refresh()
 {
 	SelectedSource = SelectedSourceAttribute.Get();
-	AvailableSources = AvailableSourcesAttribute.Get();
 
-	if (SourceComboBox.IsValid())
+	if (SelectedSourceWidget.IsValid())
+	{
+		SelectedSourceWidget->RefreshSource(SelectedSource);
+	}
+
+	TSharedPtr<SWidget> SourcePicker;
+	if (bViewModels)
 	{
-		SourceComboBox->RefreshOptions();
-		SourceComboBox->SetSelectedItem(SelectedSource);
+		UMVVMEditorSubsystem* EditorSubsystem = GEditor->GetEditorSubsystem<UMVVMEditorSubsystem>();
+		if (const UMVVMBlueprintView* View = EditorSubsystem->GetView(WidgetBlueprint.Get()))
+		{
+			const TArrayView<const FMVVMBlueprintViewModelContext> ViewModels = View->GetViewModels();
+			ViewModelSources.Reserve(ViewModels.Num());
+
+			for (const FMVVMBlueprintViewModelContext& ViewModel : ViewModels)
+			{
+				FBindingSource& Source = ViewModelSources.AddDefaulted_GetRef();
+				Source.Class = ViewModel.GetViewModelClass();
+				Source.Name = ViewModel.GetViewModelName();
+				Source.DisplayName = ViewModel.GetDisplayName();
+				Source.ViewModelId = ViewModel.GetViewModelId();
+			}
+		}
 	}
+
+	if (WidgetHierarchy.IsValid())
+	{
+		WidgetHierarchy->Refresh();
+	}
+
+	if (ViewModelList.IsValid())
+	{
+		if (SelectedSource.IsValid())
+		{
+			ViewModelList->SetItemSelection(SelectedSource, true);
+		}
+		else
+		{
+			ViewModelList->ClearSelection();
+		}
+	}
+
 }
 
 EVisibility SSourceSelector::GetClearVisibility() const
@@ -89,9 +203,9 @@ EVisibility SSourceSelector::GetClearVisibility() const
 
 FReply SSourceSelector::OnClearSource()
 {
-	if (SourceComboBox.IsValid())
+	if (ViewModelList.IsValid())
 	{
-		SourceComboBox->ClearSelection();
+		ViewModelList->ClearSelection();
 	}
 
 	return FReply::Handled();
diff --git a/Engine/Plugins/Runtime/ModelViewViewModel/Source/ModelViewViewModelEditor/Private/Widgets/SMVVMSourceSelector.h b/Engine/Plugins/Runtime/ModelViewViewModel/Source/ModelViewViewModelEditor/Private/Widgets/SMVVMSourceSelector.h
--- a/Engine/Plugins/Runtime/ModelViewViewModel/Source/ModelViewViewModelEditor/Private/Widgets/SMVVMSourceSelector.h
+++ b/Engine/Plugins/Runtime/ModelViewViewModel/Source/ModelViewViewModelEditor/Private/Widgets/SMVVMSourceSelector.h
@@ -2,11 +2,13 @@
 
 #pragma once
 
-#include "CoreMinimal.h"
+#include "Hierarchy/SReadOnlyHierarchyView.h"
+#include "Layout/Visibility.h"
+#include "Templates/SharedPointer.h"
 #include "Types/MVVMBindingSource.h"
 #include "Widgets/DeclarativeSyntaxSupport.h"
 #include "Widgets/SCompoundWidget.h"
-#include "Widgets/Input/SComboBox.h"
+#include "Widgets/Input/SMenuAnchor.h"
 
 namespace UE::MVVM
 {
@@ -23,29 +25,43 @@ public:
 		{
 		}
 		SLATE_STYLE_ARGUMENT(FTextBlockStyle, TextStyle)
+		SLATE_ARGUMENT_DEFAULT(bool, ShowClear) = true;
+		SLATE_ARGUMENT_DEFAULT(bool, AutoRefresh) = false;
+		SLATE_ARGUMENT_DEFAULT(bool, ViewModels) = false;
 		SLATE_ATTRIBUTE(FBindingSource, SelectedSource)
-		SLATE_ATTRIBUTE(TArray<FBindingSource>, AvailableSources)
 		SLATE_EVENT(FSelectionChanged, OnSelectionChanged)
 	SLATE_END_ARGS()
 
-	void Construct(const FArguments& InArgs);
+	void Construct(const FArguments& InArgs, const UWidgetBlueprint* InWidgetBlueprint);
 	void Refresh();
 
 private:
-	void OnComboBoxSelectionChanged(FBindingSource Selected, ESelectInfo::Type SelectionType);
+	void OnViewModelSelectionChanged(FBindingSource Selected, ESelectInfo::Type SelectionType);
+	void OnWidgetSelectionChanged(FName SelectedName, ESelectInfo::Type SelectionType);
+	TSharedRef<SWidget> OnGetMenuContent();
 
 	EVisibility GetClearVisibility() const;
 	FReply OnClearSource();
 
 private:
-	FSelectionChanged OnSelectionChanged;
+	TWeakObjectPtr<const UWidgetBlueprint> WidgetBlueprint;
+	FSelectionChanged OnSelectionChangedDelegate;
+	const FTextBlockStyle* TextStyle = nullptr;
+
 	TAttribute<TArray<FBindingSource>> AvailableSourcesAttribute;
 	TAttribute<FBindingSource> SelectedSourceAttribute;
-	TSharedPtr<SComboBox<FBindingSource>> SourceComboBox;
-	TArray<FBindingSource> AvailableSources;
+
+	TArray<FBindingSource> ViewModelSources;
 	FBindingSource SelectedSource;
+
+	TSharedPtr<SMenuAnchor> MenuAnchor;
+	TSharedPtr<SListView<FBindingSource>> ViewModelList;
+	TSharedPtr<SReadOnlyHierarchyView> WidgetHierarchy;
 	TSharedPtr<SSourceEntry> SelectedSourceWidget;
-	const FTextBlockStyle* TextStyle = nullptr;
+
+	bool bAutoRefresh = false;
+	bool bViewModels = false;
+	bool bShowClear = true;
 };
 
 } // namespace UE::MVVM
