Fix bad variable initialization introduced with CL. Fix the issue 194729
#rnx

diff --git a/Engine/Plugins/Runtime/ModelViewViewModel/Source/ModelViewViewModelEditor/Private/Widgets/SMVVMViewBindingListView.cpp b/Engine/Plugins/Runtime/ModelViewViewModel/Source/ModelViewViewModelEditor/Private/Widgets/SMVVMViewBindingListView.cpp
--- a/Engine/Plugins/Runtime/ModelViewViewModel/Source/ModelViewViewModelEditor/Private/Widgets/SMVVMViewBindingListView.cpp
+++ b/Engine/Plugins/Runtime/ModelViewViewModel/Source/ModelViewViewModelEditor/Private/Widgets/SMVVMViewBindingListView.cpp
@@ -485,7 +485,7 @@ private:
 
 		UWidgetBlueprint* WidgetBlueprint = WidgetBlueprintWeak.Get();
 
-		EFieldVisibility Flags;
+		EFieldVisibility Flags = EFieldVisibility::None;
 		if (UE::MVVM::IsForwardBinding(ViewBinding->BindingType))
 		{
 			Flags |= EFieldVisibility::Writable;
