MVVM: Add the rename and delete option to the viewmodel panel.

diff --git a/Engine/Plugins/Runtime/ModelViewViewModel/Source/ModelViewViewModelBlueprint/Public/MVVMBlueprintViewModelContext.h b/Engine/Plugins/Runtime/ModelViewViewModel/Source/ModelViewViewModelBlueprint/Public/MVVMBlueprintViewModelContext.h
--- a/Engine/Plugins/Runtime/ModelViewViewModel/Source/ModelViewViewModelBlueprint/Public/MVVMBlueprintViewModelContext.h
+++ b/Engine/Plugins/Runtime/ModelViewViewModel/Source/ModelViewViewModelBlueprint/Public/MVVMBlueprintViewModelContext.h
@@ -76,10 +76,10 @@ public:
 
 private:
 	/** When the view is spawn, create an instance of the viewmodel. */
-	UPROPERTY(VisibleAnywhere, Category = "MVVM")
+	UPROPERTY(VisibleAnywhere, AdvancedDisplay, Category = "MVVM")
 	FGuid ViewModelContextId;
 
-	UPROPERTY(EditAnywhere, Category = "MVVM")
+	UPROPERTY(VisibleAnywhere, Category = "MVVM", NoClear, meta = (AllowedClasses = "/Script/UMG.NotifyFieldValueChanged", DisallowedClasses = "/Script/UMG.Widget"))
 	UClass* NotifyFieldValueClass = nullptr;
 
 	UPROPERTY()
diff --git a/Engine/Plugins/Runtime/ModelViewViewModel/Source/ModelViewViewModelEditor/Private/Widgets/SMVVMViewModelPanel.cpp b/Engine/Plugins/Runtime/ModelViewViewModel/Source/ModelViewViewModelEditor/Private/Widgets/SMVVMViewModelPanel.cpp
--- a/Engine/Plugins/Runtime/ModelViewViewModel/Source/ModelViewViewModelEditor/Private/Widgets/SMVVMViewModelPanel.cpp
+++ b/Engine/Plugins/Runtime/ModelViewViewModel/Source/ModelViewViewModelEditor/Private/Widgets/SMVVMViewModelPanel.cpp
@@ -5,26 +5,36 @@
 #include "Bindings/MVVMBindingHelper.h"
 #include "Editor.h"
 #include "Editor/EditorEngine.h"
+#include "IStructureDetailsView.h"
 #include "Kismet2/Kismet2NameValidators.h"
 #include "MVVMBlueprintView.h"
 #include "MVVMBlueprintViewModelContext.h"
 #include "MVVMEditorSubsystem.h"
 #include "MVVMSubsystem.h"
 #include "MVVMViewModelBase.h"
+#include "PropertyEditorModule.h"
 #include "WidgetBlueprint.h"
 #include "WidgetBlueprintEditor.h"
 
+#include "Framework/Commands/GenericCommands.h"
+#include "Framework/Commands/UICommandList.h"
+#include "Framework/MultiBox/MultiBoxBuilder.h"
+
 #include "Styling/AppStyle.h"
 #include "Styling/SlateIconFinder.h"
 #include "Styling/MVVMEditorStyle.h"
 
 #include "Widgets/SBoxPanel.h"
+#include "Widgets/Images/SImage.h"
 #include "Widgets/Input/SSearchBox.h"
 #include "Widgets/Layout/SBorder.h"
 #include "Widgets/Layout/SBox.h"
+#include "Widgets/PropertyViewer/SFieldIcon.h"
+#include "Widgets/PropertyViewer/SFieldName.h"
 #include "Widgets/PropertyViewer/SPropertyViewer.h"
 #include "Widgets/SMVVMSelectViewModel.h"
-
+#include "Widgets/Text/SInlineEditableTextBlock.h"
+#include "Widgets/Text/STextBlock.h"
 #include "SPositiveActionButton.h"
 
 #define LOCTEXT_NAMESPACE "ViewModelPanel"
@@ -48,31 +58,58 @@ void SMVVMViewModelPanel::Construct(const FArguments& InArgs, TSharedPtr<FWidget
 		ViewModelsUpdatedHandle = CurrentBlueprintView->OnViewModelsUpdated.AddSP(this, &SMVVMViewModelPanel::HandleViewModelsUpdated);
 	}
 
+	CreateCommandList();
+
 	ViewModelTreeView = SNew(UE::PropertyViewer::SPropertyViewer)
 		.PropertyVisibility(UE::PropertyViewer::SPropertyViewer::EPropertyVisibility::Visible)
 		.bShowSearchBox(true)
 		.bShowFieldIcon(true)
 		.bSanitizeName(true)
 		.FieldIterator(FieldIterator.Get())
+		.OnContextMenuOpening(this, &SMVVMViewModelPanel::HandleContextMenuOpening)
+		.OnSelectionChanged(this, &SMVVMViewModelPanel::HandleSelectionChanged)
+		.OnGenerateContainer(this, &SMVVMViewModelPanel::HandleGenerateContainer)
 		.SearchBoxPreSlot()
 		[
 			SAssignNew(AddMenuButton, SPositiveActionButton)
 			.OnGetMenuContent(this, &SMVVMViewModelPanel::MakeAddMenu)
 			.Text(LOCTEXT("Viewmodel", "Viewmodel"))
 			.IsEnabled(this, &SMVVMViewModelPanel::HandleCanEditViewmodelList)
-		]
-		;
-		//.OnRenamed(this, &SMVVMViewModelPanel::HandleViewModelRenamed);
+		];
 
 	FillViewModel();
 
+
+	FPropertyEditorModule& PropertyEditor = FModuleManager::Get().GetModuleChecked<FPropertyEditorModule>("PropertyEditor");
+	FDetailsViewArgs DetailsViewArgs;
+	DetailsViewArgs.NameAreaSettings = FDetailsViewArgs::HideNameArea;
+	DetailsViewArgs.NotifyHook = this;
+	DetailsViewArgs.bHideSelectionTip = true;
+	DetailsViewArgs.DefaultsOnlyVisibility = EEditDefaultsOnlyNodeVisibility::Hide;
+	DetailsViewArgs.bAllowSearch = false;
+	FStructureDetailsViewArgs StructDetailViewArgs;
+	StructDetailViewArgs.bShowObjects = true;
+	StructDetailViewArgs.bShowInterfaces = true;
+
+	PropertyView = PropertyEditor.CreateStructureDetailView(DetailsViewArgs, StructDetailViewArgs, TSharedPtr<class FStructOnScope>());
+
 	ChildSlot
 	[
 		SNew(SBorder)
 		.BorderImage(FAppStyle::Get().GetBrush("Brushes.Recessed"))
 		.Padding(0)
 		[
-			ViewModelTreeView.ToSharedRef()
+			SNew(SSplitter)
+			.Orientation(EOrientation::Orient_Vertical)
+			+ SSplitter::Slot()
+			[
+				ViewModelTreeView.ToSharedRef()
+			]
+			+ SSplitter::Slot()
+			.SizeRule(SSplitter::SizeToContent)
+			[
+				PropertyView->GetWidget().ToSharedRef()
+			]
 		]
 	];
 }
@@ -99,6 +136,10 @@ SMVVMViewModelPanel::~SMVVMViewModelPanel()
 
 void SMVVMViewModelPanel::FillViewModel()
 {
+	PropertyViewerHandles.Reset();
+	EditableTextBlocks.Reset();
+	SelectedViewModelGuid = FGuid();
+
 	if (ViewModelTreeView)
 	{
 		ViewModelTreeView->RemoveAll();
@@ -107,22 +148,25 @@ void SMVVMViewModelPanel::FillViewModel()
 		{
 			for (const FMVVMBlueprintViewModelContext& ViewModelContext : View->GetViewModels())
 			{
-				if (ViewModelContext.GetViewModelClass())
-				{
-					UObject* ViewModelInstance = ViewModelContext.GetViewModelClass()->GetDefaultObject();
-					ViewModelTreeView->AddInstance(ViewModelInstance, ViewModelContext.GetDisplayName());
-				}
-				else
-				{
-					// Find a way to show context that are not valid anymore
-					//ViewModelHandle.Value = ViewModelTreeView->AddInstance(ViewModelInstance, ViewModelContext.GetDisplayName());
-				}
+				const UE::PropertyViewer::SPropertyViewer::FHandle Handle = ViewModelTreeView->AddContainer(ViewModelContext.GetViewModelClass());
+				PropertyViewerHandles.Add(Handle, ViewModelContext.GetViewModelId());
 			}
 		}
 	}
 }
 
 
+FReply SMVVMViewModelPanel::OnKeyDown(const FGeometry& MyGeometry, const FKeyEvent& InKeyEvent)
+{
+	if (CommandList->ProcessCommandBindings(InKeyEvent))
+	{
+		return FReply::Handled();
+	}
+
+	return FReply::Unhandled();
+}
+
+
 void SMVVMViewModelPanel::HandleViewUpdated(UBlueprintExtension*)
 {
 	bool bViewUpdated = false;
@@ -165,6 +209,7 @@ TSharedRef<SWidget> SMVVMViewModelPanel::MakeAddMenu()
 	{
 		WidgetBlueprint = EditorPin->GetWidgetBlueprintObj();
 	}
+
 	return SNew(SBox)
 		.WidthOverride(600)
 		.HeightOverride(500)
@@ -214,6 +259,308 @@ bool SMVVMViewModelPanel::HandleCanEditViewmodelList() const
 	return false;
 }
 
+
+TSharedRef<SWidget> SMVVMViewModelPanel::HandleGenerateContainer(UE::PropertyViewer::SPropertyViewer::FHandle ContainerHandle, TOptional<FText> DisplayName)
+{
+	if (const FGuid* VMGuidPtr = PropertyViewerHandles.Find(ContainerHandle))
+	{
+		const FGuid VMGuid = *VMGuidPtr;
+		if (UMVVMBlueprintView* View = WeakBlueprintView.Get())
+		{
+			if (const FMVVMBlueprintViewModelContext* ViewModelContext = View->FindViewModel(VMGuid))
+			{
+				if (ViewModelContext->GetViewModelClass())
+				{
+					TSharedRef<SInlineEditableTextBlock> EditableTextBlock = SNew(SInlineEditableTextBlock)
+						.Text(ViewModelContext->GetDisplayName())
+						.OnVerifyTextChanged(this, &SMVVMViewModelPanel::HandleVerifyNameTextChanged, VMGuid)
+						.OnTextCommitted(this, &SMVVMViewModelPanel::HandleNameTextCommited, VMGuid);
+					EditableTextBlocks.Add(VMGuid, EditableTextBlock);
+
+					return SNew(SHorizontalBox)
+						+ SHorizontalBox::Slot()
+						.AutoWidth()
+						.HAlign(HAlign_Right)
+						.VAlign(VAlign_Center)
+						[
+							SNew(UE::PropertyViewer::SFieldIcon, ViewModelContext->GetViewModelClass())
+						]
+						+ SHorizontalBox::Slot()
+						.Padding(4.0f)
+						[
+							EditableTextBlock
+						];
+				}
+				else
+				{
+					return SNew(SHorizontalBox)
+						+ SHorizontalBox::Slot()
+						.AutoWidth()
+						.HAlign(HAlign_Right)
+						.VAlign(VAlign_Center)
+						[
+							SNew(SImage)
+							.Image(FAppStyle::Get().GetBrush("Icons.ErrorWithColor"))
+						]
+						+ SHorizontalBox::Slot()
+						.Padding(4.0f)
+						[
+							SNew(STextBlock)
+							.Text(ViewModelContext->GetDisplayName())
+						];
+				}
+			}
+		}
+	}
+
+	return SNew(STextBlock)
+		.Text(LOCTEXT("ViewmodelError", "Error. Invalid viewmodel."));
+}
+
+
+bool SMVVMViewModelPanel::HandleVerifyNameTextChanged(const FText& InText, FText& OutErrorMessage, FGuid ViewModelGuid)
+{
+	return RenameViewModelProperty(ViewModelGuid, InText, false, OutErrorMessage);
+}
+
+
+void SMVVMViewModelPanel::HandleNameTextCommited(const FText& InText, ETextCommit::Type CommitInfo, FGuid ViewModelGuid)
+{
+	if (CommitInfo == ETextCommit::OnEnter)
+	{
+		FText OutErrorMessage;
+		RenameViewModelProperty(ViewModelGuid, InText, true, OutErrorMessage);
+	}
+}
+
+
+void SMVVMViewModelPanel::CreateCommandList()
+{
+	CommandList = MakeShared<FUICommandList>();
+
+	CommandList->MapAction(
+		FGenericCommands::Get().Delete,
+		FExecuteAction::CreateSP(this, &SMVVMViewModelPanel::HandleDeleteViewModel),
+		FCanExecuteAction::CreateSP(this, &SMVVMViewModelPanel::HandleCanDeleteViewModel)
+	);
+	CommandList->MapAction(
+		FGenericCommands::Get().Rename,
+		FExecuteAction::CreateSP(this, &SMVVMViewModelPanel::HandleRenameViewModel),
+		FCanExecuteAction::CreateSP(this, &SMVVMViewModelPanel::HandleCanRenameViewModel)
+	);
+}
+
+
+void SMVVMViewModelPanel::HandleDeleteViewModel()
+{
+	TSharedPtr<FWidgetBlueprintEditor> WidgetBlueprintEditor = WeakBlueprintEditor.Pin();
+	if (WidgetBlueprintEditor == nullptr)
+	{
+		return;
+	}
+
+	UWidgetBlueprint* WidgetBP = WidgetBlueprintEditor->GetWidgetBlueprintObj();
+	if (WidgetBP == nullptr)
+	{
+		return;
+	}
+
+	UMVVMBlueprintView* BlueprintView = WeakBlueprintView.Get();
+	if (BlueprintView == nullptr)
+	{
+		return;
+	}
+
+	TArray<UE::PropertyViewer::SPropertyViewer::FSelectedItem> Items = ViewModelTreeView->GetSelectedItems();
+	for (UE::PropertyViewer::SPropertyViewer::FSelectedItem& SelectedItem : Items)
+	{
+		if (SelectedItem.bIsContainerSelected)
+		{
+			if (FGuid* VMGuidPtr = PropertyViewerHandles.Find(SelectedItem.Handle))
+			{
+				if (const FMVVMBlueprintViewModelContext* ViewModelContext = BlueprintView->FindViewModel(*VMGuidPtr))
+				{
+					GEditor->GetEditorSubsystem<UMVVMEditorSubsystem>()->RemoveViewModel(WidgetBP, ViewModelContext->GetViewModelName());
+				}
+			}
+		}
+	}
+}
+
+
+bool SMVVMViewModelPanel::HandleCanDeleteViewModel() const
+{
+	if (TSharedPtr<FWidgetBlueprintEditor> BlueprintEditor = WeakBlueprintEditor.Pin())
+	{
+		return BlueprintEditor->InEditingMode();
+	}
+	return false;
+}
+
+
+void SMVVMViewModelPanel::HandleRenameViewModel()
+{
+	TArray<UE::PropertyViewer::SPropertyViewer::FSelectedItem> Items = ViewModelTreeView->GetSelectedItems();
+	if (Items.Num() == 1 && Items[0].bIsContainerSelected)
+	{
+		if (const FGuid* VMGuidPtr = PropertyViewerHandles.Find(Items[0].Handle))
+		{
+			if (TSharedPtr<SInlineEditableTextBlock>* TextBlockPtr = EditableTextBlocks.Find(*VMGuidPtr))
+			{
+				(*TextBlockPtr)->EnterEditingMode();
+			}
+		}
+	}
+}
+
+
+bool SMVVMViewModelPanel::HandleCanRenameViewModel() const
+{
+	if (TSharedPtr<FWidgetBlueprintEditor> BlueprintEditor = WeakBlueprintEditor.Pin())
+	{
+		return BlueprintEditor->InEditingMode();
+	}
+	return false;
+}
+
+
+TSharedPtr<SWidget> SMVVMViewModelPanel::HandleContextMenuOpening(UE::PropertyViewer::SPropertyViewer::FHandle ContainerHandle, TArrayView<const FFieldVariant> Fields) const
+{
+	if (Fields.Num() == 0 && ContainerHandle.IsValid() )
+	{
+		FMenuBuilder MenuBuilder(true, CommandList);
+
+		MenuBuilder.BeginSection("Edit", LOCTEXT("Edit", "Edit"));
+		{
+			MenuBuilder.AddMenuEntry(FGenericCommands::Get().Delete);
+			MenuBuilder.AddMenuEntry(FGenericCommands::Get().Rename);
+		}
+		MenuBuilder.EndSection();
+
+		return MenuBuilder.MakeWidget();
+	}
+	return TSharedPtr<SWidget>();
+}
+
+
+void SMVVMViewModelPanel::HandleSelectionChanged(UE::PropertyViewer::SPropertyViewer::FHandle ContainerHandle, TArrayView<const FFieldVariant> Fields, ESelectInfo::Type SelectionType)
+{
+	SelectedViewModelGuid = FGuid();
+
+	bool bSet = false;
+	if (Fields.Num() == 0)
+	{
+		if (const FGuid* VMGuidPtr = PropertyViewerHandles.Find(ContainerHandle))
+		{
+			SelectedViewModelGuid = *VMGuidPtr;
+			if (UMVVMBlueprintView* BlueprintView = WeakBlueprintView.Get())
+			{
+				if (FMVVMBlueprintViewModelContext* ViewModelContext = BlueprintView->FindViewModel(SelectedViewModelGuid))
+				{
+					PropertyView->SetStructureData(MakeShared<FStructOnScope>(FMVVMBlueprintViewModelContext::StaticStruct(), reinterpret_cast<uint8*>(ViewModelContext)));
+					bSet = true;
+				}
+			}
+		}
+	}
+	
+	if (!bSet)
+	{
+		PropertyView->SetStructureData(TSharedPtr<FStructOnScope>());
+	}
+}
+
+
+bool SMVVMViewModelPanel::RenameViewModelProperty(FGuid ViewModelGuid, const FText& RenameTo, bool bCommit, FText& OutErrorMessage) const
+{
+	if (RenameTo.IsEmptyOrWhitespace())
+	{
+		OutErrorMessage = LOCTEXT("EmptyViewModelName", "Empty viewmodel name");
+		return false;
+	}
+
+	const FString& NewNameString = RenameTo.ToString();
+	if (NewNameString.Len() >= NAME_SIZE)
+	{
+		OutErrorMessage = LOCTEXT("ViewModelNameTooLong", "Viewmodel Name is Too Long");
+		return false;
+	}
+
+	FString GeneratedName = SlugStringForValidName(NewNameString);
+	if (GeneratedName.IsEmpty())
+	{
+		OutErrorMessage = LOCTEXT("EmptyViewModelName", "Empty viewmodel name");
+		return false;
+	}
+
+	const FName GeneratedFName(*GeneratedName);
+	check(GeneratedFName.IsValidXName(INVALID_OBJECTNAME_CHARACTERS));
+
+	if (TSharedPtr<FWidgetBlueprintEditor> WidgetBlueprintEditor = WeakBlueprintEditor.Pin())
+	{
+		if (UWidgetBlueprint* WidgetBP = WidgetBlueprintEditor->GetWidgetBlueprintObj())
+		{
+			if (UMVVMBlueprintView* View = WeakBlueprintView.Get())
+			{
+				if (const FMVVMBlueprintViewModelContext* ViewModelContext = View->FindViewModel(ViewModelGuid))
+				{
+					if (bCommit)
+					{
+						return GEditor->GetEditorSubsystem<UMVVMEditorSubsystem>()->RenameViewModel(WidgetBP, ViewModelContext->GetViewModelName(), *NewNameString, OutErrorMessage);
+					}
+					else
+					{
+						return GEditor->GetEditorSubsystem<UMVVMEditorSubsystem>()->VerifyViewModelRename(WidgetBP, ViewModelContext->GetViewModelName(), *NewNameString, OutErrorMessage);
+					}
+				}
+			}
+		}
+	}
+	return false;
+}
+
+
+namespace Private
+{
+	static FName NAME_ViewModelName = TEXT("ViewModelName");
+}
+
+void SMVVMViewModelPanel::NotifyPreChange(FProperty* PropertyAboutToChange)
+{
+	PreviousViewModelPropertyName = FName();
+	if (PropertyAboutToChange && PropertyAboutToChange->GetFName() == Private::NAME_ViewModelName && SelectedViewModelGuid.IsValid())
+	{
+		if (UMVVMBlueprintView* BlueprintView = WeakBlueprintView.Get())
+		{
+			if (const FMVVMBlueprintViewModelContext* ViewModelContext = BlueprintView->FindViewModel(SelectedViewModelGuid))
+			{
+				PreviousViewModelPropertyName = ViewModelContext->GetViewModelName();
+			}
+		}
+	}
+}
+
+
+void SMVVMViewModelPanel::NotifyPostChange(const FPropertyChangedEvent& PropertyChangedEvent, FProperty* PropertyThatChanged)
+{
+	if (PropertyThatChanged && PropertyThatChanged->GetFName() == Private::NAME_ViewModelName && SelectedViewModelGuid.IsValid())
+	{
+		UMVVMBlueprintView* BlueprintView = WeakBlueprintView.Get();
+		TSharedPtr<FWidgetBlueprintEditor> WidgetBlueprintEditor = WeakBlueprintEditor.Pin();
+		UWidgetBlueprint* WidgetBP = WidgetBlueprintEditor ? WidgetBlueprintEditor->GetWidgetBlueprintObj() : nullptr;
+		if (BlueprintView && WidgetBP)
+		{
+			if (FMVVMBlueprintViewModelContext* ViewModelContext = BlueprintView->FindViewModel(SelectedViewModelGuid))
+			{
+				FName NewName = ViewModelContext->GetViewModelName();
+				ViewModelContext->ViewModelName = PreviousViewModelPropertyName;
+				FText OutErrorMessage;
+				GEditor->GetEditorSubsystem<UMVVMEditorSubsystem>()->RenameViewModel(WidgetBP, PreviousViewModelPropertyName, NewName, OutErrorMessage);
+			}
+		}
+	}
+}
+
 } // namespace
 
 #undef LOCTEXT_NAMESPACE
\ No newline at end of file
diff --git a/Engine/Plugins/Runtime/ModelViewViewModel/Source/ModelViewViewModelEditor/Private/Widgets/SMVVMViewModelPanel.h b/Engine/Plugins/Runtime/ModelViewViewModel/Source/ModelViewViewModelEditor/Private/Widgets/SMVVMViewModelPanel.h
--- a/Engine/Plugins/Runtime/ModelViewViewModel/Source/ModelViewViewModelEditor/Private/Widgets/SMVVMViewModelPanel.h
+++ b/Engine/Plugins/Runtime/ModelViewViewModel/Source/ModelViewViewModelEditor/Private/Widgets/SMVVMViewModelPanel.h
@@ -4,47 +4,78 @@
 
 #include "CoreMinimal.h"
 
+#include "Misc/NotifyHook.h"
 #include "Widgets/PropertyViewer/SPropertyViewer.h"
 #include "Widgets/SMVVMViewModelBindingListWidget.h"
 #include "Widgets/SCompoundWidget.h"
 
 class FWidgetBlueprintEditor;
+class SInlineEditableTextBlock;
 class SPositiveActionButton;
 class UBlueprintExtension;
 class UMVVMBlueprintView;
+class FUICommandList;
+class IStructureDetailsView;
 
 namespace UE::MVVM
 {
-	/**
-	 * 
-	 */
-	class SMVVMViewModelPanel : public SCompoundWidget
-	{
-	public:
-		SLATE_BEGIN_ARGS(SMVVMViewModelPanel) {}
-		SLATE_END_ARGS()
-
-		void Construct(const FArguments& InArgs, TSharedPtr<FWidgetBlueprintEditor> Editor);
-		virtual ~SMVVMViewModelPanel();
-
-	private:
-		void HandleViewUpdated(UBlueprintExtension* Extension);
-		void HandleViewModelsUpdated();
-		TSharedRef<SWidget> MakeAddMenu();
-		void HandleCancelAddMenu();
-		void HandleAddMenuViewModel(const UClass* SelectedClass);
-		bool HandleCanEditViewmodelList() const;
-
-		void FillViewModel();
-
-	private:
-		TSharedPtr<SPositiveActionButton> AddMenuButton;
-		TSharedPtr<UE::PropertyViewer::SPropertyViewer> ViewModelTreeView;
-		TUniquePtr<FFieldIterator_Bindable> FieldIterator;
-
-		TWeakPtr<FWidgetBlueprintEditor> WeakBlueprintEditor;
-		TWeakObjectPtr<UMVVMBlueprintView> WeakBlueprintView;
-		FDelegateHandle ViewModelsUpdatedHandle;
-	};
+
+/**
+ *
+ */
+class SMVVMViewModelPanel : public SCompoundWidget, public FNotifyHook
+{
+public:
+	SLATE_BEGIN_ARGS(SMVVMViewModelPanel) {}
+	SLATE_END_ARGS()
+
+	void Construct(const FArguments& InArgs, TSharedPtr<FWidgetBlueprintEditor> Editor);
+	virtual ~SMVVMViewModelPanel();
+
+private:
+	void HandleViewUpdated(UBlueprintExtension* Extension);
+	void HandleViewModelsUpdated();
+	TSharedRef<SWidget> MakeAddMenu();
+	void HandleCancelAddMenu();
+	void HandleAddMenuViewModel(const UClass* SelectedClass);
+	bool HandleCanEditViewmodelList() const;
+	FReply OnKeyDown(const FGeometry& MyGeometry, const FKeyEvent& InKeyEvent) override;
+	TSharedRef<SWidget> HandleGenerateContainer(UE::PropertyViewer::SPropertyViewer::FHandle ContainerHandle, TOptional<FText> DisplayName);
+	TSharedPtr<SWidget> HandleContextMenuOpening(UE::PropertyViewer::SPropertyViewer::FHandle ContainerHandle, TArrayView<const FFieldVariant> Field) const;
+	void HandleSelectionChanged(UE::PropertyViewer::SPropertyViewer::FHandle, TArrayView<const FFieldVariant>, ESelectInfo::Type);
+	bool HandleVerifyNameTextChanged(const FText& InText, FText& OutErrorMessage, FGuid ViewModelGuid);
+	void HandleNameTextCommited(const FText& InText, ETextCommit::Type CommitInfo, FGuid ViewModelGuid);
+
+	void HandleDeleteViewModel();
+	bool HandleCanDeleteViewModel() const;
+	void HandleRenameViewModel();
+	bool HandleCanRenameViewModel() const;
+
+	void CreateCommandList();
+	void FillViewModel();
+	bool RenameViewModelProperty(FGuid ViewModelGuid, const FText& RenameTo, bool bCommit, FText& OutErrorMessage) const;
+
+	//~ Begin FNotifyHook
+	virtual void NotifyPreChange(FProperty* PropertyAboutToChange) override;
+	virtual void NotifyPostChange(const FPropertyChangedEvent& PropertyChangedEvent, FProperty* PropertyThatChanged) override;
+	//~ End FNotifyHook
+
+private:
+	TSharedPtr<SPositiveActionButton> AddMenuButton;
+	TSharedPtr<UE::PropertyViewer::SPropertyViewer> ViewModelTreeView;
+	TUniquePtr<FFieldIterator_Bindable> FieldIterator;
+	TSharedPtr<FUICommandList> CommandList;
+	TSharedPtr<IStructureDetailsView> PropertyView;
+
+	TMap<UE::PropertyViewer::SPropertyViewer::FHandle, FGuid> PropertyViewerHandles;
+	TMap<FGuid, TSharedPtr<SInlineEditableTextBlock>> EditableTextBlocks;
+
+	TWeakPtr<FWidgetBlueprintEditor> WeakBlueprintEditor;
+	TWeakObjectPtr<UMVVMBlueprintView> WeakBlueprintView;
+	FDelegateHandle ViewModelsUpdatedHandle;
+
+	FGuid SelectedViewModelGuid;
+	FName PreviousViewModelPropertyName;
+};
 
 } // namespace
diff --git a/Engine/Source/Runtime/AdvancedWidgets/Private/Widgets/PropertyViewer/SPropertyViewer.cpp b/Engine/Source/Runtime/AdvancedWidgets/Private/Widgets/PropertyViewer/SPropertyViewer.cpp
--- a/Engine/Source/Runtime/AdvancedWidgets/Private/Widgets/PropertyViewer/SPropertyViewer.cpp
+++ b/Engine/Source/Runtime/AdvancedWidgets/Private/Widgets/PropertyViewer/SPropertyViewer.cpp
@@ -100,6 +100,12 @@ void SPropertyViewer::RemoveAll()
 }
 
 
+TArray<SPropertyViewer::FSelectedItem> SPropertyViewer::GetSelectedItems() const
+{
+	return Implementation->GetSelectedItems();
+}
+
+
 void SPropertyViewer::SetRawFilterText(const FText& InFilterText)
 {
 	Implementation->SetRawFilterText(InFilterText);
diff --git a/Engine/Source/Runtime/AdvancedWidgets/Private/Widgets/PropertyViewer/SPropertyViewerImpl.cpp b/Engine/Source/Runtime/AdvancedWidgets/Private/Widgets/PropertyViewer/SPropertyViewerImpl.cpp
--- a/Engine/Source/Runtime/AdvancedWidgets/Private/Widgets/PropertyViewer/SPropertyViewerImpl.cpp
+++ b/Engine/Source/Runtime/AdvancedWidgets/Private/Widgets/PropertyViewer/SPropertyViewerImpl.cpp
@@ -395,8 +395,10 @@ FPropertyViewerImpl::FPropertyViewerImpl(const SPropertyViewer::FArguments& InAr
 	NotifyHook = InArgs._NotifyHook;
 	OnGetPreSlot = InArgs._OnGetPreSlot;
 	OnGetPostSlot = InArgs._OnGetPostSlot;
+	OnContextMenuOpening = InArgs._OnContextMenuOpening;
 	OnSelectionChanged = InArgs._OnSelectionChanged;
 	OnDoubleClicked = InArgs._OnDoubleClicked;
+	OnGenerateContainer = InArgs._OnGenerateContainer;
 	PropertyVisibility = InArgs._PropertyVisibility;
 	bSanitizeName = InArgs._bSanitizeName;
 	bShowFieldIcon = InArgs._bShowFieldIcon;
@@ -622,6 +624,43 @@ TSharedRef<SWidget> FPropertyViewerImpl::CreateSearch()
 }
 
 
+TArray<SPropertyViewer::FSelectedItem> FPropertyViewerImpl::GetSelectedItems() const
+{
+	TArray<SPropertyViewer::FSelectedItem> Result;
+
+	TArray<TSharedPtr<FTreeNode>> SelectedItem = TreeWidget->GetSelectedItems();
+	for (TSharedPtr<FTreeNode> Node : SelectedItem)
+	{
+		if (TSharedPtr<FContainer> Container = Node->GetContainer())
+		{
+			int32 FoundIndex = Result.IndexOfByPredicate([Container](const SPropertyViewer::FSelectedItem& Other) { return Other.Handle == Container->GetIdentifier(); });
+			if (FoundIndex == INDEX_NONE)
+			{
+				SPropertyViewer::FSelectedItem Item;
+				Item.Handle = Container->GetIdentifier();
+				FoundIndex = Result.Add(Item);
+			}
+			Result[FoundIndex].bIsContainerSelected = true;
+		}
+		else
+		{
+			if (TSharedPtr<FContainer> OwnerContainer = Node->GetOwnerContainer())
+			{
+				int32 FoundIndex = Result.IndexOfByPredicate([OwnerContainer](const SPropertyViewer::FSelectedItem& Other){ return Other.Handle == OwnerContainer->GetIdentifier(); });
+				if (FoundIndex == INDEX_NONE)
+				{
+					SPropertyViewer::FSelectedItem Item;
+					Item.Handle = OwnerContainer->GetIdentifier();
+					FoundIndex = Result.Add(Item);
+				}
+				Result[FoundIndex].Fields.Add(Node->GetFieldPath());
+			}
+		}
+	}
+	return Result;
+}
+
+
 void FPropertyViewerImpl::SetRawFilterText(const FText& InFilterText)
 {
 	SetRawFilterTextInternal(InFilterText);
@@ -752,7 +791,11 @@ TSharedRef<ITableRow> FPropertyViewerImpl::HandleGenerateRow(TSharedPtr<FTreeNod
 	TSharedPtr<SWidget> ItemWidget;
 	if (TSharedPtr<FContainer> ContainerPin = Item->GetContainer())
 	{
-		if (ContainerPin->IsValid())
+		if (OnGenerateContainer.IsBound())
+		{
+			ItemWidget = OnGenerateContainer.Execute(ContainerPin->GetIdentifier(), Item->GetOverrideDisplayName());
+		}
+		else if (ContainerPin->IsValid())
 		{
 			if (const UClass* Class = Cast<const UClass>(ContainerPin->GetStruct()))
 			{
diff --git a/Engine/Source/Runtime/AdvancedWidgets/Private/Widgets/PropertyViewer/SPropertyViewerImpl.h b/Engine/Source/Runtime/AdvancedWidgets/Private/Widgets/PropertyViewer/SPropertyViewerImpl.h
--- a/Engine/Source/Runtime/AdvancedWidgets/Private/Widgets/PropertyViewer/SPropertyViewerImpl.h
+++ b/Engine/Source/Runtime/AdvancedWidgets/Private/Widgets/PropertyViewer/SPropertyViewerImpl.h
@@ -197,6 +197,7 @@ private:
 	SPropertyViewer::FOnContextMenuOpening OnContextMenuOpening;
 	SPropertyViewer::FOnSelectionChanged OnSelectionChanged;
 	SPropertyViewer::FOnDoubleClicked OnDoubleClicked;
+	SPropertyViewer::FOnGenerateContainer OnGenerateContainer;
 	SPropertyViewer::EPropertyVisibility PropertyVisibility;
 	bool bSanitizeName = false;
 	bool bShowFieldIcon = false;
@@ -217,6 +218,7 @@ public:
 	void AddContainerInstance(SPropertyViewer::FHandle Identifier, TOptional<FText> DisplayName, const UScriptStruct* Struct, void* Data);
 	void Remove(SPropertyViewer::FHandle Identifier);
 	void RemoveAll();
+	TArray<SPropertyViewer::FSelectedItem> GetSelectedItems() const;
 	void SetRawFilterText(const FText& InFilterText);
 
 private:
diff --git a/Engine/Source/Runtime/AdvancedWidgets/Public/Widgets/PropertyViewer/SPropertyViewer.h b/Engine/Source/Runtime/AdvancedWidgets/Public/Widgets/PropertyViewer/SPropertyViewer.h
--- a/Engine/Source/Runtime/AdvancedWidgets/Public/Widgets/PropertyViewer/SPropertyViewer.h
+++ b/Engine/Source/Runtime/AdvancedWidgets/Public/Widgets/PropertyViewer/SPropertyViewer.h
@@ -41,6 +41,10 @@ public:
 		{
 			return Id != 0;
 		}
+		friend uint32 GetTypeHash(const FHandle& Other)
+		{
+			return GetTypeHash(Other.Id);
+		}
 	};
 
 	enum class EPropertyVisibility
@@ -50,11 +54,19 @@ public:
 		Editable,
 	};
 
+	struct FSelectedItem
+	{
+		FHandle Handle;
+		TArray<TArray<FFieldVariant>> Fields;
+		bool bIsContainerSelected = false;
+	};
+
 public:
 	DECLARE_DELEGATE_RetVal_TwoParams(TSharedPtr<SWidget>, FGetFieldWidget, FHandle, TArrayView<const FFieldVariant>);
 	DECLARE_DELEGATE_RetVal_TwoParams(TSharedPtr<SWidget>, FOnContextMenuOpening, FHandle, TArrayView<const FFieldVariant>);
 	DECLARE_DELEGATE_ThreeParams(FOnSelectionChanged, FHandle, TArrayView<const FFieldVariant>, ESelectInfo::Type);
 	DECLARE_DELEGATE_TwoParams(FOnDoubleClicked, FHandle, TArrayView<const FFieldVariant>);
+	DECLARE_DELEGATE_RetVal_TwoParams(TSharedRef<SWidget>, FOnGenerateContainer, FHandle, TOptional<FText> DisplayName);
 
 	SLATE_BEGIN_ARGS(SPropertyViewer)
 	{}
@@ -88,6 +100,8 @@ public:
 		SLATE_EVENT(FOnSelectionChanged, OnSelectionChanged)
 		/** Delegate to invoke when an item is double-clicked. */
 		SLATE_EVENT(FOnDoubleClicked, OnDoubleClicked);
+		/** Delegate to invoke when we generate the entry for an added container. */
+		SLATE_EVENT(FOnGenerateContainer, OnGenerateContainer)
 	SLATE_END_ARGS()
 
 	void Construct(const FArguments& InArgs);
@@ -107,6 +121,8 @@ public:
 	void Remove(FHandle Identifier);
 	void RemoveAll();
 
+	TArray<FSelectedItem> GetSelectedItems() const;
+
 	void SetRawFilterText(const FText& InFilterText);
 
 private:
