MVVM: Update the ViewModelBase SetProperty to support replication.

diff --git a/Engine/Plugins/Runtime/ModelViewViewModel/Source/ModelViewViewModel/ModelViewViewModel.Build.cs b/Engine/Plugins/Runtime/ModelViewViewModel/Source/ModelViewViewModel/ModelViewViewModel.Build.cs
--- a/Engine/Plugins/Runtime/ModelViewViewModel/Source/ModelViewViewModel/ModelViewViewModel.Build.cs
+++ b/Engine/Plugins/Runtime/ModelViewViewModel/Source/ModelViewViewModel/ModelViewViewModel.Build.cs
@@ -19,6 +19,7 @@ public class ModelViewViewModel : ModuleRules
 		PrivateDependencyModuleNames.AddRange(
 			new string[]
 			{
+				"NetCore",
 				"SlateCore",
 				"Slate",
 				"UMG",
diff --git a/Engine/Plugins/Runtime/ModelViewViewModel/Source/ModelViewViewModel/Private/MVVMViewModelBase.cpp b/Engine/Plugins/Runtime/ModelViewViewModel/Source/ModelViewViewModel/Private/MVVMViewModelBase.cpp
--- a/Engine/Plugins/Runtime/ModelViewViewModel/Source/ModelViewViewModel/Private/MVVMViewModelBase.cpp
+++ b/Engine/Plugins/Runtime/ModelViewViewModel/Source/ModelViewViewModel/Private/MVVMViewModelBase.cpp
@@ -2,8 +2,8 @@
 
 #include "MVVMViewModelBase.h"
 
-#include "ViewModel/MVVMViewModelBlueprintGeneratedClass.h"
 #include "Bindings/MVVMBindingHelper.h"
+#include "ViewModel/MVVMViewModelBlueprintGeneratedClass.h"
 
 #define LOCTEXT_NAMESPACE "MVVMViewModelBase"
 
@@ -78,12 +78,27 @@ void UMVVMViewModelBase::FFieldNotificationClassDescriptor::ForEachField(const U
 }
 
 
-void UMVVMViewModelBase::NotifyFieldValudChanged(UE::FieldNotification::FFieldId InFieldId)
+void UMVVMViewModelBase::MarkRepDirty(int32 InRepIndex)
 {
-	// Execute replication call
+#if WITH_PUSH_MODEL
+	check(InRepIndex >= 0 && InRepIndex < GetClass()->ClassReps.Num());
+	if (IS_PUSH_MODEL_ENABLED())
+	{
+		MARK_PROPERTY_DIRTY_UNSAFE(this, InRepIndex);
+	}
+#endif
+}
 
-	// Execute Broadcast call
-	BroadcastFieldValueChanged(InFieldId);
+
+void UMVVMViewModelBase::MarkRepDirty(FProperty* InProperty)
+{
+#if WITH_PUSH_MODEL
+	check(InProperty);
+	if (IS_PUSH_MODEL_ENABLED() && InProperty->RepIndex >= 0)
+	{
+		MARK_PROPERTY_DIRTY(this, InProperty);
+	}
+#endif
 }
 
 
@@ -146,7 +161,8 @@ DEFINE_FUNCTION(UMVVMViewModelBase::execK2_SetPropertyValue)
 			{
 				// Set the value then notify that the value changed.
 				TargetProperty->SetValue_InContainer(ViewModelContext, SourceValuePtr);
-				ViewModelContext->NotifyFieldValudChanged(FieldId);
+				ViewModelContext->MarkRepDirty(TargetProperty);
+				ViewModelContext->BroadcastFieldValueChanged(FieldId);
 			}
 		}
 
diff --git a/Engine/Plugins/Runtime/ModelViewViewModel/Source/ModelViewViewModel/Public/MVVMViewModelBase.h b/Engine/Plugins/Runtime/ModelViewViewModel/Source/ModelViewViewModel/Public/MVVMViewModelBase.h
--- a/Engine/Plugins/Runtime/ModelViewViewModel/Source/ModelViewViewModel/Public/MVVMViewModelBase.h
+++ b/Engine/Plugins/Runtime/ModelViewViewModel/Source/ModelViewViewModel/Public/MVVMViewModelBase.h
@@ -9,26 +9,32 @@
 #include "FieldNotification/FieldNotificationDeclaration.h"
 #include "FieldNotification/FieldMulticastDelegate.h"
 #include "FieldNotification/IFieldValueChanged.h"
+#include "Net/Core/PushModel/PushModel.h"
 #include "Types/MVVMAvailableBinding.h"
 #include "Types/MVVMBindingName.h"
 
 #include "MVVMViewModelBase.generated.h"
 
-/** After a field value changed. Execute all the notification (replication and broadcast). */
-#define UE_MVVM_NOTIFY_FIELD_VALUE_CHANGED(MemberName) \
-	NotifyFieldValudChanged(ThisClass::FFieldNotificationClassDescriptor::MemberName)
-
 /** After a field value changed. Broadcast the event (doesn't execute the replication code). */
 #define UE_MVVM_BROADCAST_FIELD_VALUE_CHANGED(MemberName) \
 	BroadcastFieldValueChanged(ThisClass::FFieldNotificationClassDescriptor::MemberName)
 
-/** Set the new value and notify if the property value changed. */
+/** After a field value changed. Replicate and broadcast the event. */
+#define UE_MVVM_BROADCAST_FIELD_VALUE_CHANGED_WITH_REP(MemberName ) \
+	MARK_PROPERTY_DIRTY_FROM_NAME(ThisClass, this, MemberName); \
+	BroadcastFieldValueChanged(ThisClass::FFieldNotificationClassDescriptor::MemberName)
+
+/** If the property value changed then set the new value and notify. */
 #define UE_MVVM_SET_PROPERTY_VALUE(MemberName, NewValue) \
 	SetPropertyValue(MemberName, NewValue, ThisClass::FFieldNotificationClassDescriptor::MemberName)
+	
+/** If the property value changed then set the new value, replicate and notify. */
+#define UE_MVVM_SET_PROPERTY_VALUE_WITH_REP(MemberName, NewValue) \
+	SetPropertyValue(MemberName, NewValue, ThisClass::FFieldNotificationClassDescriptor::MemberName, ThisClass::ENetFields_Private::MemberName)
 
 
 /** Base class for MVVM viewmodel. */
-UCLASS(Blueprintable, Abstract, DisplayName="MVVM ViewModel")
+UCLASS(Abstract, DisplayName="MVVM ViewModel")
 class MODELVIEWVIEWMODEL_API UMVVMViewModelBase : public UObject, public INotifyFieldValueChanged
 {
 	GENERATED_BODY()
@@ -50,8 +56,10 @@ public:
 
 
 protected:
-	/** Execute all the notification (replication and broadcast). */
-	void NotifyFieldValudChanged(UE::FieldNotification::FFieldId InFieldId);
+	/** Execute the replication code. */
+	void MarkRepDirty(FProperty* InProperty);
+private:
+	void MarkRepDirty(int32 InRepIndex);
 
 protected:
 	UFUNCTION(BlueprintCallable, Category="FieldNotify", meta=(DisplayName="Broadcast Field Value Changed", ScriptName="BroadcastFieldValueChanged", BlueprintInternalUseOnly="true"))
@@ -65,8 +73,22 @@ protected:
 	bool K2_SetPropertyValue(UPARAM(ref) const int32& OldValue, UPARAM(ref) const int32& NewValue);
 
 	/** Set the new value and notify if the property value changed. */
-	template<typename T>
-	bool SetPropertyValue(T& Value, const T& NewValue, UE::FieldNotification::FFieldId FieldId)
+	template<typename T, typename U>
+	bool SetPropertyValue(T& Value, const U& NewValue, UE::FieldNotification::FFieldId FieldId)
+	{
+		if (Value == NewValue)
+		{
+			return false;
+		}
+
+		Value = NewValue;
+		BroadcastFieldValueChanged(FieldId);
+		return true;
+	}
+
+	/** Set the new value and notify if the property value changed. */
+	template<typename T, typename U, typename NetFieldsEnum>
+	bool SetPropertyValue(T& Value, const U& NewValue, UE::FieldNotification::FFieldId FieldId, NetFieldsEnum InRepIndex)
 	{
 		if (Value == NewValue)
 		{
@@ -74,7 +96,8 @@ protected:
 		}
 
 		Value = NewValue;
-		NotifyFieldValudChanged(FieldId);
+		MarkRepDirty((int32)InRepIndex);
+		BroadcastFieldValueChanged(FieldId);
 		return true;
 	}
 
diff --git a/Engine/Plugins/Runtime/ModelViewViewModel/Source/ModelViewViewModelEditor/Public/MVVMEditorSubsystem.h b/Engine/Plugins/Runtime/ModelViewViewModel/Source/ModelViewViewModelEditor/Public/MVVMEditorSubsystem.h
--- a/Engine/Plugins/Runtime/ModelViewViewModel/Source/ModelViewViewModelEditor/Public/MVVMEditorSubsystem.h
+++ b/Engine/Plugins/Runtime/ModelViewViewModel/Source/ModelViewViewModelEditor/Public/MVVMEditorSubsystem.h
@@ -13,7 +13,6 @@
 #include "MVVMEditorSubsystem.generated.h"
 
 class UEdGraph;
-class UMVVMViewModelBase;
 class UWidgetBlueprint;
 
 /** */
