MVVM: Update the Binding panel to use the editor subsytem. Fix the "dock" drawer. Remove the "use viewmodel" step.

diff --git a/Engine/Plugins/Runtime/ModelViewViewModel/Source/ModelViewViewModelEditor/Private/MVVMEditorSubsystem.cpp b/Engine/Plugins/Runtime/ModelViewViewModel/Source/ModelViewViewModelEditor/Private/MVVMEditorSubsystem.cpp
--- a/Engine/Plugins/Runtime/ModelViewViewModel/Source/ModelViewViewModelEditor/Private/MVVMEditorSubsystem.cpp
+++ b/Engine/Plugins/Runtime/ModelViewViewModel/Source/ModelViewViewModelEditor/Private/MVVMEditorSubsystem.cpp
@@ -15,6 +15,7 @@
 #include "K2Node_FunctionEntry.h"
 #include "K2Node_FunctionResult.h"
 #include "K2Node_VariableGet.h"
+#include "MVVMBlueprintView.h"
 #include "MVVMSubsystem.h"
 #include "MVVMWidgetBlueprintExtension_View.h"
 #include "ScopedTransaction.h"
@@ -186,7 +187,12 @@ FName UMVVMEditorSubsystem::GetConversionFunctionWrapperName(const UWidgetBluepr
 
 UMVVMBlueprintView* UMVVMEditorSubsystem::RequestView(UWidgetBlueprint* WidgetBlueprint) const
 {
-	return UMVVMWidgetBlueprintExtension_View::RequestExtension<UMVVMWidgetBlueprintExtension_View>(WidgetBlueprint)->GetBlueprintView();
+	UMVVMWidgetBlueprintExtension_View* Extension = UMVVMWidgetBlueprintExtension_View::RequestExtension<UMVVMWidgetBlueprintExtension_View>(WidgetBlueprint);
+	if (Extension->GetBlueprintView() == nullptr)
+	{
+		Extension->CreateBlueprintViewInstance();
+	}
+	return Extension->GetBlueprintView();
 }
 
 UMVVMBlueprintView* UMVVMEditorSubsystem::GetView(const UWidgetBlueprint* WidgetBlueprint) const
diff --git a/Engine/Plugins/Runtime/ModelViewViewModel/Source/ModelViewViewModelEditor/Private/Widgets/SMVVMViewBindingListView.cpp b/Engine/Plugins/Runtime/ModelViewViewModel/Source/ModelViewViewModelEditor/Private/Widgets/SMVVMViewBindingListView.cpp
--- a/Engine/Plugins/Runtime/ModelViewViewModel/Source/ModelViewViewModelEditor/Private/Widgets/SMVVMViewBindingListView.cpp
+++ b/Engine/Plugins/Runtime/ModelViewViewModel/Source/ModelViewViewModelEditor/Private/Widgets/SMVVMViewBindingListView.cpp
@@ -846,6 +846,7 @@ void SBindingsList::Construct(const FArguments& InArgs, TSharedPtr<SBindingsPane
 	BindingPanel = Owner;
 	MVVMExtension = InMVVMExtension;
 	check(InMVVMExtension);
+	check(InMVVMExtension->GetBlueprintView());
 
 	MVVMExtension->OnBlueprintViewChangedDelegate().AddSP(this, &SBindingsList::Refresh);
 	MVVMExtension->GetBlueprintView()->OnBindingsUpdated.AddSP(this, &SBindingsList::Refresh);
diff --git a/Engine/Plugins/Runtime/ModelViewViewModel/Source/ModelViewViewModelEditor/Private/Widgets/SMVVMViewBindingPanel.cpp b/Engine/Plugins/Runtime/ModelViewViewModel/Source/ModelViewViewModelEditor/Private/Widgets/SMVVMViewBindingPanel.cpp
--- a/Engine/Plugins/Runtime/ModelViewViewModel/Source/ModelViewViewModelEditor/Private/Widgets/SMVVMViewBindingPanel.cpp
+++ b/Engine/Plugins/Runtime/ModelViewViewModel/Source/ModelViewViewModelEditor/Private/Widgets/SMVVMViewBindingPanel.cpp
@@ -2,10 +2,6 @@
 
 #include "Widgets/SMVVMViewBindingPanel.h"
 
-#include "Widgets/SMVVMManageViewModelsWidget.h"
-#include "Widgets/SMVVMViewModelContextListWidget.h"
-#include "Widgets/SMVVMViewBindingListView.h"
-
 #include "WidgetBlueprintEditor.h"
 #include "WidgetBlueprintToolMenuContext.h"
 #include "MVVMBlueprintView.h"
@@ -16,21 +12,25 @@
 #include "Styling/AppStyle.h"
 #include "IStructureDetailsView.h"
 #include "PropertyEditorModule.h"
+#include "StatusBarSubsystem.h"
 #include "ToolMenuContext.h"
 #include "ToolMenus.h"
+#include "Tabs/MVVMBindingSummoner.h"
 
 #include "Framework/MultiBox/MultiBoxBuilder.h"
 #include "Styling/MVVMEditorStyle.h"
+#include "Widgets/Docking/SDockTab.h"
 #include "Widgets/Input/SButton.h"
-#include "Widgets/SNullWidget.h"
 #include "Widgets/Images/SImage.h"
 #include "Widgets/Layout/SBorder.h"
 #include "Widgets/Layout/SSpacer.h"
 #include "Widgets/Layout/SSplitter.h"
+#include "Widgets/SMVVMManageViewModelsWidget.h"
+#include "Widgets/SMVVMViewModelContextListWidget.h"
+#include "Widgets/SMVVMViewBindingListView.h"
+#include "Widgets/SNullWidget.h"
+#include "SPositiveActionButton.h"
 
-#include "StatusBarSubsystem.h"
-#include "Tabs/MVVMBindingSummoner.h"
-#include "Widgets/Docking/SDockTab.h"
 
 #define LOCTEXT_NAMESPACE "BindingPanel"
 
@@ -116,25 +116,9 @@ bool SBindingsPanel::SupportsKeyboardFocus() const
 
 void SBindingsPanel::HandleBlueprintViewChangedDelegate()
 {
-	UMVVMWidgetBlueprintExtension_View* MVVMExtensionPtr = MVVMExtension.Get();
-
-	bool bHasExtension = MVVMExtensionPtr && MVVMExtensionPtr->GetBlueprintView() != nullptr;
 	ChildSlot
 	[
-		SNew(SOverlay)
-		+ SOverlay::Slot()
-		.HAlign(HAlign_Fill)
-		.VAlign(VAlign_Fill)
-		[
-			bHasExtension ? GenerateEditViewWidget() : GenerateCreateViewWidget()
-		]
-		+SOverlay::Slot()
-		.HAlign(HAlign_Right)
-		.VAlign(VAlign_Top)
-		.Padding(FMargin(bHasExtension ? 160.0 : 24.0, 10.0))
-		[
-			CreateDrawerDockButton()
-		]
+		GenerateEditViewWidget()
 	];
 }
 
@@ -155,16 +139,19 @@ void SBindingsPanel::OnBindingListSelectionChanged(TConstArrayView<FMVVMBlueprin
 	}
 }
 
-void SBindingsPanel::AddDefaultBinding()
+
+FReply SBindingsPanel::AddDefaultBinding()
 {
 	if (UMVVMWidgetBlueprintExtension_View* MVVMExtensionPtr = MVVMExtension.Get())
 	{
 		if (UMVVMBlueprintView* EditorData = MVVMExtensionPtr->GetBlueprintView())
 		{
+			check(BindingsList);
 			EditorData->AddDefaultBinding();
 			BindingsList->Refresh();
 		}
 	}
+	return FReply::Handled();
 }
 
 bool SBindingsPanel::CanAddBinding() const
@@ -196,31 +183,31 @@ TSharedRef<SWidget> SBindingsPanel::CreateDrawerDockButton()
 {
 	if (bIsDrawerTab)
 	{
-		return 
+		return
 			SNew(SButton)
 			.ButtonStyle(FAppStyle::Get(), "SimpleButton")
-			.ToolTipText(LOCTEXT("DockInLayout_Tooltip", "Docks MVVM drawer in tab."))
-			.ContentPadding(FMargin(1, 0))
+			.ToolTipText(LOCTEXT("BindingDockInLayout_Tooltip", "Docks the binding drawer in tab."))
+			.ContentPadding(FMargin(1.f, 0.f))
 			.OnClicked(this, &SBindingsPanel::CreateDrawerDockButtonClicked)
 			[
 				SNew(SHorizontalBox)
 				+ SHorizontalBox::Slot()
-				.AutoWidth()
-				.VAlign(VAlign_Center)
-				.Padding(4.0, 0.0f)
-				[
-					SNew(SImage)
-					.ColorAndOpacity(FSlateColor::UseForeground())
-					.Image(FAppStyle::Get().GetBrush("EditorViewport.SubMenu.Layouts"))
-				]
-				+ SHorizontalBox::Slot()
-				.VAlign(VAlign_Center)
-				.Padding(4.0, 0.0f)
-				[
-					SNew(STextBlock)
-					.Text(LOCTEXT("DockInLayout", "Dock in Layout"))
-					.ColorAndOpacity(FSlateColor::UseForeground())
-				]
+			.AutoWidth()
+			.VAlign(VAlign_Center)
+			.Padding(4.0, 0.0f)
+			[
+				SNew(SImage)
+				.ColorAndOpacity(FSlateColor::UseForeground())
+			.Image(FAppStyle::Get().GetBrush("Icons.Layout"))
+			]
+		+ SHorizontalBox::Slot()
+			.VAlign(VAlign_Center)
+			.Padding(4.0, 0.0f)
+			[
+				SNew(STextBlock)
+				.Text(LOCTEXT("DockInLayout", "Dock in Layout"))
+			.ColorAndOpacity(FSlateColor::UseForeground())
+			]
 			];
 	}
 
@@ -288,39 +275,18 @@ void SBindingsPanel::RegisterSettingsMenu()
 	UToolMenu* Menu = UToolMenus::Get()->RegisterMenu("MVVMEditor.Panel.Toolbar.Settings");
 }
 
-TSharedRef<SWidget> SBindingsPanel::GenerateCreateViewWidget()
-{
-	return SNew(SBorder)
-		[
-			SNew(SHorizontalBox)
-
-			+ SHorizontalBox::Slot()
-			.HAlign(EHorizontalAlignment::HAlign_Center)
-			.VAlign(EVerticalAlignment::VAlign_Center)
-			[
-				SNew(SButton)
-				.OnClicked(this, &SBindingsPanel::HandleCreateViewClicked)
-				.Text(LOCTEXT("AddViewModelButtonText", "Add View Model"))
-			]
-		];
-}
 
 TSharedRef<SWidget> SBindingsPanel::GenerateEditViewWidget()
 {
 	FSlimHorizontalToolBarBuilder ToolbarBuilderGlobal(TSharedPtr<const FUICommandList>(), FMultiBoxCustomization::None);
 	{
 		ToolbarBuilderGlobal.BeginSection("Binding");
-		ToolbarBuilderGlobal.AddToolBarButton(
-			FUIAction(
-				FExecuteAction::CreateSP(this, &SBindingsPanel::AddDefaultBinding),
-				FCanExecuteAction::CreateSP(this, &SBindingsPanel::CanAddBinding),
-				FGetActionCheckState()
-			),
-			NAME_None,
-			LOCTEXT("AddBinding", "Add Binding"),
-			TAttribute<FText>(this, &SBindingsPanel::GetAddBindingToolTip),
-			FSlateIcon(FMVVMEditorStyle::Get().GetStyleSetName(), "BindingView.AddBinding"),
-			EUserInterfaceActionType::Button
+		ToolbarBuilderGlobal.AddWidget(
+			SNew(SPositiveActionButton)
+			.OnClicked(this, &SBindingsPanel::AddDefaultBinding)
+			.Text(LOCTEXT("AddWidget", "Add Widget"))
+			.ToolTipText(this, &SBindingsPanel::GetAddBindingToolTip)
+			.IsEnabled(this, &SBindingsPanel::CanAddBinding)
 		);
 		ToolbarBuilderGlobal.EndSection();
 	}
@@ -330,6 +296,8 @@ TSharedRef<SWidget> SBindingsPanel::GenerateEditViewWidget()
 	}
 
 	{
+		ToolbarBuilderGlobal.AddWidget(CreateDrawerDockButton());
+
 		ToolbarBuilderGlobal.BeginSection("Options");
 		ToolbarBuilderGlobal.AddComboButton(
 			FUIAction(
@@ -346,6 +314,14 @@ TSharedRef<SWidget> SBindingsPanel::GenerateEditViewWidget()
 		ToolbarBuilderGlobal.EndSection();
 	}
 
+	BindingsList = nullptr;
+	if (MVVMExtension.Get())
+	{
+		BindingsList = SNew(SBindingsList, StaticCastSharedRef<SBindingsPanel>(AsShared()), MVVMExtension.Get());
+	}
+
+	TSharedRef<SWidget> BindingWidget = BindingsList ? BindingsList.ToSharedRef() : SNullWidget::NullWidget;
+
 	return SNew(SVerticalBox)
 
 		+ SVerticalBox::Slot()
@@ -355,7 +331,7 @@ TSharedRef<SWidget> SBindingsPanel::GenerateEditViewWidget()
 			ToolbarBuilderGlobal.MakeWidget()
 		]
 
-		+ SVerticalBox::Slot()
+		+SVerticalBox::Slot()
 		.FillHeight(1.0f)
 		[
 			SNew(SBorder)
@@ -364,13 +340,13 @@ TSharedRef<SWidget> SBindingsPanel::GenerateEditViewWidget()
 			[
 				SNew(SSplitter)
 				.Orientation(EOrientation::Orient_Horizontal)
-
+	
 				+ SSplitter::Slot()
 				.Value(0.75f)
 				[
-					SAssignNew(BindingsList, SBindingsList, StaticCastSharedRef<SBindingsPanel>(AsShared()), MVVMExtension.Get())
+					BindingWidget
 				]
-
+	
 				+ SSplitter::Slot()
 				.Value(0.25f)
 				[
@@ -383,75 +359,6 @@ TSharedRef<SWidget> SBindingsPanel::GenerateEditViewWidget()
 		];
 }
 
-
-FReply SBindingsPanel::HandleCreateViewClicked()
-{
-	if (TSharedPtr<FWidgetBlueprintEditor> WidgetBlueprintEditor = WeakBlueprintEditor.Pin())
-	{
-		UWidgetBlueprint* WidgetBlueprint = WidgetBlueprintEditor->GetWidgetBlueprintObj();
-		check(WidgetBlueprint);
-
-		WidgetBlueprint->OnExtensionAdded.RemoveAll(this);
-		UMVVMWidgetBlueprintExtension_View* MVVMExtensionPtr = UMVVMWidgetBlueprintExtension_View::RequestExtension<UMVVMWidgetBlueprintExtension_View>(WidgetBlueprint);
-		check(MVVMExtensionPtr);
-		MVVMExtension = MVVMExtensionPtr;
-
-		if (!BlueprintViewChangedDelegateHandle.IsValid())
-		{
-			BlueprintViewChangedDelegateHandle = MVVMExtensionPtr->OnBlueprintViewChangedDelegate().AddSP(this, &SBindingsPanel::HandleBlueprintViewChangedDelegate);
-		}
-
-		if (MVVMExtensionPtr->GetBlueprintView() == nullptr)
-		{
-			MVVMExtensionPtr->CreateBlueprintViewInstance();
-		}
-
-		HandleBlueprintViewChangedDelegate();
-
-		ShowManageViewModelsWindow();
-	}
-
-	return FReply::Handled();
-}
-
-void SBindingsPanel::ShowManageViewModelsWindow()
-{
-	if (UMVVMWidgetBlueprintExtension_View* MVVMExtensionPtr = MVVMExtension.Get())
-	{
-		if (UMVVMBlueprintView* BlueprintView = MVVMExtensionPtr->GetBlueprintView())
-		{
-			TWeakObjectPtr<UMVVMWidgetBlueprintExtension_View> LocalMVVMExtension = MVVMExtension;
-			FOnViewModelContextsPicked ViewModelContextsPickedDelegate = FOnViewModelContextsPicked::CreateLambda
-			(
-				[LocalMVVMExtension](TArray<FMVVMBlueprintViewModelContext> UpdatedContexts)
-				{
-					UMVVMWidgetBlueprintExtension_View* MVVMExtensionPtr = LocalMVVMExtension.Get();
-					if (ensure(MVVMExtensionPtr))
-					{
-						if (UMVVMBlueprintView* BlueprintView = MVVMExtensionPtr->GetBlueprintView())
-						{
-							BlueprintView->SetViewModels(UpdatedContexts);
-						}
-					}
-				}
-			);
-
-			TSharedRef<SWindow> ManageViewModelsWindow = SNew(SWindow)
-				.Title(LOCTEXT("MVVMManageViewMoedlsWindowHeader", "Manage ViewModels"))
-				.SupportsMaximize(false)
-				.ClientSize(FVector2D(800.0f, 600.0f));
-
-			TSharedRef<SMVVMManageViewModelsWidget> ManageViewModelsWidget = SNew(SMVVMManageViewModelsWidget)
-				.ParentWindow(ManageViewModelsWindow)
-				.OnViewModelContextsPickedDelegate(ViewModelContextsPickedDelegate)
-				.WidgetBlueprint(WeakBlueprintEditor.Pin()->GetWidgetBlueprintObj());
-
-			ManageViewModelsWindow->SetContent(ManageViewModelsWidget);
-			GEditor->EditorAddModalWindow(ManageViewModelsWindow);
-		}
-	}
-}
-
 } // namespace UE::MVVM
 
 #undef LOCTEXT_NAMESPACE
diff --git a/Engine/Plugins/Runtime/ModelViewViewModel/Source/ModelViewViewModelEditor/Private/Widgets/SMVVMViewBindingPanel.h b/Engine/Plugins/Runtime/ModelViewViewModel/Source/ModelViewViewModelEditor/Private/Widgets/SMVVMViewBindingPanel.h
--- a/Engine/Plugins/Runtime/ModelViewViewModel/Source/ModelViewViewModelEditor/Private/Widgets/SMVVMViewBindingPanel.h
+++ b/Engine/Plugins/Runtime/ModelViewViewModel/Source/ModelViewViewModelEditor/Private/Widgets/SMVVMViewBindingPanel.h
@@ -4,10 +4,11 @@
 
 #include "CoreMinimal.h"
 
+#include "Input/Reply.h"
+#include "Layout/Visibility.h"
 #include "Widgets/DeclarativeSyntaxSupport.h"
 #include "Widgets/SCompoundWidget.h"
-#include "Layout/Visibility.h"
-#include "Input/Reply.h"
+
 
 struct FMVVMBlueprintViewBinding;
 class FWidgetBlueprintEditor;
@@ -42,12 +43,9 @@ public:
 private:
 	void HandleBlueprintViewChangedDelegate();
 
-	TSharedRef<SWidget> GenerateCreateViewWidget();
-	FReply HandleCreateViewClicked();
-
 	TSharedRef<SWidget> GenerateEditViewWidget();
 
-	void AddDefaultBinding();
+	FReply AddDefaultBinding();
 	bool CanAddBinding() const;
 	FText GetAddBindingToolTip() const;
 
@@ -57,7 +55,6 @@ private:
 	FReply CreateDrawerDockButtonClicked();
 
 	void HandleExtensionAdded(UBlueprintExtension* NewExtension);
-	void ShowManageViewModelsWindow();
 
 private:
 	TWeakPtr<FWidgetBlueprintEditor> WeakBlueprintEditor;
diff --git a/Engine/Plugins/Runtime/ModelViewViewModel/Source/ModelViewViewModelEditor/Private/Widgets/SMVVMViewModelPanel.cpp b/Engine/Plugins/Runtime/ModelViewViewModel/Source/ModelViewViewModelEditor/Private/Widgets/SMVVMViewModelPanel.cpp
--- a/Engine/Plugins/Runtime/ModelViewViewModel/Source/ModelViewViewModelEditor/Private/Widgets/SMVVMViewModelPanel.cpp
+++ b/Engine/Plugins/Runtime/ModelViewViewModel/Source/ModelViewViewModelEditor/Private/Widgets/SMVVMViewModelPanel.cpp
@@ -242,6 +242,16 @@ void SMVVMViewModelPanel::HandleAddMenuViewModel(const UClass* SelectedClass)
 				{
 					UMVVMEditorSubsystem* EditorSubsystem = GEditor->GetEditorSubsystem<UMVVMEditorSubsystem>();
 					check(EditorSubsystem);
+
+					UMVVMBlueprintView* CurrentBlueprintView = WeakBlueprintView.Get();
+					if (!CurrentBlueprintView)
+					{
+						EditorSubsystem->RequestView(WidgetBlueprint);
+						CurrentBlueprintView = EditorSubsystem->GetView(WidgetBlueprint);
+						WeakBlueprintView = CurrentBlueprintView;
+						ViewModelsUpdatedHandle = CurrentBlueprintView->OnViewModelsUpdated.AddSP(this, &SMVVMViewModelPanel::HandleViewModelsUpdated);
+					}
+
 					EditorSubsystem->AddViewModel(WidgetBlueprint, SelectedClass);
 				}
 			}
