MVVM: Add "Optional" on viewmodel context to prevent logging warning when the viewmodel is not set and that it's optional.

diff --git a/Engine/Plugins/Runtime/ModelViewViewModel/Source/ModelViewViewModel/Private/View/MVVMView.cpp b/Engine/Plugins/Runtime/ModelViewViewModel/Source/ModelViewViewModel/Private/View/MVVMView.cpp
--- a/Engine/Plugins/Runtime/ModelViewViewModel/Source/ModelViewViewModel/Private/View/MVVMView.cpp
+++ b/Engine/Plugins/Runtime/ModelViewViewModel/Source/ModelViewViewModel/Private/View/MVVMView.cpp
@@ -340,8 +340,11 @@ bool UMVVMView::RegisterLibraryBinding(const FMVVMViewClass_CompiledBinding& Bin
 	TScriptInterface<INotifyFieldValueChanged> Source = FindSource(Binding, true);
 	if (Source.GetInterface() == nullptr)
 	{
-		UE::MVVM::FMessageLog Log(GetUserWidget());
-		Log.Error(FText::Format(LOCTEXT("RegisterBindingInvalidSourceInstance", "Can't register binding '{0}'. The source is invalid."), FText::FromString(Binding.ToString())));
+		if (!Binding.IsRegistrationOptional())
+		{
+			UE::MVVM::FMessageLog Log(GetUserWidget());
+			Log.Error(FText::Format(LOCTEXT("RegisterBindingInvalidSourceInstance", "Can't register binding '{0}'. The source is invalid."), FText::FromString(Binding.ToString())));
+		}
 		return false;
 	}
 
diff --git a/Engine/Plugins/Runtime/ModelViewViewModel/Source/ModelViewViewModel/Private/View/MVVMViewClass.cpp b/Engine/Plugins/Runtime/ModelViewViewModel/Source/ModelViewViewModel/Private/View/MVVMViewClass.cpp
--- a/Engine/Plugins/Runtime/ModelViewViewModel/Source/ModelViewViewModel/Private/View/MVVMViewClass.cpp
+++ b/Engine/Plugins/Runtime/ModelViewViewModel/Source/ModelViewViewModel/Private/View/MVVMViewClass.cpp
@@ -47,7 +47,7 @@ FMVVMViewClass_SourceCreator FMVVMViewClass_SourceCreator::MakeInstance(FName In
 }
 
 
-FMVVMViewClass_SourceCreator FMVVMViewClass_SourceCreator::MakeFieldPath(FName InName, UClass* InNotifyFieldValueChangedClass, FMVVMVCompiledFieldPath InFieldPath)
+FMVVMViewClass_SourceCreator FMVVMViewClass_SourceCreator::MakeFieldPath(FName InName, UClass* InNotifyFieldValueChangedClass, FMVVMVCompiledFieldPath InFieldPath, bool bOptional)
 {
 	FMVVMViewClass_SourceCreator Result;
 	if (ensure(InNotifyFieldValueChangedClass&& InNotifyFieldValueChangedClass->ImplementsInterface(UNotifyFieldValueChanged::StaticClass())))
@@ -57,14 +57,14 @@ FMVVMViewClass_SourceCreator FMVVMViewClass_SourceCreator::MakeFieldPath(FName I
 			Result.PropertyName = InName;
 			Result.ExpectedSourceType = InNotifyFieldValueChangedClass;
 			Result.FieldPath = InFieldPath;
-			Result.bOptional = false;
+			Result.bOptional = bOptional;
 		}
 	}
 	return Result;
 }
 
 
-FMVVMViewClass_SourceCreator FMVVMViewClass_SourceCreator::MakeGlobalContext(FName InName, FMVVMViewModelContext InContext)
+FMVVMViewClass_SourceCreator FMVVMViewClass_SourceCreator::MakeGlobalContext(FName InName, FMVVMViewModelContext InContext, bool bOptional)
 {
 	FMVVMViewClass_SourceCreator Result;
 	if (ensure(InContext.ContextClass && InContext.ContextClass->ImplementsInterface(UNotifyFieldValueChanged::StaticClass())))
@@ -72,7 +72,7 @@ FMVVMViewClass_SourceCreator FMVVMViewClass_SourceCreator::MakeGlobalContext(FNa
 		Result.PropertyName = InName;
 		Result.ExpectedSourceType = InContext.ContextClass;
 		Result.GlobalViewModelInstance = MoveTemp(InContext);
-		Result.bOptional = false;
+		Result.bOptional = bOptional;
 	}
 	return Result;
 }
diff --git a/Engine/Plugins/Runtime/ModelViewViewModel/Source/ModelViewViewModel/Public/View/MVVMViewClass.h b/Engine/Plugins/Runtime/ModelViewViewModel/Source/ModelViewViewModel/Public/View/MVVMViewClass.h
--- a/Engine/Plugins/Runtime/ModelViewViewModel/Source/ModelViewViewModel/Public/View/MVVMViewClass.h
+++ b/Engine/Plugins/Runtime/ModelViewViewModel/Source/ModelViewViewModel/Public/View/MVVMViewClass.h
@@ -42,8 +42,8 @@ struct FMVVMViewClass_SourceCreator
 public:
 	MODELVIEWVIEWMODEL_API static FMVVMViewClass_SourceCreator MakeManual(FName Name, UClass* NotifyFieldValueChangedClass);
 	MODELVIEWVIEWMODEL_API static FMVVMViewClass_SourceCreator MakeInstance(FName Name, UClass* NotifyFieldValueChangedClass);
-	MODELVIEWVIEWMODEL_API static FMVVMViewClass_SourceCreator MakeFieldPath(FName Name, UClass* NotifyFieldValueChangedClass, FMVVMVCompiledFieldPath FieldPath);
-	MODELVIEWVIEWMODEL_API static FMVVMViewClass_SourceCreator MakeGlobalContext(FName Name, FMVVMViewModelContext Context);
+	MODELVIEWVIEWMODEL_API static FMVVMViewClass_SourceCreator MakeFieldPath(FName Name, UClass* NotifyFieldValueChangedClass, FMVVMVCompiledFieldPath FieldPath, bool bOptional);
+	MODELVIEWVIEWMODEL_API static FMVVMViewClass_SourceCreator MakeGlobalContext(FName Name, FMVVMViewModelContext Context, bool bOptional);
 
 	UObject* CreateInstance(const UMVVMViewClass* ViewClass, UMVVMView* View, UUserWidget* UserWidget) const;
 
@@ -135,6 +135,12 @@ public:
 		return (Flags & EBindingFlags::EnabledByDefault) != 0;
 	}
 
+	/** @return true if the binding . */
+	bool IsRegistrationOptional() const
+	{
+		return (Flags & EBindingFlags::RegistrationOptional) != 0;
+	}
+
 	/** @return a human readable version of the binding that can be use for debugging purposes. */
 	FString ToString() const;
 
@@ -159,6 +165,7 @@ private:
 		TwoWayBinding = 1 << 1, // The binding is one part of a 2 ways binding.
 		OneTime = 1 << 2,
 		EnabledByDefault = 1 << 3,
+		RegistrationOptional = 1 << 4,	// The source (viewmodel) can be nullptr and the binding could failed without warning.
 	};
 
 	UPROPERTY()
diff --git a/Engine/Plugins/Runtime/ModelViewViewModel/Source/ModelViewViewModelBlueprint/Private/MVVMViewBlueprintCompiler.cpp b/Engine/Plugins/Runtime/ModelViewViewModel/Source/ModelViewViewModelBlueprint/Private/MVVMViewBlueprintCompiler.cpp
--- a/Engine/Plugins/Runtime/ModelViewViewModel/Source/ModelViewViewModelBlueprint/Private/MVVMViewBlueprintCompiler.cpp
+++ b/Engine/Plugins/Runtime/ModelViewViewModel/Source/ModelViewViewModelBlueprint/Private/MVVMViewBlueprintCompiler.cpp
@@ -699,7 +699,7 @@ bool FMVVMViewBlueprintCompiler::CompileSourceCreators(const FCompiledBindingLib
 					bAreSourcesCreatorValid = false;
 					continue;
 				}
-				CompiledSourceCreator = FMVVMViewClass_SourceCreator::MakeFieldPath(ViewModelContext.GetViewModelName(), ViewModelContext.GetViewModelClass(), *CompiledFieldPath);
+				CompiledSourceCreator = FMVVMViewClass_SourceCreator::MakeFieldPath(ViewModelContext.GetViewModelName(), ViewModelContext.GetViewModelClass(), *CompiledFieldPath, ViewModelContext.bOptional);
 			}
 			else if (ViewModelContext.CreationType == EMVVMBlueprintViewModelContextCreationType::GlobalViewModelCollection)
 			{
@@ -722,7 +722,7 @@ bool FMVVMViewBlueprintCompiler::CompileSourceCreators(const FCompiledBindingLib
 					continue;
 				}
 
-				CompiledSourceCreator = FMVVMViewClass_SourceCreator::MakeGlobalContext(ViewModelContext.GetViewModelName(), MoveTemp(GlobalViewModelInstance));
+				CompiledSourceCreator = FMVVMViewClass_SourceCreator::MakeGlobalContext(ViewModelContext.GetViewModelName(), MoveTemp(GlobalViewModelInstance), ViewModelContext.bOptional);
 			}
 			else
 			{
diff --git a/Engine/Plugins/Runtime/ModelViewViewModel/Source/ModelViewViewModelBlueprint/Public/MVVMBlueprintViewModelContext.h b/Engine/Plugins/Runtime/ModelViewViewModel/Source/ModelViewViewModelBlueprint/Public/MVVMBlueprintViewModelContext.h
--- a/Engine/Plugins/Runtime/ModelViewViewModel/Source/ModelViewViewModelBlueprint/Public/MVVMBlueprintViewModelContext.h
+++ b/Engine/Plugins/Runtime/ModelViewViewModel/Source/ModelViewViewModelBlueprint/Public/MVVMBlueprintViewModelContext.h
@@ -105,9 +105,19 @@ public:
 	UPROPERTY(EditAnywhere, Category = "MVVM", AdvancedDisplay, meta = (EditCondition = "CreationType == EMVVMBlueprintViewModelContextCreationType::PropertyPath"))
 	FString ViewModelPropertyPath;
 
-	/** Generate a setter function for this viewmodel. */
-	UPROPERTY(EditAnywhere, Category = "MVVM", AdvancedDisplay)
+	/**
+	 * Generate a setter function for this viewmodel.
+	 * @note Always true when the Creation Type is Manual.
+	 */
+	UPROPERTY(EditAnywhere, Category = "MVVM", AdvancedDisplay, meta = (EditCondition = "CreationType != EMVVMBlueprintViewModelContextCreationType::Manual"))
 	bool bCreateSetterFunction = false;
+
+	/**
+	 * Optional. Will not warn if the instance is not set or found.
+	 * @note Always true when the Creation Type is Manual.
+	 */
+	UPROPERTY(EditAnywhere, Category = "MVVM", AdvancedDisplay, meta = (EditCondition = "CreationType == EMVVMBlueprintViewModelContextCreationType::GlobalViewModelCollection || CreationType == EMVVMBlueprintViewModelContextCreationType::PropertyPath"))
+	bool bOptional = false;
 };
 
 template<>
