MVVM: Fix binding unit test

diff --git a/Engine/Plugins/Runtime/ModelViewViewModel/Source/ModelViewViewModelEditor/Private/Tests/MVVMBindingHelperTest.cpp b/Engine/Plugins/Runtime/ModelViewViewModel/Source/ModelViewViewModelEditor/Private/Tests/MVVMBindingHelperTest.cpp
--- a/Engine/Plugins/Runtime/ModelViewViewModel/Source/ModelViewViewModelEditor/Private/Tests/MVVMBindingHelperTest.cpp
+++ b/Engine/Plugins/Runtime/ModelViewViewModel/Source/ModelViewViewModelEditor/Private/Tests/MVVMBindingHelperTest.cpp
@@ -121,13 +121,13 @@ bool FMVVMBindingHelperTest::RunTest(const FString& Parameters)
 	TestFunction("FunctionGetA", false, false, false, nullptr, nullptr);
 	TestFunction("FunctionGetB", false, false, false, nullptr, nullptr);
 	TestFunction("FunctionGetC", true, false, false, FIntProperty::StaticClass(), nullptr);
-	TestFunction("FunctionGetD", false, false, false, nullptr, nullptr);
+	TestFunction("FunctionGetD", false, false, true, nullptr, nullptr);
 	TestFunction("FunctionGetE", false, false, false, nullptr, nullptr);
 	TestFunction("FunctionGetF", true, false, false, FArrayProperty::StaticClass(), nullptr);
 	TestFunction("FunctionGetG", true, false, false, FArrayProperty::StaticClass(), nullptr);
 	TestFunction("FunctionGetH", false, false, false, nullptr, nullptr);
 	TestFunction("FunctionGetI", false, false, false, nullptr, nullptr);
-	TestFunction("FunctionGetJ", false, false, false, nullptr, nullptr);
+	TestFunction("FunctionGetJ", false, false, true, nullptr, nullptr);
 
 	TestFunction("FunctionSetA", false, false, false, nullptr, nullptr);
 	TestFunction("FunctionSetB", false, true, false, nullptr, FIntProperty::StaticClass());
