MVVM: Fix bad call to UFunction::GetReturnProperty. Fix BP function to not be editable (for now).

diff --git a/Engine/Plugins/Runtime/ModelViewViewModel/Source/ModelViewViewModelEditor/Private/MVVMEditorSubsystem.cpp b/Engine/Plugins/Runtime/ModelViewViewModel/Source/ModelViewViewModelEditor/Private/MVVMEditorSubsystem.cpp
--- a/Engine/Plugins/Runtime/ModelViewViewModel/Source/ModelViewViewModelEditor/Private/MVVMEditorSubsystem.cpp
+++ b/Engine/Plugins/Runtime/ModelViewViewModel/Source/ModelViewViewModelEditor/Private/MVVMEditorSubsystem.cpp
@@ -56,7 +56,7 @@ namespace UE::MVVM::Private
 
 	UEdGraph* CreateFunctionGraph(UBlueprint* Blueprint, FName InFunctionName, EFunctionFlags ExtraFunctionFlag, const FStringView Category)
 	{
-		bool bIsEditable = true;
+		bool bIsEditable = false;
 
 		FName UniqueFunctionName = FBlueprintEditorUtils::FindUniqueKismetName(Blueprint, InFunctionName.ToString());
 		UEdGraph* FunctionGraph = FBlueprintEditorUtils::CreateNewGraph(Blueprint, UniqueFunctionName, UEdGraph::StaticClass(), UEdGraphSchema_K2::StaticClass());
@@ -363,7 +363,8 @@ UEdGraph* UMVVMEditorSubsystem::CreateConversionFunctionWrapperGraph(UWidgetBlue
 				const UEdGraphSchema_K2* K2Schema = GetDefault<UEdGraphSchema_K2>();
 
 				// create return value pin
-				const FProperty* ReturnProperty = ConversionFunction->GetReturnProperty();
+				const FProperty* ReturnProperty = UE::MVVM::BindingHelper::GetReturnProperty(ConversionFunction);
+				check(ReturnProperty);
 				TSharedPtr<FUserPinInfo> PinInfo = MakeShared<FUserPinInfo>();
 				K2Schema->ConvertPropertyToPinType(ReturnProperty, PinInfo->PinType);
 				PinInfo->PinName = ReturnProperty->GetFName();
