MVVM: Fix view compiler that uses the same conversion function for forward and backward binding.

diff --git a/Engine/Plugins/Runtime/ModelViewViewModel/Source/ModelViewViewModelBlueprint/Private/MVVMViewBlueprintCompiler.cpp b/Engine/Plugins/Runtime/ModelViewViewModel/Source/ModelViewViewModelBlueprint/Private/MVVMViewBlueprintCompiler.cpp
--- a/Engine/Plugins/Runtime/ModelViewViewModel/Source/ModelViewViewModelBlueprint/Private/MVVMViewBlueprintCompiler.cpp
+++ b/Engine/Plugins/Runtime/ModelViewViewModel/Source/ModelViewViewModelBlueprint/Private/MVVMViewBlueprintCompiler.cpp
@@ -860,13 +860,14 @@ bool FMVVMViewBlueprintCompiler::PreCompileBindings(UWidgetBlueprintGeneratedCla
 			}
 		}
 
-		FMemberReference ConversionFunctionRef = Binding.Conversion.SourceToDestinationFunction;
-		if (!Binding.Conversion.SourceToDestinationWrapper.IsNone())
+		FMemberReference ConversionFunctionReference = BindingSourceContext.bIsForwardBinding ? Binding.Conversion.SourceToDestinationFunction : Binding.Conversion.DestinationToSourceFunction;
+		FName ConversionFunctionWrapper = BindingSourceContext.bIsForwardBinding ? Binding.Conversion.SourceToDestinationWrapper : Binding.Conversion.DestinationToSourceWrapper;
+		if (!ConversionFunctionWrapper.IsNone())
 		{
-			ConversionFunctionRef.SetSelfMember(Binding.Conversion.SourceToDestinationWrapper);
+			ConversionFunctionReference.SetSelfMember(ConversionFunctionWrapper);
 		}
 
-		const UFunction* ConversionFunction = ConversionFunctionRef.ResolveMember<UFunction>(Class);
+		const UFunction* ConversionFunction = ConversionFunctionReference.ResolveMember<UFunction>(Class);
 
 		TValueOrError<FCompilerBinding, FString> AddBindingResult = AddBinding(Class, GetterFields, SetterPath, ConversionFunction);
 		if (AddBindingResult.HasError())
