MVVM: Create helper function to find all the PropertyPath use by a graph.

diff --git a/Engine/Plugins/Runtime/ModelViewViewModel/Source/ModelViewViewModelBlueprint/Private/Bindings/MVVMConversionFunctionHelper.cpp b/Engine/Plugins/Runtime/ModelViewViewModel/Source/ModelViewViewModelBlueprint/Private/Bindings/MVVMConversionFunctionHelper.cpp
new file mode 100644
--- /dev/null
+++ b/Engine/Plugins/Runtime/ModelViewViewModel/Source/ModelViewViewModelBlueprint/Private/Bindings/MVVMConversionFunctionHelper.cpp
@@ -0,0 +1,97 @@
+// Copyright Epic Games, Inc. All Rights Reserved.
+
+#include "Bindings/MVVMConversionFunctionHelper.h"
+
+#include "BlueprintTypePromotion.h"
+#include "EngineLogs.h"
+#include "EdGraph/EdGraph.h"
+#include "K2Node_PromotableOperator.h"
+#include "K2Node_VariableGet.h"
+#include "Kismet2/BlueprintEditorUtils.h"
+
+#include "Components/Widget.h"
+#include "MVVMBlueprintView.h"
+#include "MVVMBlueprintViewModelContext.h"
+#include "MVVMViewModelBase.h"
+#include "View/MVVMView.h"
+#include "WidgetBlueprint.h"
+
+
+namespace UE::MVVM::ConversionFunctionHelper
+{
+
+namespace Private
+{
+void BuildPropertyPath(TArray<FMVVMBlueprintPropertyPath>& Result, FMVVMBlueprintPropertyPath BasePath, UClass* Class, const UK2Node_VariableGet* CurrentNode)
+{
+	bool bAddPath = false;
+	const UEdGraphPin* OutputPin = CurrentNode->FindPinByPredicate([](UEdGraphPin* Pin) { return Pin->Direction == EGPD_Output; });
+	if (OutputPin)
+	{
+		for (const UEdGraphPin* LinkedTo : OutputPin->LinkedTo)
+		{
+			if (const UK2Node_VariableGet* LinkToVariable = Cast<UK2Node_VariableGet>(LinkedTo->GetOwningNode()))
+			{
+				BasePath.SetBasePropertyPath(UE::MVVM::FMVVMConstFieldVariant(LinkToVariable->VariableReference.ResolveMember<FProperty>(Class)));
+				BuildPropertyPath(Result, BasePath, Class, LinkToVariable);
+			}
+			else
+			{
+				bAddPath = true;
+			}
+		}
+	}
+
+	if (bAddPath)
+	{
+		Result.Add(BasePath);
+	}
+}
+} //namespace
+
+TArray<FMVVMBlueprintPropertyPath> FindAllPropertyPathInGraph(const UEdGraph* Graph, const UMVVMBlueprintView* BlueprintView, UClass* Class)
+{
+	TArray<FMVVMBlueprintPropertyPath> Result;
+
+	TArray<UK2Node_VariableGet*> GetNodes;
+	Graph->GetNodesOfClass<UK2Node_VariableGet>(GetNodes);
+	for (const UK2Node_VariableGet* Node : GetNodes)
+	{
+		if (Node->VariableReference.IsSelfContext())
+		{
+			if (FProperty* Property = Node->VariableReference.ResolveMember<FProperty>(Class))
+			{
+				const UClass* OwningClass = Class;
+				if (const FObjectProperty* ObjectProperty = CastField<FObjectProperty>(Property))
+				{
+					OwningClass = ObjectProperty->PropertyClass;
+				}
+				check(OwningClass);
+
+				FMVVMBlueprintPropertyPath PropertyPath;
+				if (OwningClass->IsChildOf<UWidget>() || OwningClass->IsChildOf<UWidgetBlueprint>())
+				{
+					PropertyPath.SetWidgetName(Property->GetFName());
+				}
+				else if (OwningClass->ImplementsInterface(UNotifyFieldValueChanged::StaticClass()))
+				{
+					if (const FMVVMBlueprintViewModelContext* ViewModel = BlueprintView->FindViewModel(Property->GetFName()))
+					{
+						PropertyPath.SetViewModelId(ViewModel->GetViewModelId());
+					}
+				}
+
+				if (PropertyPath.IsEmpty())
+				{
+					continue;
+				}
+
+				Private::BuildPropertyPath(Result, PropertyPath, Class, Node);
+			}
+		}
+	}
+
+	return Result;
+}
+
+} //namespace
\ No newline at end of file
diff --git a/Engine/Plugins/Runtime/ModelViewViewModel/Source/ModelViewViewModelBlueprint/Public/Bindings/MVVMConversionFunctionHelper.h b/Engine/Plugins/Runtime/ModelViewViewModel/Source/ModelViewViewModelBlueprint/Public/Bindings/MVVMConversionFunctionHelper.h
new file mode 100644
--- /dev/null
+++ b/Engine/Plugins/Runtime/ModelViewViewModel/Source/ModelViewViewModelBlueprint/Public/Bindings/MVVMConversionFunctionHelper.h
@@ -0,0 +1,18 @@
+// Copyright Epic Games, Inc. All Rights Reserved.
+
+#pragma once
+
+#include "CoreMinimal.h"
+#include "MVVMPropertyPath.h"
+
+class UClass;
+class UEdGraph;
+class UMVVMBlueprintView;
+
+namespace UE::MVVM::ConversionFunctionHelper
+{
+
+	/** Find all BlueprintPropertyPath used in the Graph. */
+	MODELVIEWVIEWMODELBLUEPRINT_API TArray<FMVVMBlueprintPropertyPath> FindAllPropertyPathInGraph(const UEdGraph* Graph, const UMVVMBlueprintView* View, UClass* Class);
+
+} //namespace
\ No newline at end of file
