MVVM: Renamed variables for View Compiler. The indexes was confusing.

diff --git a/Engine/Plugins/Runtime/ModelViewViewModel/Source/ModelViewViewModelBlueprint/Private/MVVMViewBlueprintCompiler.cpp b/Engine/Plugins/Runtime/ModelViewViewModel/Source/ModelViewViewModelBlueprint/Private/MVVMViewBlueprintCompiler.cpp
--- a/Engine/Plugins/Runtime/ModelViewViewModel/Source/ModelViewViewModelBlueprint/Private/MVVMViewBlueprintCompiler.cpp
+++ b/Engine/Plugins/Runtime/ModelViewViewModel/Source/ModelViewViewModelBlueprint/Private/MVVMViewBlueprintCompiler.cpp
@@ -67,7 +67,7 @@ void FMVVMViewBlueprintCompiler::CleanOldData(UWidgetBlueprintGeneratedClass* Cl
 
 void FMVVMViewBlueprintCompiler::CleanTemporaries(UWidgetBlueprintGeneratedClass* ClassToClean)
 {
-	for (FCompilerSourceCreatorContext& SourceCreatorContext : SourceCreatorContexts)
+	for (FCompilerSourceCreatorContext& SourceCreatorContext : CompilerSourceCreatorContexts)
 	{
 		if (SourceCreatorContext.SetterGraph)
 		{
@@ -87,7 +87,7 @@ void FMVVMViewBlueprintCompiler::CreateFunctions(UMVVMBlueprintView* BlueprintVi
 		return;
 	}
 
-	for (const FCompilerSourceCreatorContext& SourceCreator : SourceCreatorContexts)
+	for (const FCompilerSourceCreatorContext& SourceCreator : CompilerSourceCreatorContexts)
 	{
 		if (SourceCreator.SetterGraph)
 		{
@@ -152,7 +152,7 @@ void FMVVMViewBlueprintCompiler::CreateVariables(const FWidgetBlueprintCompilerC
 		return ViewModelProperty;
 	};
 
-	for (FCompilerSourceContext& SourceContext : SourceContexts)
+	for (FCompilerSourceContext& SourceContext : CompilerSourceContexts)
 	{
 		SourceContext.Field = BindingHelper::FindFieldByName(Context.GetSkeletonGeneratedClass(), FMVVMBindingName(SourceContext.PropertyName));
 
@@ -243,8 +243,8 @@ void FMVVMViewBlueprintCompiler::CreateWidgetMap(const FWidgetBlueprintCompilerC
 
 void FMVVMViewBlueprintCompiler::CreateSourceLists(const FWidgetBlueprintCompilerContext::FCreateVariableContext& Context, UMVVMBlueprintView* BlueprintView)
 {
-	SourceContexts.Reset();
-	SourceCreatorContexts.Reset();
+	CompilerSourceContexts.Reset();
+	CompilerSourceCreatorContexts.Reset();
 
 	TSet<FGuid> ViewModelsGuid;
 	TSet<FName> WidgetSources;
@@ -287,12 +287,12 @@ void FMVVMViewBlueprintCompiler::CreateSourceLists(const FWidgetBlueprintCompile
 			{
 				SourceContext.SetterFunctionName = TEXT("Set") + ViewModelContext.GetViewModelName().ToString();
 			}
-			FoundSourceCreatorContextIndex = SourceCreatorContexts.Emplace(MoveTemp(SourceContext));
+			FoundSourceCreatorContextIndex = CompilerSourceCreatorContexts.Emplace(MoveTemp(SourceContext));
 		}
 		else
 		{
 			FGuid ViewModelId = ViewModelContext.GetViewModelId();
-			FoundSourceCreatorContextIndex = SourceCreatorContexts.IndexOfByPredicate([ViewModelId](const FCompilerSourceCreatorContext& Other)
+			FoundSourceCreatorContextIndex = CompilerSourceCreatorContexts.IndexOfByPredicate([ViewModelId](const FCompilerSourceCreatorContext& Other)
 				{
 					return Other.ViewModelContext.GetViewModelId() == ViewModelId;
 				});
@@ -305,8 +305,8 @@ void FMVVMViewBlueprintCompiler::CreateSourceLists(const FWidgetBlueprintCompile
 		SourceVariable.DisplayName = ViewModelContext.GetDisplayName().ToString();
 		SourceVariable.CategoryName = TEXT("Viewmodel");
 		SourceVariable.bExposeOnSpawn = bCreateSetterFunction;
-		SourceVariable.BlueprintSetter = SourceCreatorContexts[FoundSourceCreatorContextIndex].SetterFunctionName;
-		SourceContexts.Emplace(MoveTemp(SourceVariable));
+		SourceVariable.BlueprintSetter = CompilerSourceCreatorContexts[FoundSourceCreatorContextIndex].SetterFunctionName;
+		CompilerSourceContexts.Emplace(MoveTemp(SourceVariable));
 	}
 
 	bAreSourceContextsValid = bAreSourcesCreatorValid;
@@ -363,7 +363,7 @@ void FMVVMViewBlueprintCompiler::CreateSourceLists(const FWidgetBlueprintCompile
 					SourceVariable.PropertyName = PropertyPath.GetWidgetName();
 					SourceVariable.DisplayName = Widget->GetDisplayLabel();
 					SourceVariable.CategoryName = TEXT("Widget");
-					Self->SourceContexts.Emplace(MoveTemp(SourceVariable));
+					Self->CompilerSourceContexts.Emplace(MoveTemp(SourceVariable));
 				}
 			}
 			else if (PropertyPath.IsFromViewModel())
@@ -402,7 +402,7 @@ void FMVVMViewBlueprintCompiler::CreateSourceLists(const FWidgetBlueprintCompile
 
 void FMVVMViewBlueprintCompiler::CreateFunctionsDeclaration(const FWidgetBlueprintCompilerContext::FCreateVariableContext& Context, UMVVMBlueprintView* BlueprintView)
 {
-	for (FCompilerSourceCreatorContext& SourceCreator : SourceCreatorContexts)
+	for (FCompilerSourceCreatorContext& SourceCreator : CompilerSourceCreatorContexts)
 	{
 		if (!SourceCreator.SetterFunctionName.IsEmpty() && SourceCreator.Type == ECompilerSourceCreatorType::ViewModel)
 		{
@@ -433,7 +433,7 @@ bool FMVVMViewBlueprintCompiler::PreCompile(UWidgetBlueprintGeneratedClass* Clas
 	}
 
 	const int32 NumBindings = BlueprintView->GetNumBindings();
-	Bindings.Reset(NumBindings*2);
+	CompilerBindings.Reset(NumBindings*2);
 	BindingSourceContexts.Reset(NumBindings*2);
 
 	PreCompileBindingSources(Class, BlueprintView);
@@ -575,7 +575,7 @@ bool FMVVMViewBlueprintCompiler::PreCompileSourceCreators(UWidgetBlueprintGenera
 		return false;
 	}
 
-	for (FCompilerSourceCreatorContext& SourceCreatorContext : SourceCreatorContexts)
+	for (FCompilerSourceCreatorContext& SourceCreatorContext : CompilerSourceCreatorContexts)
 	{
 		FMVVMViewClass_SourceCreator CompiledSourceCreator;
 
@@ -673,7 +673,7 @@ bool FMVVMViewBlueprintCompiler::CompileSourceCreators(const FCompiledBindingLib
 		return false;
 	}
 
-	for (const FCompilerSourceCreatorContext& SourceCreatorContext : SourceCreatorContexts)
+	for (const FCompilerSourceCreatorContext& SourceCreatorContext : CompilerSourceCreatorContexts)
 	{
 		if (SourceCreatorContext.Type == ECompilerSourceCreatorType::ViewModel)
 		{
@@ -879,11 +879,11 @@ bool FMVVMViewBlueprintCompiler::PreCompileBindings(UWidgetBlueprintGeneratedCla
 
 		FCompilerBinding NewBinding = AddBindingResult.StealValue();
 		NewBinding.BindingIndex = BindingSourceContext.BindingIndex;
-		NewBinding.SourceContextIndex = BindingSourceContext.CompilerSourceContextIndex;
+		NewBinding.CompilerSourceContextIndex = BindingSourceContext.CompilerSourceContextIndex;
 		NewBinding.FieldIdHandle = AddFieldResult.StealValue();
 		NewBinding.bSourceIsUserWidget = BindingSourceContext.bIsRootWidget;
 		NewBinding.bFieldIdNeeded = !IsOneTimeBinding(Binding.BindingType);
-		Bindings.Emplace(NewBinding);
+		CompilerBindings.Emplace(NewBinding);
 	}
 
 	return bIsBindingsValid;
@@ -897,14 +897,14 @@ bool FMVVMViewBlueprintCompiler::CompileBindings(const FCompiledBindingLibraryCo
 		return false;
 	}
 
-	for (const FCompilerBinding& CompileBinding : Bindings)
+	for (const FCompilerBinding& CompileBinding : CompilerBindings)
 	{
 		const FMVVMBlueprintViewBinding& ViewBinding = *BlueprintView->GetBindingAt(CompileBinding.BindingIndex);
 
 		FMVVMViewClass_CompiledBinding NewBinding;
 
-		check(CompileBinding.SourceContextIndex != INDEX_NONE);
-		NewBinding.SourcePropertyName = SourceContexts[CompileBinding.SourceContextIndex].PropertyName;
+		check(CompileBinding.CompilerSourceContextIndex != INDEX_NONE);
+		NewBinding.SourcePropertyName = CompilerSourceContexts[CompileBinding.CompilerSourceContextIndex].PropertyName;
 
 		const FMVVMVCompiledFieldId* CompiledFieldId = CompileResult.FieldIds.Find(CompileBinding.FieldIdHandle);
 		if (CompiledFieldId == nullptr && CompileBinding.bFieldIdNeeded)
@@ -933,6 +933,7 @@ bool FMVVMViewBlueprintCompiler::CompileBindings(const FCompiledBindingLibraryCo
 		NewBinding.Flags |= (IsForwardBinding(ViewBinding.BindingType)) ? FMVVMViewClass_CompiledBinding::EBindingFlags::ForwardBinding : 0;
 		NewBinding.Flags |= (ViewBinding.BindingType == EMVVMBindingMode::TwoWay) ? FMVVMViewClass_CompiledBinding::EBindingFlags::TwoWayBinding : 0;
 		NewBinding.Flags |= (IsOneTimeBinding(ViewBinding.BindingType)) ? FMVVMViewClass_CompiledBinding::EBindingFlags::OneTime : 0;
+		//NewBinding.Flags |= () ? FMVVMViewClass_CompiledBinding::EBindingFlags::RegistrationOptional : 0;
 
 		ViewExtension->CompiledBindings.Emplace(MoveTemp(NewBinding));
 	}
@@ -943,7 +944,7 @@ bool FMVVMViewBlueprintCompiler::CompileBindings(const FCompiledBindingLibraryCo
 
 const FMVVMViewBlueprintCompiler::FCompilerSourceCreatorContext* FMVVMViewBlueprintCompiler::FindViewModelSource(FGuid Id) const
 {
-	return SourceCreatorContexts.FindByPredicate([Id](const FCompilerSourceCreatorContext& Other)
+	return CompilerSourceCreatorContexts.FindByPredicate([Id](const FCompilerSourceCreatorContext& Other)
 		{
 			return Other.Type == ECompilerSourceCreatorType::ViewModel ? Other.ViewModelContext.GetViewModelId() == Id : false;
 		});
@@ -966,11 +967,11 @@ TValueOrError<FMVVMViewBlueprintCompiler::FBindingSourceContext, FString> FMVVMV
 		const FMVVMBlueprintViewModelContext* SourceViewModelContext = BlueprintView->FindViewModel(PropertyPath.GetViewModelId());
 		check(SourceViewModelContext);
 		const FName SourceName = SourceViewModelContext->GetViewModelName();
-		Result.CompilerSourceContextIndex = SourceContexts.IndexOfByPredicate([SourceName](const FCompilerSourceContext& Other) { return Other.PropertyName == SourceName; });
+		Result.CompilerSourceContextIndex = CompilerSourceContexts.IndexOfByPredicate([SourceName](const FCompilerSourceContext& Other) { return Other.PropertyName == SourceName; });
 		check(Result.CompilerSourceContextIndex != INDEX_NONE);
 
-		Result.SourceClass = SourceContexts[Result.CompilerSourceContextIndex].Class;
-		Result.PropertyPath = CreatePropertyPath(Class, SourceContexts[Result.CompilerSourceContextIndex].PropertyName, PropertyPath.GetFields());
+		Result.SourceClass = CompilerSourceContexts[Result.CompilerSourceContextIndex].Class;
+		Result.PropertyPath = CreatePropertyPath(Class, CompilerSourceContexts[Result.CompilerSourceContextIndex].PropertyName, PropertyPath.GetFields());
 	}
 	else if (PropertyPath.IsFromWidget())
 	{
@@ -984,10 +985,10 @@ TValueOrError<FMVVMViewBlueprintCompiler::FBindingSourceContext, FString> FMVVMV
 		}
 		else
 		{
-			Result.CompilerSourceContextIndex = SourceContexts.IndexOfByPredicate([SourceName](const FCompilerSourceContext& Other) { return Other.PropertyName == SourceName; });
+			Result.CompilerSourceContextIndex = CompilerSourceContexts.IndexOfByPredicate([SourceName](const FCompilerSourceContext& Other) { return Other.PropertyName == SourceName; });
 			check(Result.CompilerSourceContextIndex != INDEX_NONE);
-			Result.SourceClass = SourceContexts[Result.CompilerSourceContextIndex].Class;
-			Result.PropertyPath = CreatePropertyPath(Class, SourceContexts[Result.CompilerSourceContextIndex].PropertyName, PropertyPath.GetFields());
+			Result.SourceClass = CompilerSourceContexts[Result.CompilerSourceContextIndex].Class;
+			Result.PropertyPath = CreatePropertyPath(Class, CompilerSourceContexts[Result.CompilerSourceContextIndex].PropertyName, PropertyPath.GetFields());
 		}
 	}
 	else
@@ -1016,7 +1017,7 @@ TValueOrError<FMVVMViewBlueprintCompiler::FBindingSourceContext, FString> FMVVMV
 		Result.bIsRootWidget = false;
 
 		FString InterfacePath = FieldPathHelper::ToString(BindingInfo.NotifyFieldInterfacePath);
-		int32 SourceCreatorContextIndex = SourceCreatorContexts.IndexOfByPredicate([&InterfacePath](const FCompilerSourceCreatorContext& Other)
+		int32 SourceCreatorContextIndex = CompilerSourceCreatorContexts.IndexOfByPredicate([&InterfacePath](const FCompilerSourceCreatorContext& Other)
 			{
 				return (Other.Type == ECompilerSourceCreatorType::ViewModel)
 					? (Other.ViewModelContext.CreationType == EMVVMBlueprintViewModelContextCreationType::PropertyPath
@@ -1026,13 +1027,13 @@ TValueOrError<FMVVMViewBlueprintCompiler::FBindingSourceContext, FString> FMVVMV
 
 
 		// Add if the path doesn't already exists
-		if (!SourceCreatorContexts.IsValidIndex(SourceCreatorContextIndex))
+		if (!CompilerSourceCreatorContexts.IsValidIndex(SourceCreatorContextIndex))
 		{
 			return MakeError(TEXT("The field in the path is from a viewmodel that is not the root of the path. This is not yet supported."));
 		}
 
 		Result.PropertyPath.RemoveAt(0, BindingInfo.ViewModelIndex);
-		if (SourceCreatorContexts[SourceCreatorContextIndex].Type == ECompilerSourceCreatorType::ViewModel)
+		if (CompilerSourceCreatorContexts[SourceCreatorContextIndex].Type == ECompilerSourceCreatorType::ViewModel)
 		{
 			Result.CompilerSourceContextIndex = SourceCreatorContextIndex;
 		}
@@ -1062,10 +1063,10 @@ TArray<FMVVMConstFieldVariant> FMVVMViewBlueprintCompiler::CreateBindingDestinat
 		const FMVVMBlueprintViewModelContext* SourceViewModelContext = BlueprintView->FindViewModel(PropertyPath.GetViewModelId());
 		check(SourceViewModelContext);
 		FName DestinationName = SourceViewModelContext->GetViewModelName();
-		const int32 DestinationVariableContextIndex = SourceContexts.IndexOfByPredicate([DestinationName](const FCompilerSourceContext& Other) { return Other.PropertyName == DestinationName; });
+		const int32 DestinationVariableContextIndex = CompilerSourceContexts.IndexOfByPredicate([DestinationName](const FCompilerSourceContext& Other) { return Other.PropertyName == DestinationName; });
 		check(DestinationVariableContextIndex != INDEX_NONE);
 
-		return CreatePropertyPath(Class, SourceContexts[DestinationVariableContextIndex].PropertyName, PropertyPath.GetFields());
+		return CreatePropertyPath(Class, CompilerSourceContexts[DestinationVariableContextIndex].PropertyName, PropertyPath.GetFields());
 	}
 	else if (PropertyPath.IsFromWidget())
 	{
@@ -1078,10 +1079,10 @@ TArray<FMVVMConstFieldVariant> FMVVMViewBlueprintCompiler::CreateBindingDestinat
 		}
 		else
 		{
-			const int32 DestinationVariableContextIndex = SourceContexts.IndexOfByPredicate([DestinationName](const FCompilerSourceContext& Other) { return Other.PropertyName == DestinationName; });
+			const int32 DestinationVariableContextIndex = CompilerSourceContexts.IndexOfByPredicate([DestinationName](const FCompilerSourceContext& Other) { return Other.PropertyName == DestinationName; });
 			check(DestinationVariableContextIndex != INDEX_NONE);
 
-			return CreatePropertyPath(Class, SourceContexts[DestinationVariableContextIndex].PropertyName, PropertyPath.GetFields());
+			return CreatePropertyPath(Class, CompilerSourceContexts[DestinationVariableContextIndex].PropertyName, PropertyPath.GetFields());
 		}
 	}
 	else
diff --git a/Engine/Plugins/Runtime/ModelViewViewModel/Source/ModelViewViewModelBlueprint/Private/MVVMViewBlueprintCompiler.h b/Engine/Plugins/Runtime/ModelViewViewModel/Source/ModelViewViewModelBlueprint/Private/MVVMViewBlueprintCompiler.h
--- a/Engine/Plugins/Runtime/ModelViewViewModel/Source/ModelViewViewModelBlueprint/Private/MVVMViewBlueprintCompiler.h
+++ b/Engine/Plugins/Runtime/ModelViewViewModel/Source/ModelViewViewModelBlueprint/Private/MVVMViewBlueprintCompiler.h
@@ -92,7 +92,7 @@ private:
 
 		bool bExposeOnSpawn = false;
 	};
-	TArray<FCompilerSourceContext> SourceContexts;
+	TArray<FCompilerSourceContext> CompilerSourceContexts;
 
 	enum class ECompilerSourceCreatorType
 	{
@@ -106,7 +106,7 @@ private:
 		FString SetterFunctionName;
 		UEdGraph* SetterGraph = nullptr;
 	};
-	TArray<FCompilerSourceCreatorContext> SourceCreatorContexts;
+	TArray<FCompilerSourceCreatorContext> CompilerSourceCreatorContexts;
 
 	enum class ECompilerBindingType
 	{
@@ -115,7 +115,7 @@ private:
 	struct FCompilerBinding
 	{
 		int32 BindingIndex = INDEX_NONE;
-		int32 SourceContextIndex = INDEX_NONE;
+		int32 CompilerSourceContextIndex = INDEX_NONE;
 		bool bSourceIsUserWidget = false;
 		bool bFieldIdNeeded = false;
 
@@ -132,7 +132,7 @@ private:
 		FCompiledBindingLibraryCompiler::FFieldPathHandle DestinationWrite;
 		FCompiledBindingLibraryCompiler::FFieldPathHandle ConversionFunction;
 	};
-	TArray<FCompilerBinding> Bindings;
+	TArray<FCompilerBinding> CompilerBindings;
 
 	struct FBindingSourceContext
 	{
