MVVM: Fix remove graph function by moving it to the transient pacakge.

diff --git a/Engine/Plugins/Runtime/ModelViewViewModel/Source/ModelViewViewModelBlueprint/Private/ViewModel/MVVMViewModelBlueprintCompiler.cpp b/Engine/Plugins/Runtime/ModelViewViewModel/Source/ModelViewViewModelBlueprint/Private/ViewModel/MVVMViewModelBlueprintCompiler.cpp
--- a/Engine/Plugins/Runtime/ModelViewViewModel/Source/ModelViewViewModelBlueprint/Private/ViewModel/MVVMViewModelBlueprintCompiler.cpp
+++ b/Engine/Plugins/Runtime/ModelViewViewModel/Source/ModelViewViewModelBlueprint/Private/ViewModel/MVVMViewModelBlueprintCompiler.cpp
@@ -286,7 +286,7 @@ void FViewModelBlueprintCompilerContext::FinishCompilingClass(UClass* Class)
 		{
 			ensureMsgf(Graph->HasAnyFlags(RF_Transient), TEXT("The graph should be temporary and should be generated automatically."));
 			// GC may not have clean the graph (GC doesn't run when bRegenerateSkeletonOnly is on)
-			Graph->Rename(nullptr, Graph->GetOuter(), REN_DoNotDirty | REN_ForceNoResetLoaders);
+			Graph->Rename(nullptr, GetTransientPackage(), REN_DoNotDirty | REN_ForceNoResetLoaders);
 		}
 	};
 
