Move function from the MVVM editor module to the blueprint compiler module.

diff --git a/Engine/Plugins/Runtime/ModelViewViewModel/Source/ModelViewViewModelBlueprint/Private/Bindings/MVVMConversionFunctionHelper.cpp b/Engine/Plugins/Runtime/ModelViewViewModel/Source/ModelViewViewModelBlueprint/Private/Bindings/MVVMConversionFunctionHelper.cpp
--- a/Engine/Plugins/Runtime/ModelViewViewModel/Source/ModelViewViewModelBlueprint/Private/Bindings/MVVMConversionFunctionHelper.cpp
+++ b/Engine/Plugins/Runtime/ModelViewViewModel/Source/ModelViewViewModelBlueprint/Private/Bindings/MVVMConversionFunctionHelper.cpp
@@ -8,10 +8,12 @@
 #include "K2Node_PromotableOperator.h"
 #include "K2Node_VariableGet.h"
 #include "Kismet2/BlueprintEditorUtils.h"
+#include "Misc/StringBuilder.h"
 
 #include "Components/Widget.h"
 #include "MVVMBlueprintView.h"
 #include "MVVMBlueprintViewModelContext.h"
+#include "MVVMWidgetBlueprintExtension_View.h"
 #include "MVVMViewModelBase.h"
 #include "View/MVVMView.h"
 #include "WidgetBlueprint.h"
@@ -94,4 +96,43 @@ TArray<FMVVMBlueprintPropertyPath> FindAllPropertyPathInGraph(const UEdGraph* Gr
 	return Result;
 }
 
+
+FName GetWrapperName(const UMVVMBlueprintView* View, const FMVVMBlueprintViewBinding& Binding, bool bSourceToDestination)
+{
+	TStringBuilder<256> StringBuilder;
+	StringBuilder << TEXT("__");
+	StringBuilder << Binding.GetFName(View);
+	StringBuilder << (bSourceToDestination ? TEXT("_SourceToDest") : TEXT("_DestToSource"));
+
+	return FName(StringBuilder.ToString());
+}
+
+
+namespace Private
+{
+UEdGraph* FindExistingConversionFunctionWrapper(const UWidgetBlueprint* WidgetBlueprint, FName WrapperName)
+{
+	const TObjectPtr<UEdGraph>* Result = WidgetBlueprint->FunctionGraphs.FindByPredicate([WrapperName](const UEdGraph* GraphPtr) { return GraphPtr->GetFName() == WrapperName; });
+	return Result ? Result->Get() : nullptr;
+}
+} // namespace
+
+UEdGraph* GetGraph(const UWidgetBlueprint* WidgetBlueprint, const FMVVMBlueprintViewBinding& Binding, bool bSourceToDestination)
+{
+	check(WidgetBlueprint);
+
+	const UMVVMBlueprintView* View = nullptr;
+	if (UMVVMWidgetBlueprintExtension_View* ExtensionView = UMVVMWidgetBlueprintExtension_View::GetExtension<UMVVMWidgetBlueprintExtension_View>(WidgetBlueprint))
+	{
+		View = ExtensionView->GetBlueprintView();
+	}
+
+	if (View)
+	{
+		const FName WrapperName = GetWrapperName(View, Binding, bSourceToDestination);
+		return Private::FindExistingConversionFunctionWrapper(WidgetBlueprint, WrapperName);
+	}
+	return nullptr;
+}
+
 } //namespace
\ No newline at end of file
diff --git a/Engine/Plugins/Runtime/ModelViewViewModel/Source/ModelViewViewModelBlueprint/Public/Bindings/MVVMConversionFunctionHelper.h b/Engine/Plugins/Runtime/ModelViewViewModel/Source/ModelViewViewModelBlueprint/Public/Bindings/MVVMConversionFunctionHelper.h
--- a/Engine/Plugins/Runtime/ModelViewViewModel/Source/ModelViewViewModelBlueprint/Public/Bindings/MVVMConversionFunctionHelper.h
+++ b/Engine/Plugins/Runtime/ModelViewViewModel/Source/ModelViewViewModelBlueprint/Public/Bindings/MVVMConversionFunctionHelper.h
@@ -8,6 +8,8 @@
 class UClass;
 class UEdGraph;
 class UMVVMBlueprintView;
+class UWidgetBlueprint;
+struct FMVVMBlueprintViewBinding;
 
 namespace UE::MVVM::ConversionFunctionHelper
 {
@@ -15,4 +17,10 @@ namespace UE::MVVM::ConversionFunctionHelper
 	/** Find all BlueprintPropertyPath used in the Graph. */
 	MODELVIEWVIEWMODELBLUEPRINT_API TArray<FMVVMBlueprintPropertyPath> FindAllPropertyPathInGraph(const UEdGraph* Graph, const UMVVMBlueprintView* View, UClass* Class);
 
+	/** Get the name of the conversion function wrapper if autogenerated. */
+	MODELVIEWVIEWMODELBLUEPRINT_API FName GetWrapperName(const UMVVMBlueprintView* View, const FMVVMBlueprintViewBinding& Binding, bool bSourceToDestination);
+
+	/** Get an existing conversion graph for the corresponding binding if it exist. */
+	MODELVIEWVIEWMODELBLUEPRINT_API UEdGraph* GetGraph(const UWidgetBlueprint* WidgetBlueprint, const FMVVMBlueprintViewBinding& Binding, bool bSourceToDestination);
+
 } //namespace
\ No newline at end of file
diff --git a/Engine/Plugins/Runtime/ModelViewViewModel/Source/ModelViewViewModelEditor/Private/MVVMEditorSubsystem.cpp b/Engine/Plugins/Runtime/ModelViewViewModel/Source/ModelViewViewModelEditor/Private/MVVMEditorSubsystem.cpp
--- a/Engine/Plugins/Runtime/ModelViewViewModel/Source/ModelViewViewModelEditor/Private/MVVMEditorSubsystem.cpp
+++ b/Engine/Plugins/Runtime/ModelViewViewModel/Source/ModelViewViewModelEditor/Private/MVVMEditorSubsystem.cpp
@@ -4,6 +4,7 @@
 
 #include "Algo/Reverse.h"
 #include "Bindings/MVVMBindingHelper.h"
+#include "Bindings/MVVMConversionFunctionHelper.h"
 #include "Blueprint/WidgetTree.h"
 #include "BlueprintActionDatabase.h"
 #include "BlueprintNodeSpawner.h"
@@ -142,29 +143,18 @@ namespace UE::MVVM::Private
 		return ConversionNode;
 	}
 
-	const TObjectPtr<UEdGraph>* FindExistingConversionFunctionWrapper(const UWidgetBlueprint* WidgetBlueprint, FName WrapperName)
+	UEdGraph* FindExistingConversionFunctionWrapper(UWidgetBlueprint* WidgetBlueprint, FName WrapperName)
 	{
-		const TObjectPtr<UEdGraph>* ExistingGraph = WidgetBlueprint->FunctionGraphs.FindByPredicate(
-			[&WrapperName](TObjectPtr<UEdGraph> GraphPtr) -> bool
-			{
-				if (UEdGraph* Graph = GraphPtr.Get())
-				{
-					return Graph->GetFName() == WrapperName;
-				}
-				return false;
-			});
-
-		return ExistingGraph;
+		TObjectPtr<UEdGraph>* Result = WidgetBlueprint->FunctionGraphs.FindByPredicate([WrapperName](const UEdGraph* GraphPtr) { return GraphPtr->GetFName() == WrapperName; });
+		return Result ? Result->Get() : nullptr;
 	}
 
 	void RemoveConversionFunctionWrapper(UWidgetBlueprint* WidgetBlueprint, FName WrapperName)
 	{
-		const TObjectPtr<UEdGraph>* ExistingGraph = FindExistingConversionFunctionWrapper(WidgetBlueprint, WrapperName);
-		if (ExistingGraph != nullptr)
+		if (UEdGraph* Graph = FindExistingConversionFunctionWrapper(WidgetBlueprint, WrapperName))
 		{
-			UEdGraph* Graph = ExistingGraph->Get();
-			Graph->Rename(nullptr, Graph->GetOuter(), REN_DoNotDirty | REN_ForceNoResetLoaders);
-			WidgetBlueprint->FunctionGraphs.Remove(Graph);
+			Graph->Rename(nullptr, GetTransientPackage(), REN_DoNotDirty | REN_ForceNoResetLoaders);
+			WidgetBlueprint->FunctionGraphs.RemoveSingle(Graph);
 		}
 	}
 } //namespace
@@ -173,13 +163,7 @@ FName UMVVMEditorSubsystem::GetConversionFunctionWrapperName(const UWidgetBluepr
 {
 	if (UMVVMBlueprintView* View = GetView(WidgetBlueprint))
 	{
-		TStringBuilder<256> StringBuilder;
-		StringBuilder << TEXT("__");
-		Binding.GetFName(View).AppendString(StringBuilder);
-
-		StringBuilder << (bSourceToDestination ? TEXT("_SourceToDest") : TEXT("_DestToSource"));
-
-		return FName(StringBuilder.ToString());
+		return UE::MVVM::ConversionFunctionHelper::GetWrapperName(View, Binding, bSourceToDestination);
 	}
 
 	return FName();
@@ -305,10 +289,7 @@ void UMVVMEditorSubsystem::RemoveBinding(UWidgetBlueprint* WidgetBlueprint, cons
 
 UEdGraph* UMVVMEditorSubsystem::GetConversionFunctionGraph(const UWidgetBlueprint* WidgetBlueprint, const FMVVMBlueprintViewBinding& Binding, bool bSourceToDestination) const
 {
-	const FName WrapperName = GetConversionFunctionWrapperName(WidgetBlueprint, Binding, bSourceToDestination);
-	const TObjectPtr<UEdGraph>* EdGraphPtr = UE::MVVM::Private::FindExistingConversionFunctionWrapper(WidgetBlueprint, WrapperName);
-
-	return EdGraphPtr != nullptr ? EdGraphPtr->Get() : nullptr;
+	return UE::MVVM::ConversionFunctionHelper::GetGraph(WidgetBlueprint, Binding, bSourceToDestination);
 }
 
 UEdGraph* UMVVMEditorSubsystem::CreateConversionFunctionWrapperGraph(UWidgetBlueprint* WidgetBlueprint, const FMVVMBlueprintViewBinding& Binding, const UFunction* ConversionFunction, bool bSourceToDestination)
