MVVM: Remove the blueprintable on ViewModelBase for a test and forgot to add it back.
#rnx

diff --git a/Engine/Plugins/Runtime/ModelViewViewModel/Source/ModelViewViewModel/Public/MVVMViewModelBase.h b/Engine/Plugins/Runtime/ModelViewViewModel/Source/ModelViewViewModel/Public/MVVMViewModelBase.h
--- a/Engine/Plugins/Runtime/ModelViewViewModel/Source/ModelViewViewModel/Public/MVVMViewModelBase.h
+++ b/Engine/Plugins/Runtime/ModelViewViewModel/Source/ModelViewViewModel/Public/MVVMViewModelBase.h
@@ -34,7 +34,7 @@
 
 
 /** Base class for MVVM viewmodel. */
-UCLASS(Abstract, DisplayName="MVVM ViewModel")
+UCLASS(Blueprintable, Abstract, DisplayName="MVVM ViewModel")
 class MODELVIEWVIEWMODEL_API UMVVMViewModelBase : public UObject, public INotifyFieldValueChanged
 {
 	GENERATED_BODY()
