MVVM: Compile complex conversion functon and execute them as a regular binding.
#prefligh 63062f5c5366f61a42bced80

diff --git a/Engine/Plugins/Runtime/ModelViewViewModel/Source/ModelViewViewModel/Private/Bindings/MVVMBindingHelper.cpp b/Engine/Plugins/Runtime/ModelViewViewModel/Source/ModelViewViewModel/Private/Bindings/MVVMBindingHelper.cpp
--- a/Engine/Plugins/Runtime/ModelViewViewModel/Source/ModelViewViewModel/Private/Bindings/MVVMBindingHelper.cpp
+++ b/Engine/Plugins/Runtime/ModelViewViewModel/Source/ModelViewViewModel/Private/Bindings/MVVMBindingHelper.cpp
@@ -109,11 +109,11 @@ namespace UE::MVVM::BindingHelper
 	}
 
 
-	bool IsValidForRuntimeConversion(const UFunction* InFunction)
+	bool IsValidForSimpleRuntimeConversion(const UFunction* InFunction)
 	{
 		if (Private::IsValidCommon(InFunction)
 			&& (InFunction->HasAllFunctionFlags(FUNC_Static) || InFunction->HasAllFunctionFlags(FUNC_Const | FUNC_BlueprintPure))
-			&& InFunction->NumParms == 2)
+			&& InFunction->NumParms <= 2)
 		{
 			const FProperty* ReturnProperty = GetReturnProperty(InFunction);
 			const FProperty* FirstArgumentProperty = GetFirstArgumentProperty(InFunction);
@@ -121,6 +121,20 @@ namespace UE::MVVM::BindingHelper
 		}
 		return false;
 	}
+	
+	
+	bool IsValidForComplexRuntimeConversion(const UFunction* InFunction)
+	{
+		if (Private::IsValidCommon(InFunction)
+			&& (InFunction->HasAllFunctionFlags(FUNC_Static) || InFunction->HasAllFunctionFlags(FUNC_Const | FUNC_BlueprintPure))
+			&& InFunction->NumParms <= 1)
+		{
+			const FProperty* ReturnProperty = GetReturnProperty(InFunction);
+			const FProperty* FirstArgumentProperty = GetFirstArgumentProperty(InFunction);
+			return ReturnProperty && FirstArgumentProperty == nullptr;
+		}
+		return false;
+	}
 
 
 #if WITH_EDITOR
@@ -388,11 +402,6 @@ namespace UE::MVVM::BindingHelper
 			return MakeError(FString::Printf(TEXT("The function '%s' is not static or is not const and pure."), *InFunction->GetName()));
 		}
 
-		if (InFunction->NumParms < 2)
-		{
-			return MakeError(FString::Printf(TEXT("The function '%s' does not have the correct number of arguments."), *InFunction->GetName()));
-		}
-
 		const FProperty* ReturnProperty = GetReturnProperty(InFunction);
 		if (ReturnProperty == nullptr)
 		{
@@ -618,6 +627,77 @@ namespace UE::MVVM::BindingHelper
 	}
 
 
+	void ExecuteBinding_NoCheck(const FFieldContext& Destination, const FFunctionContext& ConversionFunction)
+	{
+		check(!Destination.GetObjectVariant().IsNull());
+		check(!Destination.GetFieldVariant().IsEmpty());
+		check(ConversionFunction.GetFunction() && ConversionFunction.GetObject());
+
+		void* ConversionFunctionDataPtr = FMemory_Alloca_Aligned(ConversionFunction.GetFunction()->ParmsSize, ConversionFunction.GetFunction()->GetMinAlignment());
+		for (TFieldIterator<FProperty> It(ConversionFunction.GetFunction()); It && It->HasAnyPropertyFlags(CPF_Parm); ++It)
+		{
+			It->InitializeValue_InContainer(ConversionFunctionDataPtr);
+		}
+
+		ConversionFunction.GetObject()->ProcessEvent(ConversionFunction.GetFunction(), ConversionFunctionDataPtr);
+
+		{
+			const FProperty* ReturnConversionProperty = GetReturnProperty(ConversionFunction.GetFunction());
+			check(ReturnConversionProperty);
+
+			const bool bIsDestinationBindingIsProperty = Destination.GetFieldVariant().IsProperty();
+			const FProperty* SetterType = bIsDestinationBindingIsProperty ? Destination.GetFieldVariant().GetProperty() : GetFirstArgumentProperty(Destination.GetFieldVariant().GetFunction());
+			check(SetterType);
+			check(ArePropertiesCompatible(ReturnConversionProperty, SetterType));
+
+			if (Private::IsNumericConversionRequired(ReturnConversionProperty, SetterType))
+			{
+				// we need to do a copy because we may destroy the ReturnConversionProperty (imagine the return type is a TArray)
+				const int32 AllocationSize = FMath::Max(SetterType->GetSize(), ReturnConversionProperty->GetSize());
+				const int32 AllocationMinAlignment = FMath::Max(SetterType->GetMinAlignment(), ReturnConversionProperty->GetMinAlignment());
+				void* SetterDataPtr = FMemory_Alloca_Aligned(AllocationSize, AllocationMinAlignment);
+				SetterType->InitializeValue(SetterDataPtr); // probably not needed since they are double/float
+
+				ReturnConversionProperty->CopyCompleteValue(SetterDataPtr, ReturnConversionProperty->ContainerPtrToValuePtr<void>(ConversionFunctionDataPtr));
+				Private::ConvertNumeric(ReturnConversionProperty, SetterType, SetterDataPtr);
+
+				if (bIsDestinationBindingIsProperty)
+				{
+					SetterType->SetValue_InContainer(Destination.GetObjectVariant().GetData(), SetterDataPtr);
+				}
+				else
+				{
+					check(SetterType->GetSize() == Destination.GetFieldVariant().GetFunction()->ParmsSize); // only 1 param
+					check(Destination.GetObjectVariant().IsUObject());
+					Destination.GetObjectVariant().GetUObject()->ProcessEvent(Destination.GetFieldVariant().GetFunction(), SetterDataPtr);
+				}
+
+				SetterType->DestroyValue(SetterDataPtr);
+			}
+			else
+			{
+				// Re use the same buffer, no need to create a new copy
+				void* DestinationDataPtr = ReturnConversionProperty->ContainerPtrToValuePtr<void>(ConversionFunctionDataPtr);
+				if (bIsDestinationBindingIsProperty)
+				{
+					SetterType->SetValue_InContainer(Destination.GetObjectVariant().GetData(), DestinationDataPtr);
+				}
+				else
+				{
+					check(SetterType->GetSize() == Destination.GetFieldVariant().GetFunction()->ParmsSize); // only 1 param
+					check(Destination.GetObjectVariant().IsUObject());
+					Destination.GetObjectVariant().GetUObject()->ProcessEvent(Destination.GetFieldVariant().GetFunction(), DestinationDataPtr);
+				}
+			}
+		}
+
+		for (TFieldIterator<FProperty> It(ConversionFunction.GetFunction()); It && It->HasAnyPropertyFlags(CPF_Parm); ++It)
+		{
+			It->DestroyValue_InContainer(ConversionFunctionDataPtr);
+		}
+	}
+
+
 	void ExecuteBinding_NoCheck(const FFieldContext& Source, const FFieldContext& Destination, const FFunctionContext& ConversionFunction)
 	{
 		check(!Source.GetObjectVariant().IsNull() && !Destination.GetObjectVariant().IsNull());
diff --git a/Engine/Plugins/Runtime/ModelViewViewModel/Source/ModelViewViewModel/Private/Bindings/MVVMCompiledBindingLibrary.cpp b/Engine/Plugins/Runtime/ModelViewViewModel/Source/ModelViewViewModel/Private/Bindings/MVVMCompiledBindingLibrary.cpp
--- a/Engine/Plugins/Runtime/ModelViewViewModel/Source/ModelViewViewModel/Private/Bindings/MVVMCompiledBindingLibrary.cpp
+++ b/Engine/Plugins/Runtime/ModelViewViewModel/Source/ModelViewViewModel/Private/Bindings/MVVMCompiledBindingLibrary.cpp
@@ -93,11 +93,11 @@ void FMVVMCompiledBindingLibrary::Load()
 }
 
 
-TValueOrError<void, FMVVMCompiledBindingLibrary::EExecutionFailingReason> FMVVMCompiledBindingLibrary::Execute(UObject* InExecutionSource, const FMVVMVCompiledBinding& InBinding) const
+TValueOrError<void, FMVVMCompiledBindingLibrary::EExecutionFailingReason> FMVVMCompiledBindingLibrary::Execute(UObject* InExecutionSource, const FMVVMVCompiledBinding& InBinding, EConversionFunctionType InFunctionType) const
 {
 	check(InExecutionSource);
 
-	return ExecuteImpl(InExecutionSource, InBinding, nullptr);
+	return ExecuteImpl(InExecutionSource, InBinding, nullptr, InFunctionType);
 }
 
 
@@ -106,11 +106,11 @@ TValueOrError<void, FMVVMCompiledBindingLibrary::EExecutionFailingReason> FMVVMC
 	check(InExecutionSource);
 	check(InSource);
 
-	return ExecuteImpl(InExecutionSource, InBinding, InSource);
+	return ExecuteImpl(InExecutionSource, InBinding, InSource, EConversionFunctionType::Simple);
 }
 
 
-TValueOrError<void, FMVVMCompiledBindingLibrary::EExecutionFailingReason> FMVVMCompiledBindingLibrary::ExecuteImpl(UObject* InExecutionSource, const FMVVMVCompiledBinding& InBinding, UObject* InSource) const
+TValueOrError<void, FMVVMCompiledBindingLibrary::EExecutionFailingReason> FMVVMCompiledBindingLibrary::ExecuteImpl(UObject* InExecutionSource, const FMVVMVCompiledBinding& InBinding, UObject* InSource, EConversionFunctionType InFunctionType) const
 {
 #if WITH_EDITORONLY_DATA
 	const bool bIsValidBinding = InBinding.CompiledBindingLibraryId == CompiledBindingLibraryId;
@@ -122,45 +122,48 @@ TValueOrError<void, FMVVMCompiledBindingLibrary::EExecutionFailingReason> FMVVMC
 #endif
 
 	UE::MVVM::FFieldContext ValidSource;
-	if (InSource)
+	if (InFunctionType != EConversionFunctionType::Complex)
 	{
-#if DO_CHECK
-		TValueOrError<UE::MVVM::FFieldContext, void> TestSource = EvaluateFieldPath(InExecutionSource, InBinding.SourceFieldPath);
-		if (TestSource.HasError())
-		{
-			ensureAlwaysMsgf(false, TEXT("The source was provided but it was not able to fetch it."));
-			return MakeError(EExecutionFailingReason::InvalidSource);
-		}
-		if (!TestSource.GetValue().GetObjectVariant().IsUObject())
-		{
-			ensureAlwaysMsgf(false, TEXT("The source was provided as a UObject but the fetched source is not a UObject."));
-			return MakeError(EExecutionFailingReason::InvalidSource);
-		}
-		if (InSource != TestSource.GetValue().GetObjectVariant().GetUObject())
+		if (InSource)
 		{
-			ensureAlwaysMsgf(false, TEXT("The source was provided and it's not the same as the fetched source."));
-			return MakeError(EExecutionFailingReason::InvalidSource);
-		}
-		checkf(TestSource.GetValue().GetObjectVariant().GetData() != nullptr, TEXT("Tested in EvaluateFieldPath"));
+#if DO_CHECK
+			TValueOrError<UE::MVVM::FFieldContext, void> TestSource = EvaluateFieldPath(InExecutionSource, InBinding.SourceFieldPath);
+			if (TestSource.HasError())
+			{
+				ensureAlwaysMsgf(false, TEXT("The source was provided but it was not able to fetch it."));
+				return MakeError(EExecutionFailingReason::InvalidSource);
+			}
+			if (!TestSource.GetValue().GetObjectVariant().IsUObject())
+			{
+				ensureAlwaysMsgf(false, TEXT("The source was provided as a UObject but the fetched source is not a UObject."));
+				return MakeError(EExecutionFailingReason::InvalidSource);
+			}
+			if (InSource != TestSource.GetValue().GetObjectVariant().GetUObject())
+			{
+				ensureAlwaysMsgf(false, TEXT("The source was provided and it's not the same as the fetched source."));
+				return MakeError(EExecutionFailingReason::InvalidSource);
+			}
+			checkf(TestSource.GetValue().GetObjectVariant().GetData() != nullptr, TEXT("Tested in EvaluateFieldPath"));
 #endif
 
-		TValueOrError<UE::MVVM::FMVVMFieldVariant, void> FinalPath = GetFinalFieldFromPathImpl(InBinding.SourceFieldPath);
-		if (FinalPath.HasError())
-		{
-			return MakeError(EExecutionFailingReason::InvalidSource);
+			TValueOrError<UE::MVVM::FMVVMFieldVariant, void> FinalPath = GetFinalFieldFromPathImpl(InBinding.SourceFieldPath);
+			if (FinalPath.HasError())
+			{
+				return MakeError(EExecutionFailingReason::InvalidSource);
+			}
+			ValidSource = UE::MVVM::FFieldContext(InSource, FinalPath.StealValue());
 		}
-		ValidSource = UE::MVVM::FFieldContext(InSource, FinalPath.StealValue());
-	}
-	else
-	{
-		TValueOrError<UE::MVVM::FFieldContext, void> FoundSource = EvaluateFieldPath(InExecutionSource, InBinding.SourceFieldPath);
-		if (FoundSource.HasError())
+		else
 		{
-			return MakeError(EExecutionFailingReason::InvalidSource);
-		}
-		checkf(FoundSource.GetValue().GetObjectVariant().GetData() != nullptr, TEXT("Tested in EvaluateFieldPath"));
+			TValueOrError<UE::MVVM::FFieldContext, void> FoundSource = EvaluateFieldPath(InExecutionSource, InBinding.SourceFieldPath);
+			if (FoundSource.HasError())
+			{
+				return MakeError(EExecutionFailingReason::InvalidSource);
+			}
+			checkf(FoundSource.GetValue().GetObjectVariant().GetData() != nullptr, TEXT("Tested in EvaluateFieldPath"));
 
-		ValidSource = FoundSource.StealValue();
+			ValidSource = FoundSource.StealValue();
+		}
 	}
 
 	TValueOrError<UE::MVVM::FFieldContext, void> FoundDestination = EvaluateFieldPath(InExecutionSource, InBinding.DestinationFieldPath);
@@ -193,22 +196,25 @@ TValueOrError<void, FMVVMCompiledBindingLibrary::EExecutionFailingReason> FMVVMC
 		}
 	}
 
-	return ExecuteImpl(ValidSource, FoundDestination.GetValue(), ConversionFunction);
+	return ExecuteImpl(ValidSource, FoundDestination.GetValue(), ConversionFunction, InFunctionType);
 }
 
 
-TValueOrError<void, FMVVMCompiledBindingLibrary::EExecutionFailingReason> FMVVMCompiledBindingLibrary::ExecuteImpl(UE::MVVM::FFieldContext& Source, UE::MVVM::FFieldContext& Destination, UE::MVVM::FFunctionContext& ConversionFunction) const
+TValueOrError<void, FMVVMCompiledBindingLibrary::EExecutionFailingReason> FMVVMCompiledBindingLibrary::ExecuteImpl(UE::MVVM::FFieldContext& Source, UE::MVVM::FFieldContext& Destination, UE::MVVM::FFunctionContext& ConversionFunction, EConversionFunctionType FunctionType) const
 {
-	check(!Source.GetObjectVariant().IsEmpty());
+	if (FunctionType != EConversionFunctionType::Complex)
+	{
+		check(!Source.GetObjectVariant().IsEmpty());
+		check(!Source.GetObjectVariant().IsNull());
+		check(!Source.GetFieldVariant().IsEmpty());
+		if (!Source.GetObjectVariant().GetOwner()->IsChildOf(Source.GetFieldVariant().GetOwner()))
+		{
+			return MakeError(EExecutionFailingReason::InvalidCast);
+		}
+	}
 	check(!Destination.GetObjectVariant().IsEmpty());
-	check(!Source.GetObjectVariant().IsNull());
 	check(!Destination.GetObjectVariant().IsNull());
-	check(!Source.GetFieldVariant().IsEmpty());
 	check(!Destination.GetFieldVariant().IsEmpty());
-	if (!Source.GetObjectVariant().GetOwner()->IsChildOf(Source.GetFieldVariant().GetOwner()))
-	{
-		return MakeError(EExecutionFailingReason::InvalidCast);
-	}
 	if (!Destination.GetObjectVariant().GetOwner()->IsChildOf(Destination.GetFieldVariant().GetOwner()))
 	{
 		return MakeError(EExecutionFailingReason::InvalidCast);
@@ -216,7 +222,14 @@ TValueOrError<void, FMVVMCompiledBindingLibrary::EExecutionFailingReason> FMVVMC
 
 	if (ConversionFunction.GetFunction())
 	{
-		UE::MVVM::BindingHelper::ExecuteBinding_NoCheck(Source, Destination, ConversionFunction);
+		if (FunctionType == EConversionFunctionType::Complex)
+		{
+			UE::MVVM::BindingHelper::ExecuteBinding_NoCheck(Destination, ConversionFunction);
+		}
+		else
+		{
+			UE::MVVM::BindingHelper::ExecuteBinding_NoCheck(Source, Destination, ConversionFunction);
+		}
 	}
 	else
 	{
diff --git a/Engine/Plugins/Runtime/ModelViewViewModel/Source/ModelViewViewModel/Private/MVVMSubsystem.cpp b/Engine/Plugins/Runtime/ModelViewViewModel/Source/ModelViewViewModel/Private/MVVMSubsystem.cpp
--- a/Engine/Plugins/Runtime/ModelViewViewModel/Source/ModelViewViewModel/Private/MVVMSubsystem.cpp
+++ b/Engine/Plugins/Runtime/ModelViewViewModel/Source/ModelViewViewModel/Private/MVVMSubsystem.cpp
@@ -318,36 +318,53 @@ TValueOrError<bool, FString> UMVVMSubsystem::IsBindingValid(FConstDirectionalBin
 			return MakeError(ReturnResult.StealError());
 		}
 
-		TValueOrError<TArray<const FProperty*>, FString> ArgumentsResult = UE::MVVM::BindingHelper::TryGetArgumentsForConversionFunction(Args.ConversionFunction);
-		if (ArgumentsResult.HasError())
+		const bool bIsSimpleConversionFunction = UE::MVVM::BindingHelper::IsValidForSimpleRuntimeConversion(Args.ConversionFunction);
+		const bool bIsComplexConversionFunction = UE::MVVM::BindingHelper::IsValidForComplexRuntimeConversion(Args.ConversionFunction);
+		if (!bIsSimpleConversionFunction && !bIsComplexConversionFunction)
 		{
-			return MakeError(ArgumentsResult.StealError());
+			return MakeError(FString::Printf(TEXT("The conversion function '%s' is not valid as a conversion function"), *Args.ConversionFunction->GetName()));
 		}
 
-		// The compiled version should look like Setter(Conversion(Getter())).
-		const FProperty* ReturnProperty = ReturnResult.GetValue();
-		if (!UE::MVVM::BindingHelper::ArePropertiesCompatible(ReturnProperty, DestinationProperty))
+		// The simple compiled version look like Setter(Conversion(Getter())).
+		if (bIsSimpleConversionFunction)
 		{
-			return MakeError(FString::Printf(TEXT("The destination property '%s' (%s) does not match the return type of the conversion function (%s)."), *DestinationProperty->GetName(), *DestinationProperty->GetCPPType(), *ReturnProperty->GetCPPType()));
-		}
+			TValueOrError<TArray<const FProperty*>, FString> ArgumentsResult = UE::MVVM::BindingHelper::TryGetArgumentsForConversionFunction(Args.ConversionFunction);
+			if (ArgumentsResult.HasError())
+			{
+				return MakeError(ArgumentsResult.StealError());
+			}
+			if (ArgumentsResult.GetValue().Num() != 1)
+			{
+				return MakeError(TEXT("Internal error. The number of arguments should be 1."));
+			}
 
-		bool bAnyCompatible = false;
+			bool bAnyCompatible = false;
+			const TArray<const FProperty*>& ArgumentProperties = ArgumentsResult.GetValue();
+			for (const FProperty* ArgumentProperty : ArgumentProperties)
+			{
+				if (UE::MVVM::BindingHelper::ArePropertiesCompatible(SourceProperty, ArgumentProperty))
+				{
+					bAnyCompatible = true;
+					break;
+				}
+			}
 
-		const TArray<const FProperty*>& ArgumentProperties = ArgumentsResult.GetValue();
-		for (const FProperty* ArgumentProperty : ArgumentProperties)
-		{
-			if (UE::MVVM::BindingHelper::ArePropertiesCompatible(SourceProperty, ArgumentProperty))
+			if (!bAnyCompatible)
 			{
-				bAnyCompatible = true;
-				break;
+				FString ArgumentsString = FString::JoinBy(ArgumentProperties, TEXT(", "), [](const FProperty* Property) { return Property->GetCPPType(); });
+
+				return MakeError(FString::Printf(TEXT("The source property '%s' (%s) does not match any of the argument types of the conversion function (%s)."), *SourceProperty->GetName(), *SourceProperty->GetCPPType(), *ArgumentsString));
 			}
 		}
-
-		if (!bAnyCompatible)
+		else
 		{
-			FString ArgumentsString = FString::JoinBy(ArgumentProperties, TEXT(", "), [](const FProperty* Property) { return Property->GetCPPType(); });
-
-			return MakeError(FString::Printf(TEXT("The source property '%s' (%s) does not match any of the argument types of the conversion function (%s)."), *SourceProperty->GetName(), *SourceProperty->GetCPPType(), *ArgumentsString));
+			// The complex compiled version look like Setter(Conversion()) (no argument)
+			check(bIsComplexConversionFunction);
+			const FProperty* ReturnProperty = ReturnResult.GetValue();
+			if (!UE::MVVM::BindingHelper::ArePropertiesCompatible(ReturnProperty, DestinationProperty))
+			{
+				return MakeError(FString::Printf(TEXT("The destination property '%s' (%s) does not match the return type of the conversion function (%s)."), *DestinationProperty->GetName(), *DestinationProperty->GetCPPType(), *ReturnProperty->GetCPPType()));
+			}
 		}
 	}
 	else if (!UE::MVVM::BindingHelper::ArePropertiesCompatible(SourceProperty, DestinationProperty))
diff --git a/Engine/Plugins/Runtime/ModelViewViewModel/Source/ModelViewViewModel/Private/View/MVVMView.cpp b/Engine/Plugins/Runtime/ModelViewViewModel/Source/ModelViewViewModel/Private/View/MVVMView.cpp
--- a/Engine/Plugins/Runtime/ModelViewViewModel/Source/ModelViewViewModel/Private/View/MVVMView.cpp
+++ b/Engine/Plugins/Runtime/ModelViewViewModel/Source/ModelViewViewModel/Private/View/MVVMView.cpp
@@ -199,7 +199,8 @@ void UMVVMView::ExecuteLibraryBinding(const FMVVMViewClass_CompiledBinding& Bind
 	check(ClassExtension);
 	check(GetUserWidget());
 
-	TValueOrError<void, FMVVMCompiledBindingLibrary::EExecutionFailingReason> ExecutionResult = ClassExtension->GetBindingLibrary().Execute(GetUserWidget(), Binding.GetBinding());
+	FMVVMCompiledBindingLibrary::EConversionFunctionType FunctionType = Binding.IsConversionFunctionComplex() ? FMVVMCompiledBindingLibrary::EConversionFunctionType::Complex : FMVVMCompiledBindingLibrary::EConversionFunctionType::Simple;
+	TValueOrError<void, FMVVMCompiledBindingLibrary::EExecutionFailingReason> ExecutionResult = ClassExtension->GetBindingLibrary().Execute(GetUserWidget(), Binding.GetBinding(), FunctionType);
 	if (ExecutionResult.HasError())
 	{
 		UE::MVVM::FMessageLog Log(GetUserWidget());
@@ -219,7 +220,9 @@ void UMVVMView::ExecuteLibraryBinding(const FMVVMViewClass_CompiledBinding& Bind
 	check(GetUserWidget());
 	check(Source);
 
-	TValueOrError<void, FMVVMCompiledBindingLibrary::EExecutionFailingReason> ExecutionResult = ClassExtension->GetBindingLibrary().ExecuteWithSource(GetUserWidget(), Binding.GetBinding(), Source);
+	TValueOrError<void, FMVVMCompiledBindingLibrary::EExecutionFailingReason> ExecutionResult = Binding.IsConversionFunctionComplex()
+		? ClassExtension->GetBindingLibrary().Execute(GetUserWidget(), Binding.GetBinding(), FMVVMCompiledBindingLibrary::EConversionFunctionType::Complex)
+		: ClassExtension->GetBindingLibrary().ExecuteWithSource(GetUserWidget(), Binding.GetBinding(), Source);
 	if (ExecutionResult.HasError())
 	{
 		UE::MVVM::FMessageLog Log(GetUserWidget());
diff --git a/Engine/Plugins/Runtime/ModelViewViewModel/Source/ModelViewViewModel/Public/Bindings/MVVMBindingHelper.h b/Engine/Plugins/Runtime/ModelViewViewModel/Source/ModelViewViewModel/Public/Bindings/MVVMBindingHelper.h
--- a/Engine/Plugins/Runtime/ModelViewViewModel/Source/ModelViewViewModel/Public/Bindings/MVVMBindingHelper.h
+++ b/Engine/Plugins/Runtime/ModelViewViewModel/Source/ModelViewViewModel/Public/Bindings/MVVMBindingHelper.h
@@ -48,9 +48,16 @@ namespace UE::MVVM::BindingHelper
 	UE_NODISCARD MODELVIEWVIEWMODEL_API bool IsValidForDestinationBinding(const FMVVMConstFieldVariant InFunction);
 
 	/**
-	 * Is the Function usable as a conversion function by the binding system.
+	 * Is the Function usable as a conversion function by the binding system without any wrapper.
+	 * A simple conversion takes a single argument and returns a single value.
 	 */
-	UE_NODISCARD MODELVIEWVIEWMODEL_API bool IsValidForRuntimeConversion(const UFunction* InFunction);
+	UE_NODISCARD MODELVIEWVIEWMODEL_API bool IsValidForSimpleRuntimeConversion(const UFunction* InFunction);
+
+	/**
+	 * Is the Function usable as a conversion function by the binding system with a wrapper.
+	 * A complex conversion takes a no argument and returns a single value.
+	 */
+	UE_NODISCARD MODELVIEWVIEWMODEL_API bool IsValidForComplexRuntimeConversion(const UFunction* InFunction);
 
 #if WITH_EDITOR
 	/** Is the Property usable as a source by the binding system and can it be read directly or it requires a Getter. */
@@ -115,7 +122,7 @@ namespace UE::MVVM::BindingHelper
 	 *   LocalValue = Src.Property; Dst.SetProperty(LocalValue);
 	 *   LocalValue = Src.GetProperty(); Dst.Property = LocalValue;
 	 *   LocalValue = Src.GetProperty(); Dst.SetProperty(LocalValue);
-	 * with conversion from float to double between the getter and the setter.
+	 * with conversion from float to double or integral to integral between the getter and the setter.
 	 * 
 	 * @note No test is performed to see if the Src Property can be safely assign to the Destination Property. Use with caution.
 	 * @note No test is performed to see if the Src.GetProperty can be safely executed. Use with caution.
@@ -123,6 +130,17 @@ namespace UE::MVVM::BindingHelper
 	 */
 	MODELVIEWVIEWMODEL_API void ExecuteBinding_NoCheck(const FFieldContext& Source, const FFieldContext& Destination);
 
+	/**
+	 * Execute a binding that can be
+	 *   LocalValue = ConcFuncOwner.ConversionFunction(LocalValue); Dst.Property = LocalValue;
+	 *   LocalValue = ConcFuncOwner.ConversionFunction(LocalValue); Dst.SetProperty(LocalValue);
+	 * with conversion from float to double or integral to integral between the getter and the conversion function and/or the conversion function and the setter.
+	 *
+	 * @note No test is performed to see if the Dst.SetProperty can be safely executed. Use with caution.
+	 * @note No test is performed to see if the ConversionFunction can be safely executed. Use with caution.
+	 */
+	MODELVIEWVIEWMODEL_API void ExecuteBinding_NoCheck(const FFieldContext& Destination, const FFunctionContext& ConversionFunction);
+
 	/**
 	 * Execute a binding that can be
 	 *   LocalValue = Src.Property; LocalValue = ConcFuncOwner.ConversionFunction(LocalValue); Dst.Property = LocalValue;
diff --git a/Engine/Plugins/Runtime/ModelViewViewModel/Source/ModelViewViewModel/Public/Bindings/MVVMCompiledBindingLibrary.h b/Engine/Plugins/Runtime/ModelViewViewModel/Source/ModelViewViewModel/Public/Bindings/MVVMCompiledBindingLibrary.h
--- a/Engine/Plugins/Runtime/ModelViewViewModel/Source/ModelViewViewModel/Public/Bindings/MVVMCompiledBindingLibrary.h
+++ b/Engine/Plugins/Runtime/ModelViewViewModel/Source/ModelViewViewModel/Public/Bindings/MVVMCompiledBindingLibrary.h
@@ -258,6 +258,13 @@ public:
 		InvalidCast,
 	};
 
+	/** */
+	enum class EConversionFunctionType
+	{
+		Simple, 	// accept an single argument and returns a single property
+		Complex,	// returns a single property. The argument are fetch from inside the function 
+	};
+
 	/** Fetch the FProperty and UFunction. */
 	void Load();
 
@@ -267,7 +274,7 @@ public:
 	 * The Source of the binding (ViewModel, UserWidget, Widget, ...) can be provided if already known.
 	 * The Destination or the binding (ViewModel, UserWIdget, Widget, ...) can be provided if already known.
 	 */
-	TValueOrError<void, EExecutionFailingReason> Execute(UObject* ExecutionSource, const FMVVMVCompiledBinding& Binding) const;
+	TValueOrError<void, EExecutionFailingReason> Execute(UObject* ExecutionSource, const FMVVMVCompiledBinding& Binding, EConversionFunctionType ConversionType) const;
 	TValueOrError<void, EExecutionFailingReason> ExecuteWithSource(UObject* ExecutionSource, const FMVVMVCompiledBinding& Binding, UObject* Source) const;
 
 	/**
@@ -284,8 +291,8 @@ public:
 	TValueOrError<FString, FString> FieldPathToString(const FMVVMVCompiledFieldPath& FieldPath) const;
 
 private:
-	TValueOrError<void, EExecutionFailingReason> ExecuteImpl(UObject* ExecutionSource, const FMVVMVCompiledBinding& Binding, UObject* Source) const;
-	TValueOrError<void, EExecutionFailingReason> ExecuteImpl(UE::MVVM::FFieldContext& Source, UE::MVVM::FFieldContext& Destination, UE::MVVM::FFunctionContext& ConversionFunction) const;
+	TValueOrError<void, EExecutionFailingReason> ExecuteImpl(UObject* ExecutionSource, const FMVVMVCompiledBinding& Binding, UObject* Source, EConversionFunctionType ConversionType) const;
+	TValueOrError<void, EExecutionFailingReason> ExecuteImpl(UE::MVVM::FFieldContext& Source, UE::MVVM::FFieldContext& Destination, UE::MVVM::FFunctionContext& ConversionFunction, EConversionFunctionType ConversionType) const;
 
 	TValueOrError<UE::MVVM::FMVVMFieldVariant, void> GetFinalFieldFromPathImpl(const FMVVMVCompiledFieldPath& FieldPath) const;
 
diff --git a/Engine/Plugins/Runtime/ModelViewViewModel/Source/ModelViewViewModel/Public/View/MVVMViewClass.h b/Engine/Plugins/Runtime/ModelViewViewModel/Source/ModelViewViewModel/Public/View/MVVMViewClass.h
--- a/Engine/Plugins/Runtime/ModelViewViewModel/Source/ModelViewViewModel/Public/View/MVVMViewClass.h
+++ b/Engine/Plugins/Runtime/ModelViewViewModel/Source/ModelViewViewModel/Public/View/MVVMViewClass.h
@@ -135,10 +135,16 @@ public:
 		return (Flags & EBindingFlags::EnabledByDefault) != 0;
 	}
 
-	/** @return true if the binding . */
+	/** @return true if it's normal that the binding could not find it's source when registering it. */
 	bool IsRegistrationOptional() const
 	{
-		return (Flags & EBindingFlags::RegistrationOptional) != 0;
+		return (Flags & EBindingFlags::ViewModelOptional) != 0;
+	}
+
+	/** @return true if the binding use a conversion function and that the conversion function is complex. */
+	bool IsConversionFunctionComplex() const
+	{
+		return (Flags & EBindingFlags::ConversionFunctionIsComplex) != 0;
 	}
 
 	/** @return a human readable version of the binding that can be use for debugging purposes. */
@@ -165,7 +171,10 @@ private:
 		TwoWayBinding = 1 << 1, // The binding is one part of a 2 ways binding.
 		OneTime = 1 << 2,
 		EnabledByDefault = 1 << 3,
-		RegistrationOptional = 1 << 4,	// The source (viewmodel) can be nullptr and the binding could failed without warning.
+		ViewModelOptional = 1 << 4,	// The source (viewmodel) can be nullptr and the binding could failed and should not log a warning.
+		ConversionFunctionIsComplex = 1 << 5,	// The conversion function is complex, there is no input. The inputs are calculated in the BP function.
+		Unused01 = 1 << 6,
+		Unused02 = 1 << 7,
 	};
 
 	UPROPERTY()
diff --git a/Engine/Plugins/Runtime/ModelViewViewModel/Source/ModelViewViewModelBlueprint/Private/Bindings/MVVMCompiledBindingLibraryCompiler.cpp b/Engine/Plugins/Runtime/ModelViewViewModel/Source/ModelViewViewModelBlueprint/Private/Bindings/MVVMCompiledBindingLibraryCompiler.cpp
--- a/Engine/Plugins/Runtime/ModelViewViewModel/Source/ModelViewViewModelBlueprint/Private/Bindings/MVVMCompiledBindingLibraryCompiler.cpp
+++ b/Engine/Plugins/Runtime/ModelViewViewModel/Source/ModelViewViewModelBlueprint/Private/Bindings/MVVMCompiledBindingLibraryCompiler.cpp
@@ -361,7 +361,9 @@ TValueOrError<FCompiledBindingLibraryCompiler::FFieldPathHandle, FString> FCompi
 		return MakeError(TEXT("The function path is empty."));
 	}
 
-	if (!BindingHelper::IsValidForRuntimeConversion(Function))
+	const bool bIsSimpleFunction = BindingHelper::IsValidForSimpleRuntimeConversion(Function);
+	const bool bIsComplexFunction = BindingHelper::IsValidForComplexRuntimeConversion(Function);
+	if (!bIsSimpleFunction && !bIsComplexFunction)
 	{
 		return MakeError(FString::Printf(TEXT("Function %s cannot be used as a runtime conversion function."), *Function->GetName()));
 	}
diff --git a/Engine/Plugins/Runtime/ModelViewViewModel/Source/ModelViewViewModelBlueprint/Private/MVVMViewBlueprintCompiler.cpp b/Engine/Plugins/Runtime/ModelViewViewModel/Source/ModelViewViewModelBlueprint/Private/MVVMViewBlueprintCompiler.cpp
--- a/Engine/Plugins/Runtime/ModelViewViewModel/Source/ModelViewViewModelBlueprint/Private/MVVMViewBlueprintCompiler.cpp
+++ b/Engine/Plugins/Runtime/ModelViewViewModel/Source/ModelViewViewModelBlueprint/Private/MVVMViewBlueprintCompiler.cpp
@@ -11,6 +11,7 @@
 #include "MVVMFunctionGraphHelper.h"
 #include "MVVMSubsystem.h"
 #include "MVVMViewModelBase.h"
+#include "MVVMViewModelBase.h"
 #include "WidgetBlueprintCompiler.h"
 #include "View/MVVMViewClass.h"
 
@@ -38,12 +39,11 @@ void FMVVMViewBlueprintCompiler::CleanOldData(UWidgetBlueprintGeneratedClass* Cl
 	{
 		auto RenameObjectToTransientPackage = [](UObject* ObjectToRename)
 		{
-			const ERenameFlags RenFlags = REN_DontCreateRedirectors | REN_NonTransactional | REN_DoNotDirty;
+			const ERenameFlags RenFlags = REN_DoNotDirty | REN_ForceNoResetLoaders | REN_DontCreateRedirectors;
 
 			ObjectToRename->Rename(nullptr, GetTransientPackage(), RenFlags);
 			ObjectToRename->SetFlags(RF_Transient);
 			ObjectToRename->ClearFlags(RF_Public | RF_Standalone | RF_ArchetypeObject);
-
 			FLinkerLoad::InvalidateExport(ObjectToRename);
 		};
 
@@ -67,16 +67,6 @@ void FMVVMViewBlueprintCompiler::CleanOldData(UWidgetBlueprintGeneratedClass* Cl
 
 void FMVVMViewBlueprintCompiler::CleanTemporaries(UWidgetBlueprintGeneratedClass* ClassToClean)
 {
-	for (FCompilerSourceCreatorContext& SourceCreatorContext : CompilerSourceCreatorContexts)
-	{
-		if (SourceCreatorContext.SetterGraph)
-		{
-			ensureMsgf(SourceCreatorContext.SetterGraph->HasAnyFlags(RF_Transient), TEXT("The graph should be temporary and should be generated automaticly."));
-			// GC may not have clean the graph (GC doesn't run when bRegenerateSkeletonOnly is on)
-			SourceCreatorContext.SetterGraph->Rename(nullptr, SourceCreatorContext.SetterGraph->GetOuter(), REN_DoNotDirty | REN_ForceNoResetLoaders);
-			SourceCreatorContext.SetterGraph = nullptr;
-		}
-	}
 }
 
 
@@ -402,17 +392,40 @@ void FMVVMViewBlueprintCompiler::CreateSourceLists(const FWidgetBlueprintCompile
 
 void FMVVMViewBlueprintCompiler::CreateFunctionsDeclaration(const FWidgetBlueprintCompilerContext::FCreateVariableContext& Context, UMVVMBlueprintView* BlueprintView)
 {
+	// Clean all previous intermediate function graph. It should stay alive. The graph lives on the Blueprint not on the class and it's used to generate the UFunction.
+	{
+		auto RenameObjectToTransientPackage = [](UObject* ObjectToRename)
+		{
+			const ERenameFlags RenFlags = REN_DontCreateRedirectors | REN_NonTransactional | REN_DoNotDirty;
+			ObjectToRename->Rename(nullptr, GetTransientPackage(), RenFlags);
+			ObjectToRename->SetFlags(RF_Transient);
+			ObjectToRename->ClearFlags(RF_Public | RF_Standalone | RF_ArchetypeObject);
+			FLinkerLoad::InvalidateExport(ObjectToRename);
+		};
+
+		for (UEdGraph* OldGraph : BlueprintView->TemporaryGraph)
+		{
+			if (OldGraph)
+			{
+				RenameObjectToTransientPackage(OldGraph);
+			}
+		}
+		BlueprintView->TemporaryGraph.Reset();
+	}
+
 	for (FCompilerSourceCreatorContext& SourceCreator : CompilerSourceCreatorContexts)
 	{
 		if (!SourceCreator.SetterFunctionName.IsEmpty() && SourceCreator.Type == ECompilerSourceCreatorType::ViewModel)
 		{
 			ensure(SourceCreator.SetterGraph == nullptr);
+
 			SourceCreator.SetterGraph = UE::MVVM::FunctionGraphHelper::CreateIntermediateFunctionGraph(
 				WidgetBlueprintCompilerContext
 				, SourceCreator.SetterFunctionName
 				, (FUNC_BlueprintCallable | FUNC_Public)
 				, TEXT("Viewmodel")
 				, false);
+			BlueprintView->TemporaryGraph.Add(SourceCreator.SetterGraph);
 
 			if (SourceCreator.SetterGraph == nullptr || SourceCreator.SetterGraph->GetFName() != FName(*SourceCreator.SetterFunctionName))
 			{
@@ -754,11 +767,11 @@ bool FMVVMViewBlueprintCompiler::PreCompileBindings(UWidgetBlueprintGeneratedCla
 		FMVVMBlueprintViewBinding& Binding = *BindingPtr;
 
 		FMVVMViewBlueprintCompiler* Self = this;
-		auto AddBinding = [Self](UWidgetBlueprintGeneratedClass* Class, const TArrayView<TArray<UE::MVVM::FMVVMConstFieldVariant>> GetterFields, TArrayView<UE::MVVM::FMVVMConstFieldVariant> SetterFields, const UFunction* ConversionFunction) -> TValueOrError<FCompilerBinding, FString>
+		auto AddBinding = [Self](UWidgetBlueprintGeneratedClass* Class, TArrayView<const UE::MVVM::FMVVMConstFieldVariant> GetterField, TArrayView<const UE::MVVM::FMVVMConstFieldVariant> SetterFields, const UFunction* ConversionFunction) -> TValueOrError<FCompilerBinding, FString>
 		{
 			FCompilerBinding Result;
 
-			for (const TArray<UE::MVVM::FMVVMConstFieldVariant>& GetterField : GetterFields)
+			if (GetterField.Num() != 0)
 			{
 				// Generate a path to read the value at runtime
 				TValueOrError<TArray<FMVVMConstFieldVariant>, FString> GeneratedField = FieldPathHelper::GenerateFieldPathList(GetterField, true);
@@ -776,7 +789,7 @@ bool FMVVMViewBlueprintCompiler::PreCompileBindings(UWidgetBlueprintGeneratedCla
 						, *::UE::MVVM::FieldPathHelper::ToString(GetterField)
 						, *FieldPathResult.GetError()));
 				}
-				Result.SourceRead.Add(FieldPathResult.GetValue());
+				Result.SourceRead = FieldPathResult.StealValue();
 			}
 
 			{
@@ -795,7 +808,7 @@ bool FMVVMViewBlueprintCompiler::PreCompileBindings(UWidgetBlueprintGeneratedCla
 						, *::UE::MVVM::FieldPathHelper::ToString(SetterFields)
 						, *FieldPathResult.GetError()));
 				}
-				Result.DestinationWrite = FieldPathResult.GetValue();
+				Result.DestinationWrite = FieldPathResult.StealValue();
 			}
 
 			if (ConversionFunction != nullptr)
@@ -807,7 +820,18 @@ bool FMVVMViewBlueprintCompiler::PreCompileBindings(UWidgetBlueprintGeneratedCla
 						, *ConversionFunction->GetPathName()
 						, *FieldPathResult.GetError()));
 				}
-				Result.ConversionFunction = FieldPathResult.GetValue();
+				Result.ConversionFunction = FieldPathResult.StealValue();
+				Result.bIsConversionFunctionComplex = BindingHelper::IsValidForComplexRuntimeConversion(ConversionFunction);
+			}
+
+			// Sanity check.
+			if (GetterField.Num() == 0 && (ConversionFunction == nullptr || !BindingHelper::IsValidForComplexRuntimeConversion(ConversionFunction)))
+			{
+				return MakeError(FString::Printf(TEXT("The binding should have a getter.")));
+			}
+			if (GetterField.Num() != 0 && ConversionFunction == nullptr && BindingHelper::IsValidForComplexRuntimeConversion(ConversionFunction))
+			{
+				return MakeError(FString::Printf(TEXT("The binding has a complex conversion function and a getter.")));
 			}
 
 			// Generate the binding
@@ -838,16 +862,6 @@ bool FMVVMViewBlueprintCompiler::PreCompileBindings(UWidgetBlueprintGeneratedCla
 			continue;
 		}
 
-		TArray<TArray<UE::MVVM::FMVVMConstFieldVariant>> GetterFields;
-		//if (Binding.Conversion.HasConversionFunction && has read values defined (that are not the default one))
-		//{
-		//  Add the read values for the conversion function
-		//}
-		//else
-		{
-			GetterFields.Add(BindingSourceContext.PropertyPath);
-		}
-
 		TArray<UE::MVVM::FMVVMConstFieldVariant> SetterPath;
 		{
 			const FMVVMBlueprintPropertyPath& DestinationPath = BindingSourceContext.bIsForwardBinding ? Binding.WidgetPath : Binding.ViewModelPath;
@@ -868,8 +882,14 @@ bool FMVVMViewBlueprintCompiler::PreCompileBindings(UWidgetBlueprintGeneratedCla
 		}
 
 		const UFunction* ConversionFunction = ConversionFunctionReference.ResolveMember<UFunction>(Class);
+		if (!ConversionFunctionWrapper.IsNone() && ConversionFunction == nullptr)
+		{
+			AddErrorForBinding(Binding, BlueprintView, TEXT("A complex conversion function was created but could not be found."));
+			bIsBindingsValid = false;
+			continue;
+		}
 
-		TValueOrError<FCompilerBinding, FString> AddBindingResult = AddBinding(Class, GetterFields, SetterPath, ConversionFunction);
+		TValueOrError<FCompilerBinding, FString> AddBindingResult = AddBinding(Class, BindingSourceContext.PropertyPath, SetterPath, ConversionFunction);
 		if (AddBindingResult.HasError())
 		{
 			AddErrorForBinding(Binding, BlueprintView, FString::Printf(TEXT("The binding could not be created. %s"), *AddBindingResult.GetError()));
@@ -924,6 +944,16 @@ bool FMVVMViewBlueprintCompiler::CompileBindings(const FCompiledBindingLibraryCo
 			continue;
 		}
 
+		bool bIsOptional = false;
+		if (CompilerSourceCreatorContexts.IsValidIndex(CompileBinding.CompilerSourceContextIndex))
+		{
+			const FMVVMBlueprintViewModelContext& ModelContext = CompilerSourceCreatorContexts[CompileBinding.CompilerSourceContextIndex].ViewModelContext;
+			if (ModelContext.IsValid())
+			{
+				bIsOptional = ModelContext.bOptional;
+			}
+		}
+
 		NewBinding.FieldId = CompiledFieldId  ? *CompiledFieldId : FMVVMVCompiledFieldId();
 		NewBinding.Binding = *CompiledBinding;
 		NewBinding.UpdateMode = ViewBinding.UpdateMode;
@@ -933,7 +963,8 @@ bool FMVVMViewBlueprintCompiler::CompileBindings(const FCompiledBindingLibraryCo
 		NewBinding.Flags |= (IsForwardBinding(ViewBinding.BindingType)) ? FMVVMViewClass_CompiledBinding::EBindingFlags::ForwardBinding : 0;
 		NewBinding.Flags |= (ViewBinding.BindingType == EMVVMBindingMode::TwoWay) ? FMVVMViewClass_CompiledBinding::EBindingFlags::TwoWayBinding : 0;
 		NewBinding.Flags |= (IsOneTimeBinding(ViewBinding.BindingType)) ? FMVVMViewClass_CompiledBinding::EBindingFlags::OneTime : 0;
-		//NewBinding.Flags |= () ? FMVVMViewClass_CompiledBinding::EBindingFlags::RegistrationOptional : 0;
+		NewBinding.Flags |= (bIsOptional) ? FMVVMViewClass_CompiledBinding::EBindingFlags::ViewModelOptional : 0;
+		NewBinding.Flags |= (CompileBinding.bIsConversionFunctionComplex) ? FMVVMViewClass_CompiledBinding::EBindingFlags::ConversionFunctionIsComplex : 0;
 
 		ViewExtension->CompiledBindings.Emplace(MoveTemp(NewBinding));
 	}
diff --git a/Engine/Plugins/Runtime/ModelViewViewModel/Source/ModelViewViewModelBlueprint/Private/MVVMViewBlueprintCompiler.h b/Engine/Plugins/Runtime/ModelViewViewModel/Source/ModelViewViewModelBlueprint/Private/MVVMViewBlueprintCompiler.h
--- a/Engine/Plugins/Runtime/ModelViewViewModel/Source/ModelViewViewModelBlueprint/Private/MVVMViewBlueprintCompiler.h
+++ b/Engine/Plugins/Runtime/ModelViewViewModel/Source/ModelViewViewModelBlueprint/Private/MVVMViewBlueprintCompiler.h
@@ -118,6 +118,7 @@ private:
 		int32 CompilerSourceContextIndex = INDEX_NONE;
 		bool bSourceIsUserWidget = false;
 		bool bFieldIdNeeded = false;
+		bool bIsConversionFunctionComplex = false;
 
 		FCompiledBindingLibraryCompiler::FBindingHandle BindingHandle;
 
@@ -128,7 +129,7 @@ private:
 		//else we read the same property that TriggerField changed
 		// then we copy TriggerField to DestiantionWrite.
 		FCompiledBindingLibraryCompiler::FFieldIdHandle FieldIdHandle;
-		TArray<FCompiledBindingLibraryCompiler::FFieldPathHandle> SourceRead;
+		FCompiledBindingLibraryCompiler::FFieldPathHandle SourceRead;
 		FCompiledBindingLibraryCompiler::FFieldPathHandle DestinationWrite;
 		FCompiledBindingLibraryCompiler::FFieldPathHandle ConversionFunction;
 	};
diff --git a/Engine/Plugins/Runtime/ModelViewViewModel/Source/ModelViewViewModelBlueprint/Private/MVVMWidgetBlueprintExtension_View.cpp b/Engine/Plugins/Runtime/ModelViewViewModel/Source/ModelViewViewModelBlueprint/Private/MVVMWidgetBlueprintExtension_View.cpp
--- a/Engine/Plugins/Runtime/ModelViewViewModel/Source/ModelViewViewModelBlueprint/Private/MVVMWidgetBlueprintExtension_View.cpp
+++ b/Engine/Plugins/Runtime/ModelViewViewModel/Source/ModelViewViewModelBlueprint/Private/MVVMWidgetBlueprintExtension_View.cpp
@@ -113,24 +113,26 @@ void UMVVMWidgetBlueprintExtension_View::HandleFinishCompilingClass(UWidgetBluep
 
 	check(CurrentCompilerContext);
 
-	if (CurrentCompilerContext->GetCompilerContext().bAssignDelegateSignatureFunction)
+	if (CurrentCompilerContext->GetCompilerContext().bIsFullCompile)
 	{
+		UMVVMViewClass* ViewExtension = nullptr;
+		bool bCompiled = false;
+		if (CurrentCompilerContext->PreCompile(Class, BlueprintView))
+		{
+			ViewExtension = NewObject<UMVVMViewClass>(Class);
+			bCompiled = CurrentCompilerContext->Compile(Class, BlueprintView, ViewExtension);
+		}
+
 		CurrentCompilerContext->CleanTemporaries(Class);
-		return;
-	}
 
-	UMVVMViewClass* ViewExtension = nullptr;
-	bool bCompiled = false;
-	if (CurrentCompilerContext->PreCompile(Class, BlueprintView))
-	{
-		ViewExtension = NewObject<UMVVMViewClass>(Class);
-		bCompiled = CurrentCompilerContext->Compile(Class, BlueprintView, ViewExtension);
+		if (bCompiled)
+		{
+			check(ViewExtension);
+			CurrentCompilerContext->AddExtension(Class, ViewExtension);
+		}
 	}
-
-	CurrentCompilerContext->CleanTemporaries(Class);
-	if (bCompiled)
+	else
 	{
-		check(ViewExtension);
-		CurrentCompilerContext->AddExtension(Class, ViewExtension);
+		CurrentCompilerContext->CleanTemporaries(Class);
 	}
 }
diff --git a/Engine/Plugins/Runtime/ModelViewViewModel/Source/ModelViewViewModelBlueprint/Public/MVVMBlueprintView.h b/Engine/Plugins/Runtime/ModelViewViewModel/Source/ModelViewViewModelBlueprint/Public/MVVMBlueprintView.h
--- a/Engine/Plugins/Runtime/ModelViewViewModel/Source/ModelViewViewModelBlueprint/Public/MVVMBlueprintView.h
+++ b/Engine/Plugins/Runtime/ModelViewViewModel/Source/ModelViewViewModelBlueprint/Public/MVVMBlueprintView.h
@@ -76,6 +76,10 @@ public:
 	DECLARE_EVENT(UMVVMBlueprintView, FOnViewModelsUpdated);
 	FOnViewModelsUpdated OnViewModelsUpdated;
 
+	// Use during compilation to clean the automatically generated graph.
+	UPROPERTY(Transient)
+	TArray<TObjectPtr<UEdGraph>> TemporaryGraph;
+
 private:
 	UPROPERTY(EditAnywhere, Category = "MVVM")
 	TArray<FMVVMBlueprintViewBinding> Bindings;
diff --git a/Engine/Plugins/Runtime/ModelViewViewModel/Source/ModelViewViewModelEditor/Private/Tests/MVVMBindingExecuteTest.cpp b/Engine/Plugins/Runtime/ModelViewViewModel/Source/ModelViewViewModelEditor/Private/Tests/MVVMBindingExecuteTest.cpp
--- a/Engine/Plugins/Runtime/ModelViewViewModel/Source/ModelViewViewModelEditor/Private/Tests/MVVMBindingExecuteTest.cpp
+++ b/Engine/Plugins/Runtime/ModelViewViewModel/Source/ModelViewViewModelEditor/Private/Tests/MVVMBindingExecuteTest.cpp
@@ -614,6 +614,68 @@ bool FMVVMBindingExecuteTest::RunTest(const FString& Parameters)
 		AssignPropertyFloatDouble(UE::MVVM::FMVVMFieldVariant(GetFloat), UE::MVVM::FMVVMFieldVariant(GetDouble)
 			, UE::MVVM::FMVVMFieldVariant(PropFloat), UE::MVVM::FMVVMFieldVariant(PropDouble), ConversionIncDoubleFuntion, TEXT("PropertyReal = GetPropertyReal()"));
 	}
+	{
+		auto AssignPropertyAFromConversionFunction = [this, SourceObj, DestinationObj](UE::MVVM::FMVVMFieldVariant DestinationBinding, UFunction* Conversion, const TCHAR* Msg)
+		{
+			SourceObj->PropertyA.Value = 448;
+			DestinationObj->PropertyA.Value = 999;
+			UE::MVVM::Private::GMVVMBindingExecTextCounter = 0;
+
+			UE::MVVM::FFieldContext Destination(UE::MVVM::FObjectVariant(DestinationObj), DestinationBinding);
+			UE::MVVM::FFunctionContext Function = UE::MVVM::FFunctionContext(SourceObj, Conversion);
+			UE::MVVM::BindingHelper::ExecuteBinding_NoCheck(Destination, Function);
+
+			if (UE::MVVM::Private::GMVVMBindingExecTextCounter != 0) // -V547
+			{
+				AddError(FString::Printf(TEXT("%s failed to release the reources."), Msg));
+			}
+			if (SourceObj->PropertyA != DestinationObj->PropertyA)
+			{
+				AddError(FString::Printf(TEXT("%s failed the assignement."), Msg));
+			}
+			if (SourceObj->PropertyA.Value != 448)
+			{
+				AddError(FString::Printf(TEXT("%s the source value changes."), Msg));
+			}
+		};
+
+		FProperty* PropA = UMVVMViewModelBindingExecTest::StaticClass()->FindPropertyByName(GET_MEMBER_NAME_CHECKED(UMVVMViewModelBindingExecTest, PropertyA));
+		UFunction* SetA = UMVVMViewModelBindingExecTest::StaticClass()->FindFunctionByName(GET_FUNCTION_NAME_CHECKED(UMVVMViewModelBindingExecTest, SetterA));
+		UFunction* GetA = UMVVMViewModelBindingExecTest::StaticClass()->FindFunctionByName(GET_FUNCTION_NAME_CHECKED(UMVVMViewModelBindingExecTest, GetterA));
+		AssignPropertyAFromConversionFunction(UE::MVVM::FMVVMFieldVariant(PropA), GetA, TEXT("PropertyA = GetterA()"));
+		AssignPropertyAFromConversionFunction(UE::MVVM::FMVVMFieldVariant(SetA), GetA, TEXT("SetPropertyA(GetterA())"));
+	}
+	{
+		auto AssignPropertyCFromConversionFunction = [this, SourceObj, DestinationObj](UE::MVVM::FMVVMFieldVariant DestinationBinding, UFunction* Conversion, const TCHAR* Msg)
+		{
+			SourceObj->PropertyC = 458;
+			DestinationObj->PropertyC = 999;
+			UE::MVVM::Private::GMVVMBindingExecTextCounter = 0;
+
+			UE::MVVM::FFieldContext Destination(UE::MVVM::FObjectVariant(DestinationObj), DestinationBinding);
+			UE::MVVM::FFunctionContext Function = UE::MVVM::FFunctionContext(SourceObj, Conversion);
+			UE::MVVM::BindingHelper::ExecuteBinding_NoCheck(Destination, Function);
+
+			if (UE::MVVM::Private::GMVVMBindingExecTextCounter != 0) // -V547
+			{
+				AddError(FString::Printf(TEXT("%s failed to release the reources."), Msg));
+			}
+			if (SourceObj->PropertyC != DestinationObj->PropertyC)
+			{
+				AddError(FString::Printf(TEXT("%s failed the assignement."), Msg));
+			}
+			if (SourceObj->PropertyC != 458)
+			{
+				AddError(FString::Printf(TEXT("%s the source value changes."), Msg));
+			}
+		};
+
+		FProperty* PropC = UMVVMViewModelBindingExecTest::StaticClass()->FindPropertyByName(GET_MEMBER_NAME_CHECKED(UMVVMViewModelBindingExecTest, PropertyC));
+		UFunction* SetC = UMVVMViewModelBindingExecTest::StaticClass()->FindFunctionByName(GET_FUNCTION_NAME_CHECKED(UMVVMViewModelBindingExecTest, SetterC));
+		UFunction* GetC = UMVVMViewModelBindingExecTest::StaticClass()->FindFunctionByName(GET_FUNCTION_NAME_CHECKED(UMVVMViewModelBindingExecTest, GetterC));
+		AssignPropertyCFromConversionFunction(UE::MVVM::FMVVMFieldVariant(PropC), GetC, TEXT("PropertyC = GetterC()"));
+		AssignPropertyCFromConversionFunction(UE::MVVM::FMVVMFieldVariant(SetC), GetC, TEXT("SetPropertyC(GetterC())"));
+	}
 
 	SourceObj->RemoveFromRoot();
 	DestinationObj->RemoveFromRoot();
diff --git a/Engine/Plugins/Runtime/ModelViewViewModel/Source/ModelViewViewModelEditor/Private/Tests/MVVMBindingHelperTest.cpp b/Engine/Plugins/Runtime/ModelViewViewModel/Source/ModelViewViewModelEditor/Private/Tests/MVVMBindingHelperTest.cpp
--- a/Engine/Plugins/Runtime/ModelViewViewModel/Source/ModelViewViewModelEditor/Private/Tests/MVVMBindingHelperTest.cpp
+++ b/Engine/Plugins/Runtime/ModelViewViewModel/Source/ModelViewViewModelEditor/Private/Tests/MVVMBindingHelperTest.cpp
@@ -75,7 +75,7 @@ bool FMVVMBindingHelperTest::RunTest(const FString& Parameters)
 
 
 	auto TestFunction = [this](FName FunctionName
-		, bool bIsValidForSourceBinding, bool bIsValidForDestinationBinding, bool bIsValidForRuntimeConversion
+		, bool bIsValidForSourceBinding, bool bIsValidForDestinationBinding, bool bIsValidForSimpleRuntimeConversion, bool bIsValidForComplexRuntimeConversion
 		, FFieldClass* PropertyTypeForSourceBinding, FFieldClass* PropertyTypeForDestinationBinding)
 	{
 		UFunction* Function = UMVVMViewModelBindingHelperTest::StaticClass()->FindFunctionByName(FunctionName);
@@ -89,9 +89,13 @@ bool FMVVMBindingHelperTest::RunTest(const FString& Parameters)
 		{
 			AddError(FString::Printf(TEXT("IsValidForDestinationBinding failed for '%s'."), *FunctionName.ToString()));
 		}
-		if (::UE::MVVM::BindingHelper::IsValidForRuntimeConversion(Function) != bIsValidForRuntimeConversion)
+		if (::UE::MVVM::BindingHelper::IsValidForSimpleRuntimeConversion(Function) != bIsValidForSimpleRuntimeConversion)
 		{
-			AddError(FString::Printf(TEXT("IsValidForRuntimeConversion failed for '%s'."), *FunctionName.ToString()));
+			AddError(FString::Printf(TEXT("IsValidForSimpleRuntimeConversion failed for '%s'."), *FunctionName.ToString()));
+		}
+		if (::UE::MVVM::BindingHelper::IsValidForComplexRuntimeConversion(Function) != bIsValidForComplexRuntimeConversion)
+		{
+			AddError(FString::Printf(TEXT("IsValidForComplexRuntimeConversion failed for '%s'."), *FunctionName.ToString()));
 		}
 
 		{
@@ -118,41 +122,41 @@ bool FMVVMBindingHelperTest::RunTest(const FString& Parameters)
 		}
 	};
 
-	TestFunction("FunctionGetA", false, false, false, nullptr, nullptr);
-	TestFunction("FunctionGetB", false, false, false, nullptr, nullptr);
-	TestFunction("FunctionGetC", true, false, false, FIntProperty::StaticClass(), nullptr);
-	TestFunction("FunctionGetD", false, false, true, nullptr, nullptr);
-	TestFunction("FunctionGetE", false, false, false, nullptr, nullptr);
-	TestFunction("FunctionGetF", true, false, false, FArrayProperty::StaticClass(), nullptr);
-	TestFunction("FunctionGetG", true, false, false, FArrayProperty::StaticClass(), nullptr);
-	TestFunction("FunctionGetH", false, false, false, nullptr, nullptr);
-	TestFunction("FunctionGetI", false, false, false, nullptr, nullptr);
-	TestFunction("FunctionGetJ", false, false, true, nullptr, nullptr);
-
-	TestFunction("FunctionSetA", false, false, false, nullptr, nullptr);
-	TestFunction("FunctionSetB", false, true, false, nullptr, FIntProperty::StaticClass());
-	TestFunction("FunctionSetC", false, false, false, nullptr, nullptr);
-	TestFunction("FunctionSetD", false, false, false, nullptr, nullptr);
-	TestFunction("FunctionSetE", false, false, false, nullptr, nullptr);
-	TestFunction("FunctionSetF", false, false, false, nullptr, nullptr);
-	TestFunction("FunctionSetG", false, true, false, nullptr, FArrayProperty::StaticClass());
-	TestFunction("FunctionSetH", false, false, false, nullptr, nullptr);
-
-	TestFunction("FunctionGetProtected", true, false, false, FIntProperty::StaticClass(), nullptr);
-	TestFunction("FunctionSetProtected", false, true, false, nullptr, FIntProperty::StaticClass());
-	TestFunction("FunctionGetter", true, false, false, FIntProperty::StaticClass(), nullptr);
-	TestFunction("FunctionSetter", false, true, false, nullptr, FIntProperty::StaticClass());
-
-	TestFunction("FunctionConversionA", false, false, false, nullptr, nullptr);
-	TestFunction("FunctionConversionB", false, false, false, nullptr, nullptr);
-	TestFunction("FunctionConversionC", false, false, true, nullptr, nullptr);
-	TestFunction("FunctionConversionD", false, false, false, nullptr, nullptr);
-	TestFunction("FunctionConversionE", false, true, false, nullptr, FIntProperty::StaticClass());
-	TestFunction("FunctionConversionF", false, false, false, nullptr, nullptr);
-	TestFunction("FunctionConversionG", false, false, true, nullptr, nullptr);
-	TestFunction("FunctionConversionH", false, false, true, nullptr, nullptr);
-	TestFunction("FunctionConversionI", false, false, false, nullptr, nullptr);
-	TestFunction("FunctionConversionJ", false, false, false, nullptr, nullptr);
+	TestFunction("FunctionGetA", false, false, false, false, nullptr, nullptr);
+	TestFunction("FunctionGetB", false, false, false, false, nullptr, nullptr);
+	TestFunction("FunctionGetC", true, false, false, true, FIntProperty::StaticClass(), nullptr);
+	TestFunction("FunctionGetD", false, false, true, false, nullptr, nullptr);
+	TestFunction("FunctionGetE", false, false, false, false, nullptr, nullptr);
+	TestFunction("FunctionGetF", true, false, false, true, FArrayProperty::StaticClass(), nullptr);
+	TestFunction("FunctionGetG", true, false, false, true, FArrayProperty::StaticClass(), nullptr);
+	TestFunction("FunctionGetH", false, false, false, false, nullptr, nullptr);
+	TestFunction("FunctionGetI", false, false, false, false, nullptr, nullptr);
+	TestFunction("FunctionGetJ", false, false, true, false, nullptr, nullptr);
+
+	TestFunction("FunctionSetA", false, false, false, false, nullptr, nullptr);
+	TestFunction("FunctionSetB", false, true, false, false, nullptr, FIntProperty::StaticClass());
+	TestFunction("FunctionSetC", false, false, false, false, nullptr, nullptr);
+	TestFunction("FunctionSetD", false, false, false, false, nullptr, nullptr);
+	TestFunction("FunctionSetE", false, false, false, false, nullptr, nullptr);
+	TestFunction("FunctionSetF", false, false, false, false, nullptr, nullptr);
+	TestFunction("FunctionSetG", false, true, false, false, nullptr, FArrayProperty::StaticClass());
+	TestFunction("FunctionSetH", false, false, false, false, nullptr, nullptr);
+
+	TestFunction("FunctionGetProtected", true, false, false, true, FIntProperty::StaticClass(), nullptr);
+	TestFunction("FunctionSetProtected", false, true, false, false, nullptr, FIntProperty::StaticClass());
+	TestFunction("FunctionGetter", true, false, false, true, FIntProperty::StaticClass(), nullptr);
+	TestFunction("FunctionSetter", false, true, false, false, nullptr, FIntProperty::StaticClass());
+
+	TestFunction("FunctionConversionA", false, false, false, false, nullptr, nullptr);
+	TestFunction("FunctionConversionB", false, false, false, true, nullptr, nullptr);
+	TestFunction("FunctionConversionC", false, false, true, false, nullptr, nullptr);
+	TestFunction("FunctionConversionD", false, false, false, false, nullptr, nullptr);
+	TestFunction("FunctionConversionE", false, true, false, false, nullptr, FIntProperty::StaticClass());
+	TestFunction("FunctionConversionF", false, false, false, false, nullptr, nullptr);
+	TestFunction("FunctionConversionG", false, false, true, false, nullptr, nullptr);
+	TestFunction("FunctionConversionH", false, false, true, false, nullptr, nullptr);
+	TestFunction("FunctionConversionI", false, false, false, false, nullptr, nullptr);
+	TestFunction("FunctionConversionJ", false, false, false, false, nullptr, nullptr);
 
 	auto TestConversionFunction = [this](FName FunctionName, FFieldClass* ExpectedReturnProperty, FFieldClass* ExpectedFirstProperty)
 	{
