MVVM: Fix looping compilation. When only the skeletal is compiled, no GC is called to destroyed the temporary objects.

diff --git a/Engine/Plugins/Runtime/ModelViewViewModel/Source/ModelViewViewModelBlueprint/Private/ViewModel/MVVMViewModelBlueprintCompiler.cpp b/Engine/Plugins/Runtime/ModelViewViewModel/Source/ModelViewViewModelBlueprint/Private/ViewModel/MVVMViewModelBlueprintCompiler.cpp
--- a/Engine/Plugins/Runtime/ModelViewViewModel/Source/ModelViewViewModelBlueprint/Private/ViewModel/MVVMViewModelBlueprintCompiler.cpp
+++ b/Engine/Plugins/Runtime/ModelViewViewModel/Source/ModelViewViewModelBlueprint/Private/ViewModel/MVVMViewModelBlueprintCompiler.cpp
@@ -140,6 +140,27 @@ void FViewModelBlueprintCompilerContext::CreateClassVariablesFromBlueprint()
 {
 	Super::CreateClassVariablesFromBlueprint();
 
+	if (CompileOptions.CompileType == EKismetCompileType::SkeletonOnly)
+	{
+		auto RenameObjectToTransientPackage = [](UObject* ObjectToRename)
+		{
+			const ERenameFlags RenFlags = REN_DontCreateRedirectors | REN_NonTransactional | REN_DoNotDirty;
+			ObjectToRename->Rename(nullptr, GetTransientPackage(), RenFlags);
+			ObjectToRename->SetFlags(RF_Transient);
+			ObjectToRename->ClearFlags(RF_Public | RF_Standalone | RF_ArchetypeObject);
+			FLinkerLoad::InvalidateExport(ObjectToRename);
+		};
+
+		for (UEdGraph* OldGraph : GetViewModelBlueprint()->TemporaryGraph)
+		{
+			if (OldGraph)
+			{
+				RenameObjectToTransientPackage(OldGraph);
+			}
+		}
+		GetViewModelBlueprint()->TemporaryGraph.Reset();
+	}
+
 	const FName NAME_MetaDataBlueprintSetter = "BlueprintSetter";
 	const FName NAME_Category = "Category";
 	for (TFieldIterator<FProperty> PropertyIt(NewClass, EFieldIterationFlags::None); PropertyIt; ++PropertyIt)
@@ -186,6 +207,7 @@ void FViewModelBlueprintCompilerContext::CreateClassVariablesFromBlueprint()
 						, (FUNC_BlueprintCallable | FUNC_Public)
 						, Property->GetMetaData(NAME_Category)
 						, false);
+					GetViewModelBlueprint()->TemporaryGraph.Add(GeneratedFunction.SetterFunction);
 
 					if (GeneratedFunction.SetterFunction == nullptr)
 					{
@@ -210,6 +232,7 @@ void FViewModelBlueprintCompilerContext::CreateClassVariablesFromBlueprint()
 						, (FUNC_Private)
 						, Property->GetMetaData(NAME_Category)
 						, false);
+					GetViewModelBlueprint()->TemporaryGraph.Add(GeneratedFunction.NetRepFunction);
 
 					if (GeneratedFunction.NetRepFunction == nullptr)
 					{
@@ -280,16 +303,6 @@ void FViewModelBlueprintCompilerContext::FinishCompilingClass(UClass* Class)
 {
 	Super::FinishCompilingClass(Class);
 
-	auto DeleteGraph = [](UEdGraph* Graph)
-	{
-		if (Graph)
-		{
-			ensureMsgf(Graph->HasAnyFlags(RF_Transient), TEXT("The graph should be temporary and should be generated automatically."));
-			// GC may not have clean the graph (GC doesn't run when bRegenerateSkeletonOnly is on)
-			Graph->Rename(nullptr, GetTransientPackage(), REN_DoNotDirty | REN_ForceNoResetLoaders);
-		}
-	};
-
 	const bool bIsSkeletonOnly = CompileOptions.CompileType == EKismetCompileType::SkeletonOnly;
 	if (!bIsSkeletonOnly)
 	{
@@ -301,12 +314,6 @@ void FViewModelBlueprintCompilerContext::FinishCompilingClass(UClass* Class)
 			}
 		}
 	}
-
-	for (FGeneratedFunction& GeneratedFunction : GeneratedFunctions)
-	{
-		DeleteGraph(GeneratedFunction.SetterFunction);
-		DeleteGraph(GeneratedFunction.NetRepFunction);
-	}
 }
 
 } //namespace
diff --git a/Engine/Plugins/Runtime/ModelViewViewModel/Source/ModelViewViewModelBlueprint/Public/ViewModel/MVVMViewModelBlueprint.h b/Engine/Plugins/Runtime/ModelViewViewModel/Source/ModelViewViewModelBlueprint/Public/ViewModel/MVVMViewModelBlueprint.h
--- a/Engine/Plugins/Runtime/ModelViewViewModel/Source/ModelViewViewModelBlueprint/Public/ViewModel/MVVMViewModelBlueprint.h
+++ b/Engine/Plugins/Runtime/ModelViewViewModel/Source/ModelViewViewModelBlueprint/Public/ViewModel/MVVMViewModelBlueprint.h
@@ -9,6 +9,7 @@
 #include "MVVMViewModelBlueprint.generated.h"
 
 class FCompilerResultsLog;
+class UEdGraph;
 class UMVVMViewModelBlueprintGeneratedClass;
 
 
@@ -35,4 +36,9 @@ protected:
 	virtual bool SupportsEventGraphs() const override;
 	virtual bool SupportsDelegates() const override;
 	virtual bool SupportsAnimLayers() const override;
+
+public:
+	// Use during compilation to clean the automatically generated graph.
+	UPROPERTY(Transient)
+	TArray<TObjectPtr<UEdGraph>> TemporaryGraph;
 };
diff --git a/Engine/Plugins/Runtime/ModelViewViewModel/Source/ModelViewViewModelEditor/Private/MVVMEditorSubsystem.cpp b/Engine/Plugins/Runtime/ModelViewViewModel/Source/ModelViewViewModelEditor/Private/MVVMEditorSubsystem.cpp
--- a/Engine/Plugins/Runtime/ModelViewViewModel/Source/ModelViewViewModelEditor/Private/MVVMEditorSubsystem.cpp
+++ b/Engine/Plugins/Runtime/ModelViewViewModel/Source/ModelViewViewModelEditor/Private/MVVMEditorSubsystem.cpp
@@ -153,7 +153,10 @@ namespace UE::MVVM::Private
 	{
 		if (UEdGraph* Graph = FindExistingConversionFunctionWrapper(WidgetBlueprint, WrapperName))
 		{
-			Graph->Rename(nullptr, GetTransientPackage(), REN_DoNotDirty | REN_ForceNoResetLoaders);
+			Graph->Rename(nullptr, GetTransientPackage(), REN_DoNotDirty | REN_ForceNoResetLoaders | REN_DontCreateRedirectors);
+			Graph->SetFlags(RF_Transient);
+			Graph->ClearFlags(RF_Public | RF_Standalone | RF_ArchetypeObject);
+			FLinkerLoad::InvalidateExport(Graph);
 			WidgetBlueprint->FunctionGraphs.RemoveSingle(Graph);
 		}
 	}
