MVVM - Added a warning into the Bindings panel when you have no viewmodels created.

Bindings details panel is hidden unless there are bindings created.

#rnx

diff --git a/Engine/Plugins/Runtime/ModelViewViewModel/Source/ModelViewViewModelEditor/Private/Styling/MVVMEditorStyle.cpp b/Engine/Plugins/Runtime/ModelViewViewModel/Source/ModelViewViewModelEditor/Private/Styling/MVVMEditorStyle.cpp
--- a/Engine/Plugins/Runtime/ModelViewViewModel/Source/ModelViewViewModelEditor/Private/Styling/MVVMEditorStyle.cpp
+++ b/Engine/Plugins/Runtime/ModelViewViewModel/Source/ModelViewViewModelEditor/Private/Styling/MVVMEditorStyle.cpp
@@ -40,7 +40,8 @@ FMVVMEditorStyle::FMVVMEditorStyle()
 
 	Set("BindingView.ManageViewModels", new IMAGE_BRUSH_SVG("Slate/ViewModel", Icon16x16));
 	Set("BindingView.AddBinding", new IMAGE_BRUSH_SVG("Slate/ViewModel_AddBinding", Icon16x16));
-
+	Set("BindingView.Background", new FSlateRoundedBoxBrush(FStyleColors::Recessed, 0.0f, FStyleColors::Panel, 4.0f));
+	Set("BindingView.ViewModelWarning", new FSlateRoundedBoxBrush(FStyleColors::Transparent, 4.0f, FStyleColors::Hover, 1.0f));
 
 	// ViewModelSelectionWidget Icons
 	{
diff --git a/Engine/Plugins/Runtime/ModelViewViewModel/Source/ModelViewViewModelEditor/Private/Widgets/SMVVMViewBindingPanel.cpp b/Engine/Plugins/Runtime/ModelViewViewModel/Source/ModelViewViewModelEditor/Private/Widgets/SMVVMViewBindingPanel.cpp
--- a/Engine/Plugins/Runtime/ModelViewViewModel/Source/ModelViewViewModelEditor/Private/Widgets/SMVVMViewBindingPanel.cpp
+++ b/Engine/Plugins/Runtime/ModelViewViewModel/Source/ModelViewViewModelEditor/Private/Widgets/SMVVMViewBindingPanel.cpp
@@ -16,8 +16,10 @@
 #include "ToolMenuContext.h"
 #include "ToolMenus.h"
 #include "Tabs/MVVMBindingSummoner.h"
-
+#include "Tabs/MVVMViewModelSummoner.h"
+#include "Framework/Docking/TabManager.h"
 #include "Framework/MultiBox/MultiBoxBuilder.h"
+#include "SSimpleButton.h"
 #include "Styling/MVVMEditorStyle.h"
 #include "Widgets/Docking/SDockTab.h"
 #include "Widgets/Input/SButton.h"
@@ -27,6 +29,7 @@
 #include "Widgets/Layout/SSplitter.h"
 #include "Widgets/SMVVMManageViewModelsWidget.h"
 #include "Widgets/SMVVMViewModelContextListWidget.h"
+#include "Widgets/SMVVMViewModelPanel.h"
 #include "Widgets/SMVVMViewBindingListView.h"
 #include "Widgets/SNullWidget.h"
 #include "SPositiveActionButton.h"
@@ -320,13 +323,73 @@ TSharedRef<SWidget> SBindingsPanel::GenerateEditViewWidget()
 		BindingsList = SNew(SBindingsList, StaticCastSharedRef<SBindingsPanel>(AsShared()), MVVMExtension.Get());
 	}
 
-	TSharedRef<SWidget> BindingWidget = BindingsList ? BindingsList.ToSharedRef() : SNullWidget::NullWidget;
+	TSharedRef<SWidget> BindingWidget = SNew(SOverlay)
+		+ SOverlay::Slot()
+		.HAlign(HAlign_Center)
+		.VAlign(VAlign_Center)
+		[
+			SNew(SBorder)
+			.BorderImage(FMVVMEditorStyle::Get().GetBrush("BindingView.ViewModelWarning"))
+			.Visibility_Lambda([this]()
+				{
+					if (UMVVMWidgetBlueprintExtension_View* MVVMExtensionPtr = MVVMExtension.Get())
+					{
+						if (MVVMExtensionPtr->GetBlueprintView() != nullptr && 
+							MVVMExtensionPtr->GetBlueprintView()->GetViewModels().Num() > 0)
+						{
+							return EVisibility::Collapsed;
+						}
+					}
+					return EVisibility::Visible;
+				})
+			[
+				SNew(SHorizontalBox)
+				+ SHorizontalBox::Slot()
+				.Padding(20, 20, 12, 20)
+				.VAlign(VAlign_Center)
+				.HAlign(HAlign_Left)
+				.AutoWidth()
+				[
+					SNew(SImage)
+					.Image(FAppStyle::Get().GetBrush("Icons.Warning"))
+				]
+				+ SHorizontalBox::Slot()
+				.Padding(0, 0, 8, 0)
+				.VAlign(VAlign_Center)
+				.HAlign(HAlign_Center)
+				.AutoWidth()
+				[
+					SNew(STextBlock)
+					.Text(LOCTEXT("MissingViewModel", "This editor requires a viewmodel that widgets can bind to, would you like to create a viewmodel now?"))
+				]
+				+ SHorizontalBox::Slot()
+				.Padding(0, 0, 20, 0)
+				.VAlign(VAlign_Center)
+				.HAlign(HAlign_Right)
+				[
+					SNew(SButton)
+					.OnClicked(this, &SBindingsPanel::HandleCreateViewModelClicked)
+					[
+						SNew(STextBlock)
+						.Text(LOCTEXT("CreateViewModel", "Create Viewmodel"))
+					]
+				]
+			]
+		]
+		+ SOverlay::Slot()
+		[
+			SNew(SBox)
+			.Visibility(this, &SBindingsPanel::GetVisibility, true)
+			[
+				BindingsList ? BindingsList.ToSharedRef() : SNullWidget::NullWidget
+			]
+		];
 
 	return SNew(SVerticalBox)
 
 		+ SVerticalBox::Slot()
 		.AutoHeight()
-		.Padding(FMargin(0.0f, 2.0f))
+		.Padding(8.0f, 2.0f, 0.0f, 2.0f)
 		[
 			ToolbarBuilderGlobal.MakeWidget()
 		]
@@ -335,8 +398,7 @@ TSharedRef<SWidget> SBindingsPanel::GenerateEditViewWidget()
 		.FillHeight(1.0f)
 		[
 			SNew(SBorder)
-			.BorderImage(FAppStyle::GetBrush("ToolPanel.GroupBorder"))
-			.Padding(FMargin(4.0f, 4.0f))
+			.BorderImage(FMVVMEditorStyle::Get().GetBrush("BindingView.Background"))
 			[
 				SNew(SSplitter)
 				.Orientation(EOrientation::Orient_Horizontal)
@@ -351,6 +413,7 @@ TSharedRef<SWidget> SBindingsPanel::GenerateEditViewWidget()
 				.Value(0.25f)
 				[
 					SAssignNew(DetailContainer, SBorder)
+					.Visibility(this, &SBindingsPanel::GetVisibility, true)
 					[
 						DetailsView.ToSharedRef()
 					]
@@ -359,6 +422,35 @@ TSharedRef<SWidget> SBindingsPanel::GenerateEditViewWidget()
 		];
 }
 
+EVisibility SBindingsPanel::GetVisibility(bool bVisibleWithBindings) const
+{
+	if (UMVVMWidgetBlueprintExtension_View* MVVMExtensionPtr = MVVMExtension.Get())
+	{
+		if (MVVMExtensionPtr->GetBlueprintView() != nullptr &&
+			MVVMExtensionPtr->GetBlueprintView()->GetNumBindings() > 0)
+		{
+			return bVisibleWithBindings ? EVisibility::Visible : EVisibility::Collapsed;
+		}
+	}
+	return bVisibleWithBindings ? EVisibility::Collapsed : EVisibility::Visible;
+}
+
+FReply SBindingsPanel::HandleCreateViewModelClicked()
+{
+	if (TSharedPtr<FWidgetBlueprintEditor> BPEditor = WeakBlueprintEditor.Pin())
+	{
+		if (TSharedPtr<SDockTab> DockTab = BPEditor->GetTabManager()->TryInvokeTab(FTabId(FViewModelSummoner::TabID)))
+		{
+			TSharedPtr<SWidget> ViewModelPanel = DockTab->GetContent();
+			if (ViewModelPanel.IsValid())
+			{
+				StaticCastSharedPtr<SMVVMViewModelPanel>(ViewModelPanel)->OpenAddViewModelMenu();
+			}
+		}
+	}
+	return FReply::Handled();
+}
+
 } // namespace UE::MVVM
 
 #undef LOCTEXT_NAMESPACE
diff --git a/Engine/Plugins/Runtime/ModelViewViewModel/Source/ModelViewViewModelEditor/Private/Widgets/SMVVMViewBindingPanel.h b/Engine/Plugins/Runtime/ModelViewViewModel/Source/ModelViewViewModelEditor/Private/Widgets/SMVVMViewBindingPanel.h
--- a/Engine/Plugins/Runtime/ModelViewViewModel/Source/ModelViewViewModelEditor/Private/Widgets/SMVVMViewBindingPanel.h
+++ b/Engine/Plugins/Runtime/ModelViewViewModel/Source/ModelViewViewModelEditor/Private/Widgets/SMVVMViewBindingPanel.h
@@ -56,6 +56,10 @@ private:
 
 	void HandleExtensionAdded(UBlueprintExtension* NewExtension);
 
+	FReply HandleCreateViewModelClicked();
+
+	EVisibility GetVisibility(bool bVisibleWithBindings) const;
+
 private:
 	TWeakPtr<FWidgetBlueprintEditor> WeakBlueprintEditor;
 	TSharedPtr<SBindingsList> BindingsList;
diff --git a/Engine/Plugins/Runtime/ModelViewViewModel/Source/ModelViewViewModelEditor/Private/Widgets/SMVVMViewModelPanel.cpp b/Engine/Plugins/Runtime/ModelViewViewModel/Source/ModelViewViewModelEditor/Private/Widgets/SMVVMViewModelPanel.cpp
--- a/Engine/Plugins/Runtime/ModelViewViewModel/Source/ModelViewViewModelEditor/Private/Widgets/SMVVMViewModelPanel.cpp
+++ b/Engine/Plugins/Runtime/ModelViewViewModel/Source/ModelViewViewModelEditor/Private/Widgets/SMVVMViewModelPanel.cpp
@@ -202,6 +202,12 @@ void SMVVMViewModelPanel::HandleViewModelsUpdated()
 }
 
 
+void SMVVMViewModelPanel::OpenAddViewModelMenu()
+{
+	AddMenuButton->SetIsMenuOpen(true, true);
+}
+
+
 TSharedRef<SWidget> SMVVMViewModelPanel::MakeAddMenu()
 {
 	const UWidgetBlueprint* WidgetBlueprint = nullptr;
@@ -246,8 +252,7 @@ void SMVVMViewModelPanel::HandleAddMenuViewModel(const UClass* SelectedClass)
 					UMVVMBlueprintView* CurrentBlueprintView = WeakBlueprintView.Get();
 					if (!CurrentBlueprintView)
 					{
-						EditorSubsystem->RequestView(WidgetBlueprint);
-						CurrentBlueprintView = EditorSubsystem->GetView(WidgetBlueprint);
+						CurrentBlueprintView = EditorSubsystem->RequestView(WidgetBlueprint);
 						WeakBlueprintView = CurrentBlueprintView;
 						ViewModelsUpdatedHandle = CurrentBlueprintView->OnViewModelsUpdated.AddSP(this, &SMVVMViewModelPanel::HandleViewModelsUpdated);
 					}
diff --git a/Engine/Plugins/Runtime/ModelViewViewModel/Source/ModelViewViewModelEditor/Private/Widgets/SMVVMViewModelPanel.h b/Engine/Plugins/Runtime/ModelViewViewModel/Source/ModelViewViewModelEditor/Private/Widgets/SMVVMViewModelPanel.h
--- a/Engine/Plugins/Runtime/ModelViewViewModel/Source/ModelViewViewModelEditor/Private/Widgets/SMVVMViewModelPanel.h
+++ b/Engine/Plugins/Runtime/ModelViewViewModel/Source/ModelViewViewModelEditor/Private/Widgets/SMVVMViewModelPanel.h
@@ -32,6 +32,8 @@ public:
 	void Construct(const FArguments& InArgs, TSharedPtr<FWidgetBlueprintEditor> Editor);
 	virtual ~SMVVMViewModelPanel();
 
+	void OpenAddViewModelMenu();
+
 private:
 	void HandleViewUpdated(UBlueprintExtension* Extension);
 	void HandleViewModelsUpdated();
