MVVM - Added FindConversionFunctionArgumentPin() and GetConversionFunction() methods to UMVVMEditorSubsystem.

#rnx

diff --git a/Engine/Plugins/Runtime/ModelViewViewModel/Source/ModelViewViewModelEditor/Private/MVVMEditorSubsystem.cpp b/Engine/Plugins/Runtime/ModelViewViewModel/Source/ModelViewViewModelEditor/Private/MVVMEditorSubsystem.cpp
--- a/Engine/Plugins/Runtime/ModelViewViewModel/Source/ModelViewViewModelEditor/Private/MVVMEditorSubsystem.cpp
+++ b/Engine/Plugins/Runtime/ModelViewViewModel/Source/ModelViewViewModelEditor/Private/MVVMEditorSubsystem.cpp
@@ -295,6 +295,23 @@ UEdGraph* UMVVMEditorSubsystem::GetConversionFunctionGraph(const UWidgetBlueprin
 	return UE::MVVM::ConversionFunctionHelper::GetGraph(WidgetBlueprint, Binding, bSourceToDestination);
 }
 
+UFunction* UMVVMEditorSubsystem::GetConversionFunction(const UWidgetBlueprint* WidgetBlueprint, const FMVVMBlueprintViewBinding& Binding, bool bSourceToDestination) const
+{
+	const UEdGraph* Graph = GetConversionFunctionGraph(WidgetBlueprint, Binding, bSourceToDestination);
+
+	TArray<UK2Node_CallFunction*> FunctionNodes;
+	Graph->GetNodesOfClass<UK2Node_CallFunction>(FunctionNodes);
+
+	if (FunctionNodes.Num() != 1)
+	{
+		// ambiguous result, no idea what our function node is
+		return nullptr;
+	}
+
+	const UK2Node_CallFunction* CallFunctionNode = FunctionNodes[0];
+	return CallFunctionNode->GetTargetFunction();
+}
+
 UEdGraph* UMVVMEditorSubsystem::CreateConversionFunctionWrapperGraph(UWidgetBlueprint* WidgetBlueprint, const FMVVMBlueprintViewBinding& Binding, const UFunction* ConversionFunction, bool bSourceToDestination)
 {
 	if (UMVVMBlueprintView* View = GetView(WidgetBlueprint))
@@ -367,6 +384,25 @@ UEdGraph* UMVVMEditorSubsystem::CreateConversionFunctionWrapperGraph(UWidgetBlue
 	return nullptr;
 }
 
+UEdGraphPin* UMVVMEditorSubsystem::FindConversionFunctionArgumentPin(const UWidgetBlueprint* WidgetBlueprint, const FMVVMBlueprintViewBinding& Binding, FName ParameterName, bool bSourceToDestination)
+{
+	const UEdGraph* Graph = GetConversionFunctionGraph(WidgetBlueprint, Binding, bSourceToDestination);
+
+	TArray<UK2Node_CallFunction*> FunctionNodes;
+	Graph->GetNodesOfClass<UK2Node_CallFunction>(FunctionNodes);
+
+	if (FunctionNodes.Num() != 1)
+	{
+		// ambiguous result, no idea what our function node is
+		return nullptr;
+	}
+
+	UK2Node_CallFunction* CallFunctionNode = FunctionNodes[0];
+	UEdGraphPin* Pin = CallFunctionNode->FindPin(ParameterName, EGPD_Input);
+	
+	return Pin;
+}
+
 void UMVVMEditorSubsystem::SetSourceToDestinationConversionFunction(UWidgetBlueprint* WidgetBlueprint, FMVVMBlueprintViewBinding& Binding, const UFunction* ConversionFunction)
 {
 	if (UMVVMBlueprintView* View = GetView(WidgetBlueprint))
diff --git a/Engine/Plugins/Runtime/ModelViewViewModel/Source/ModelViewViewModelEditor/Public/MVVMEditorSubsystem.h b/Engine/Plugins/Runtime/ModelViewViewModel/Source/ModelViewViewModelEditor/Public/MVVMEditorSubsystem.h
--- a/Engine/Plugins/Runtime/ModelViewViewModel/Source/ModelViewViewModelEditor/Public/MVVMEditorSubsystem.h
+++ b/Engine/Plugins/Runtime/ModelViewViewModel/Source/ModelViewViewModelEditor/Public/MVVMEditorSubsystem.h
@@ -64,11 +64,15 @@ public:
 	UFUNCTION(BlueprintCallable, Category = "MVVM")
 	UEdGraph* GetConversionFunctionGraph(const UWidgetBlueprint* WidgetBlueprint, const FMVVMBlueprintViewBinding& Binding, bool bSourceToDestination) const;
 
+	UFUNCTION(BlueprintCallable, Category = "MVVM")
+	UFunction* GetConversionFunction(const UWidgetBlueprint* WidgetBlueprint, const FMVVMBlueprintViewBinding& Binding, bool bSourceToDestination) const;
+
 	UFUNCTION(BlueprintCallable, Category = "MVVM")
 	TArray<UFunction*> GetAvailableConversionFunctions(const UWidgetBlueprint* WidgetBlueprint, const FMVVMBlueprintPropertyPath& Source, const FMVVMBlueprintPropertyPath& Destination) const;
 
 	FMVVMBlueprintPropertyPath GetPathForConversionFunctionArgument(const UWidgetBlueprint* WidgetBlueprint, const FMVVMBlueprintViewBinding& Binding, FName ArgumentName, bool bSourceToDestination) const;
 	void SetPathForConversionFunctionArgument(UWidgetBlueprint* WidgetBlueprint, FMVVMBlueprintViewBinding& Binding, FName ArgumentName, const FMVVMBlueprintPropertyPath& Path, bool bSourceToDestination) const;
+	UEdGraphPin* FindConversionFunctionArgumentPin(const UWidgetBlueprint* WidgetBlueprint, const FMVVMBlueprintViewBinding& Binding, FName ParameterName, bool bSourceToDestination);
 
 	TArray<UE::MVVM::FBindingSource> GetBindableWidgets(const UWidgetBlueprint* WidgetBlueprint) const;
 	
