MVVM - SMVVMSourceEntry now displays the property type in the tooltip.

#rnx

diff --git a/Engine/Plugins/Runtime/ModelViewViewModel/Source/ModelViewViewModelEditor/Private/Widgets/SMVVMSourceEntry.cpp b/Engine/Plugins/Runtime/ModelViewViewModel/Source/ModelViewViewModelEditor/Private/Widgets/SMVVMSourceEntry.cpp
--- a/Engine/Plugins/Runtime/ModelViewViewModel/Source/ModelViewViewModelEditor/Private/Widgets/SMVVMSourceEntry.cpp
+++ b/Engine/Plugins/Runtime/ModelViewViewModel/Source/ModelViewViewModelEditor/Private/Widgets/SMVVMSourceEntry.cpp
@@ -83,6 +83,19 @@ void SSourceEntry::RefreshSource(const FBindingSource& Source)
 	Image->SetImage(Private::GetSourceIcon(Source));
 	Image->SetColorAndOpacity(Private::GetSourceColor(Source));
 	Label->SetText(Private::GetSourceDisplayName(Source));
+
+	if (Source.IsValid())
+	{
+		if (Source.Class != nullptr)
+		{
+			const FText ToolTipText = FText::Join(FText::FromString(TEXT("\n")), Source.DisplayName, Source.Class->GetDisplayNameText());
+			SetToolTipText(ToolTipText);
+		}
+		else
+		{
+			SetToolTipText(Source.DisplayName);
+		}
+	}
 }
 
 } // namespace UE::MVVM
