MVVM - Conversion function parameters are now displayed in the Bindings list below the relevant binding.

Numerous pieces of UI polish: tooltips, padding, colours and wording.

#rnx

diff --git a/Engine/Plugins/Runtime/ModelViewViewModel/Source/ModelViewViewModelBlueprint/Private/MVVMBlueprintView.cpp b/Engine/Plugins/Runtime/ModelViewViewModel/Source/ModelViewViewModelBlueprint/Private/MVVMBlueprintView.cpp
--- a/Engine/Plugins/Runtime/ModelViewViewModel/Source/ModelViewViewModelBlueprint/Private/MVVMBlueprintView.cpp
+++ b/Engine/Plugins/Runtime/ModelViewViewModel/Source/ModelViewViewModelBlueprint/Private/MVVMBlueprintView.cpp
@@ -115,8 +115,11 @@ FMVVMBlueprintViewBinding* UMVVMBlueprintView::FindBinding(const UWidget* Widget
 
 void UMVVMBlueprintView::RemoveBindingAt(int32 Index)
 {
-	Bindings.RemoveAt(Index);
-	OnBindingsUpdated.Broadcast();
+	if (Bindings.IsValidIndex(Index))
+	{
+		Bindings.RemoveAt(Index);
+		OnBindingsUpdated.Broadcast();
+	}
 }
 
 void UMVVMBlueprintView::RemoveBinding(const FMVVMBlueprintViewBinding* Binding)
@@ -130,8 +133,7 @@ void UMVVMBlueprintView::RemoveBinding(const FMVVMBlueprintViewBinding* Binding)
 		}
 	}
 
-	Bindings.RemoveAt(Index);
-	OnBindingsUpdated.Broadcast();
+	RemoveBindingAt(Index);
 }
 
 FMVVMBlueprintViewBinding& UMVVMBlueprintView::AddBinding(const UWidget* Widget, const FProperty* Property)
diff --git a/Engine/Plugins/Runtime/ModelViewViewModel/Source/ModelViewViewModelEditor/Private/Customizations/MVVMConversionPathCustomization.cpp b/Engine/Plugins/Runtime/ModelViewViewModel/Source/ModelViewViewModelEditor/Private/Customizations/MVVMConversionPathCustomization.cpp
--- a/Engine/Plugins/Runtime/ModelViewViewModel/Source/ModelViewViewModelEditor/Private/Customizations/MVVMConversionPathCustomization.cpp
+++ b/Engine/Plugins/Runtime/ModelViewViewModel/Source/ModelViewViewModelEditor/Private/Customizations/MVVMConversionPathCustomization.cpp
@@ -9,30 +9,22 @@
 #include "DetailWidgetRow.h"
 #include "EdGraphSchema_K2.h"
 #include "EdGraph/EdGraph.h"
-#include "EdGraph/EdGraphNode.h"
 #include "EdGraph/EdGraphPin.h"
 #include "Editor.h"
 #include "IDetailChildrenBuilder.h"
 #include "IDetailGroup.h"
 #include "IPropertyUtilities.h"
 #include "K2Node_CallFunction.h"
-#include "K2Node_VariableGet.h"
 #include "MVVMBlueprintViewBinding.h"
 #include "MVVMEditorSubsystem.h"
 #include "MVVMSubsystem.h"
-#include "NodeFactory.h"
 #include "PropertyHandle.h"
-#include "ScopedTransaction.h"
-#include "SGraphPin.h"
-#include "Styling/MVVMEditorStyle.h"
 #include "UObject/StructOnScope.h"
 #include "WidgetBlueprint.h"
-#include "Widgets/Images/SImage.h"
-#include "Widgets/Input/SCheckBox.h"
 #include "Widgets/Input/SEditableTextBox.h"
 #include "Widgets/Layout/SBox.h"
 #include "Widgets/SMVVMConversionPath.h"
-#include "Widgets/SMVVMFieldSelector.h"
+#include "Widgets/SMVVMFunctionParameter.h"
 #include "Widgets/Text/STextBlock.h"
 
 #define LOCTEXT_NAMESPACE "MVVMConversionPath"
@@ -70,8 +62,6 @@ namespace UE::MVVM
 
 	void FConversionPathCustomization::CustomizeChildren(TSharedRef<IPropertyHandle> InPropertyHandle, IDetailChildrenBuilder& ChildBuilder, IPropertyTypeCustomizationUtils& CustomizationUtils)
 	{
-		ArgumentPinWidgets.Reset();
-
 		AddRowForProperty(ChildBuilder, InPropertyHandle, true);
 		AddRowForProperty(ChildBuilder, InPropertyHandle, false);
 	}
@@ -216,29 +206,6 @@ namespace UE::MVVM
 				continue;
 			}
 
-			const FName ArgumentName = Pin->GetFName();
-
-			// create a new pin widget so that we can get the default value widget out of it
-			TSharedPtr<SGraphPin> PinWidget = FNodeFactory::CreateK2PinWidget(Pin);
-
-			TSharedRef<SWidget> ValueWidget = SNullWidget::NullWidget;
-			if (PinWidget.IsValid())
-			{
-				ArgumentPinWidgets.Add(PinWidget);
-				ValueWidget = PinWidget->GetDefaultValueWidget();
-			}
-
-			if (ValueWidget == SNullWidget::NullWidget)
-			{
-				ValueWidget = SNew(STextBlock)
-					.Text(LOCTEXT("DefaultValue", "Default Value"))
-					.TextStyle(FAppStyle::Get(), "HintText");
-			}
-
-			ValueWidget->SetVisibility(TAttribute<EVisibility>::CreateSP(this, &FConversionPathCustomization::GetArgumentWidgetVisibility, ArgumentName, bSourceToDestination, true));
-
-			TSharedPtr<UE::MVVM::SFieldSelector> FieldSelector;
-
 			ChildBuilder.AddCustomRow(Pin->GetDisplayName())
 			.NameContent()
 			[
@@ -251,198 +218,16 @@ namespace UE::MVVM
 			]
 			.ValueContent()
 			[
-				// TODO (sebastiann): This would be easier to accomplish if this entire thing was a widget
-				SNew(SHorizontalBox)
-				+ SHorizontalBox::Slot()
-				[
-					SNew(SOverlay)
-					+ SOverlay::Slot()
-					[
-						ValueWidget
-					]
-					+ SOverlay::Slot()
-					[
-						SAssignNew(FieldSelector, UE::MVVM::SFieldSelector, WidgetBlueprint, bSourceToDestination)
-						.Visibility(this, &FConversionPathCustomization::GetArgumentWidgetVisibility, ArgumentName, bSourceToDestination, false)
-						.SelectedField(this, &FConversionPathCustomization::OnGetSelectedField, ArgumentName, bSourceToDestination)
-						.BindingMode(this, &FConversionPathCustomization::GetBindingMode)
-						.OnSelectionChanged(this, &FConversionPathCustomization::OnSetProperty, ArgumentName, bSourceToDestination)
-					]
-				]
-				+ SHorizontalBox::Slot()
-				.AutoWidth()
-				.HAlign(HAlign_Right)
-				.VAlign(VAlign_Center)
-				.Padding(5, 0, 0, 0)
-				[
-					SNew(SCheckBox)
-					.ToolTipText(LOCTEXT("BindArgument", "Bind this argument to a property."))
-					.Style(FAppStyle::Get(), "ToggleButtonCheckbox")
-					.IsChecked(this, &FConversionPathCustomization::OnGetIsArgumentBound, ArgumentName, bSourceToDestination)
-					.OnCheckStateChanged(this, &FConversionPathCustomization::OnCheckedBindArgument, ArgumentName, bSourceToDestination)
-					[
-						SNew(SBox)
-						.WidthOverride(16)
-						.HeightOverride(16)
-						[
-							SNew(SImage)
-							.Image_Lambda([this, ArgumentName, bSourceToDestination]()
-								{
-									ECheckBoxState CheckState = OnGetIsArgumentBound(ArgumentName, bSourceToDestination);
-									return CheckState == ECheckBoxState::Checked ? FAppStyle::GetBrush("Icons.Link") : FAppStyle::GetBrush("Icons.Unlink");
-								})
-						]
-					]
-				]
+				SNew(UE::MVVM::SFunctionParameter)
+				.WidgetBlueprint(WidgetBlueprint)
+				.Binding(ViewBindings[0])
+				.ParameterName(Pin->GetFName())
+				.SourceToDestination(bSourceToDestination)
+				.OnGetBindingMode(this, &FConversionPathCustomization::GetBindingMode)
 			];
-
-			ArgumentFieldSelectors.Add(FieldSelector);
 		}
 	}
 
-	void FConversionPathCustomization::OnSetProperty(FMVVMBlueprintPropertyPath NewSelection, FName ArgumentName, bool bSourceToDestination)
-	{
-		UMVVMEditorSubsystem* EditorSubsystem = GEditor->GetEditorSubsystem<UMVVMEditorSubsystem>();
-
-		TArray<void*> RawBindings;
-		ParentHandle->AccessRawData(RawBindings);
-
-		for (void* RawBinding : RawBindings)
-		{
-			FMVVMBlueprintViewBinding* Binding = reinterpret_cast<FMVVMBlueprintViewBinding*>(RawBinding);
-			EditorSubsystem->SetPathForConversionFunctionArgument(WidgetBlueprint, *Binding, ArgumentName, NewSelection, bSourceToDestination);
-		}
-	}
-
-	FMVVMBlueprintPropertyPath FConversionPathCustomization::OnGetSelectedField(FName ArgumentName, bool bSourceToDestination) const
-	{
-		UMVVMEditorSubsystem* EditorSubsystem = GEditor->GetEditorSubsystem<UMVVMEditorSubsystem>();
-		
-		TArray<void*> RawBindings;
-		ParentHandle->AccessRawData(RawBindings);
-
-		FMVVMBlueprintPropertyPath Path;
-		bool bFirst = true;
-
-		for (void* RawBinding : RawBindings)
-		{
-			FMVVMBlueprintViewBinding* Binding = reinterpret_cast<FMVVMBlueprintViewBinding*>(RawBinding);
-			
-			FMVVMBlueprintPropertyPath CurrentPath = EditorSubsystem->GetPathForConversionFunctionArgument(WidgetBlueprint, *Binding, ArgumentName, bSourceToDestination);
-			if (bFirst)
-			{
-				Path = CurrentPath;
-				bFirst = false;
-			}
-			else if (Path != CurrentPath)
-			{
-				return FMVVMBlueprintPropertyPath();
-			}
-		}
-
-		return Path;
-	}
-
-	ECheckBoxState FConversionPathCustomization::OnGetIsArgumentBound(FName ArgumentName, bool bSourceToDestination) const
-	{
-		UMVVMEditorSubsystem* EditorSubsystem = GEditor->GetEditorSubsystem<UMVVMEditorSubsystem>();
-
-		TArray<void*> RawBindings;
-		ParentHandle->AccessRawData(RawBindings);
-
-		FMVVMBlueprintPropertyPath Path;
-		bool bFirst = true;
-
-		for (void* RawBinding : RawBindings)
-		{
-			FMVVMBlueprintViewBinding* Binding = reinterpret_cast<FMVVMBlueprintViewBinding*>(RawBinding);
-
-			FMVVMBlueprintPropertyPath CurrentPath = EditorSubsystem->GetPathForConversionFunctionArgument(WidgetBlueprint, *Binding, ArgumentName, bSourceToDestination);
-			if (bFirst)
-			{
-				Path = CurrentPath;
-				bFirst = false;
-			}
-			else if (Path != CurrentPath)
-			{
-				return ECheckBoxState::Undetermined;
-			}
-		}
-
-		return Path.IsEmpty() ? ECheckBoxState::Unchecked : ECheckBoxState::Checked;
-	}
-
-	void FConversionPathCustomization::OnCheckedBindArgument(ECheckBoxState CheckState, FName ArgumentName, bool bSourceToDestination)
-	{
-		UMVVMEditorSubsystem* EditorSubsystem = GEditor->GetEditorSubsystem<UMVVMEditorSubsystem>();
-
-		FMVVMBlueprintPropertyPath Path;
-
-		if (CheckState == ECheckBoxState::Checked)
-		{
-			// HACK: Just set a placeholder viewmodel reference
-			if (const UMVVMBlueprintView* View = EditorSubsystem->GetView(WidgetBlueprint))
-			{
-				const TArrayView<const FMVVMBlueprintViewModelContext> ViewModels = View->GetViewModels();
-				if (ViewModels.Num() > 0)
-				{
-					Path.SetViewModelId(ViewModels[0].GetViewModelId());
-				}
-			}
-		}
-
-		TArray<void*> RawBindings;
-		ParentHandle->AccessRawData(RawBindings);
-
-		for (void* RawBinding : RawBindings)
-		{
-			FMVVMBlueprintViewBinding* Binding = reinterpret_cast<FMVVMBlueprintViewBinding*>(RawBinding);
-			EditorSubsystem->SetPathForConversionFunctionArgument(WidgetBlueprint, *Binding, ArgumentName, Path, bSourceToDestination);
-		}
-	}
-
-	EVisibility FConversionPathCustomization::GetArgumentWidgetVisibility(FName ArgumentName, bool bSourceToDestination, bool bDefaultValue) const
-	{
-		FMVVMBlueprintPropertyPath Path = OnGetSelectedField(ArgumentName, bSourceToDestination);
-		if (Path.IsEmpty())
-		{
-			return bDefaultValue ? EVisibility::Visible : EVisibility::Collapsed;
-		}
-		else
-		{
-			return bDefaultValue ? EVisibility::Collapsed : EVisibility::Visible;
-		}
-	}
-
-	FEdGraphPinType FConversionPathCustomization::GetArgumentPinType(FName ArgumentName, bool bSourceToDestination) const 
-	{
-		TArray<void*> RawData;
-		ParentHandle->AccessRawData(RawData);
-		if (RawData.Num() != 1)
-		{
-			return FEdGraphPinType();
-		}
-
-		UMVVMEditorSubsystem* EditorSubsystem = GEditor->GetEditorSubsystem<UMVVMEditorSubsystem>();
-		UEdGraph* Graph = EditorSubsystem->GetConversionFunctionGraph(WidgetBlueprint, *reinterpret_cast<const FMVVMBlueprintViewBinding*>(RawData[0]), bSourceToDestination);
-		if (Graph == nullptr)
-		{
-			return FEdGraphPinType();
-		}
-
-		TArray<UK2Node_CallFunction*> FunctionNodes;
-		Graph->GetNodesOfClass<UK2Node_CallFunction>(FunctionNodes);
-		if (FunctionNodes.Num() != 1)
-		{
-			// ambiguous result, no idea what our function node is
-			return FEdGraphPinType();
-		}
-
-		UK2Node_CallFunction* CallFunctionNode = FunctionNodes[0];
-		UEdGraphPin* Pin = CallFunctionNode->FindPin(ArgumentName);
-		return Pin != nullptr ? Pin->PinType : FEdGraphPinType();
-	}
-
 	EMVVMBindingMode FConversionPathCustomization::GetBindingMode() const
 	{
 		uint8 EnumValue = 0;
diff --git a/Engine/Plugins/Runtime/ModelViewViewModel/Source/ModelViewViewModelEditor/Private/Customizations/MVVMConversionPathCustomization.h b/Engine/Plugins/Runtime/ModelViewViewModel/Source/ModelViewViewModelEditor/Private/Customizations/MVVMConversionPathCustomization.h
--- a/Engine/Plugins/Runtime/ModelViewViewModel/Source/ModelViewViewModelEditor/Private/Customizations/MVVMConversionPathCustomization.h
+++ b/Engine/Plugins/Runtime/ModelViewViewModel/Source/ModelViewViewModelEditor/Private/Customizations/MVVMConversionPathCustomization.h
@@ -42,15 +42,6 @@ private:
 	void OnFunctionPathChanged(const UFunction* NewFunction, TSharedRef<IPropertyHandle> Property, bool bSourceToDestination);
 	void RefreshDetailsView() const;
 
-	FMVVMBlueprintPropertyPath OnGetSelectedField(FName ArgumentName, bool bSourceToDestination) const;
-	void OnSetProperty(FMVVMBlueprintPropertyPath NewSelection, FName ArgumentName, bool bSourceToDestination);
-
-	ECheckBoxState OnGetIsArgumentBound(FName ArgumentName, bool bSourceToDestination) const;
-	void OnCheckedBindArgument(ECheckBoxState CheckState, FName ArgumentName, bool bSourceToDestination);
-	EVisibility GetArgumentWidgetVisibility(FName ArgumentName, bool bSourceToDestination, bool bDefaultValue) const;
-
-	FEdGraphPinType GetArgumentPinType(FName ArgumentName, bool bSourceToDestination) const;
-
 	EMVVMBindingMode GetBindingMode() const;
 
 private:
@@ -59,9 +50,6 @@ private:
 	TSharedPtr<IPropertyHandle> ParentHandle;
 	TSharedPtr<IPropertyHandle> BindingModeHandle;
 	TWeakPtr<IPropertyUtilities> WeakPropertyUtilities;
-
-	TArray<TSharedPtr<SGraphPin>> ArgumentPinWidgets; 
-	TArray<TSharedPtr<UE::MVVM::SFieldSelector>> ArgumentFieldSelectors;
 };
 
 } // namespace UE::MVVM
diff --git a/Engine/Plugins/Runtime/ModelViewViewModel/Source/ModelViewViewModelEditor/Private/Customizations/MVVMPropertyPathCustomization.cpp b/Engine/Plugins/Runtime/ModelViewViewModel/Source/ModelViewViewModelEditor/Private/Customizations/MVVMPropertyPathCustomization.cpp
--- a/Engine/Plugins/Runtime/ModelViewViewModel/Source/ModelViewViewModelEditor/Private/Customizations/MVVMPropertyPathCustomization.cpp
+++ b/Engine/Plugins/Runtime/ModelViewViewModel/Source/ModelViewViewModelEditor/Private/Customizations/MVVMPropertyPathCustomization.cpp
@@ -127,7 +127,7 @@ EMVVMBindingMode FPropertyPathCustomization::GetCurrentBindingMode() const
 	return EMVVMBindingMode::OneWayToDestination;
 }
 
-void FPropertyPathCustomization::OnPropertySelectionChanged(FMVVMBlueprintPropertyPath Selected)
+void FPropertyPathCustomization::OnPropertySelectionChanged(const FMVVMBlueprintPropertyPath& Selected)
 {
 	if (bPropertySelectionChanging)
 	{
diff --git a/Engine/Plugins/Runtime/ModelViewViewModel/Source/ModelViewViewModelEditor/Private/Customizations/MVVMPropertyPathCustomization.h b/Engine/Plugins/Runtime/ModelViewViewModel/Source/ModelViewViewModelEditor/Private/Customizations/MVVMPropertyPathCustomization.h
--- a/Engine/Plugins/Runtime/ModelViewViewModel/Source/ModelViewViewModelEditor/Private/Customizations/MVVMPropertyPathCustomization.h
+++ b/Engine/Plugins/Runtime/ModelViewViewModel/Source/ModelViewViewModelEditor/Private/Customizations/MVVMPropertyPathCustomization.h
@@ -30,7 +30,7 @@ public:
 	}
 
 private:
-	void OnPropertySelectionChanged(FMVVMBlueprintPropertyPath Selected);
+	void OnPropertySelectionChanged(const FMVVMBlueprintPropertyPath& Selected);
 
 	void OnOtherPropertyChanged();
 
diff --git a/Engine/Plugins/Runtime/ModelViewViewModel/Source/ModelViewViewModelEditor/Private/Widgets/SMVVMFieldEntry.cpp b/Engine/Plugins/Runtime/ModelViewViewModel/Source/ModelViewViewModelEditor/Private/Widgets/SMVVMFieldEntry.cpp
--- a/Engine/Plugins/Runtime/ModelViewViewModel/Source/ModelViewViewModelEditor/Private/Widgets/SMVVMFieldEntry.cpp
+++ b/Engine/Plugins/Runtime/ModelViewViewModel/Source/ModelViewViewModelEditor/Private/Widgets/SMVVMFieldEntry.cpp
@@ -18,19 +18,31 @@ FText GetFieldDisplayName(const FMVVMConstFieldVariant& Field)
 	{
 		return Field.GetProperty()->GetDisplayNameText();
 	}
-	else if (Field.IsFunction())
+	if (Field.IsFunction())
 	{
 		return Field.GetFunction()->GetDisplayNameText();
 	}
 	return LOCTEXT("None", "<None>");
 }
 
+FText GetFieldToolTip(const FMVVMConstFieldVariant& Field)
+{
+	if (Field.IsFunction())
+	{
+		return Field.GetFunction()->GetToolTipText();
+	}
+	if (Field.IsProperty())
+	{
+		return FText::Join(FText::FromString(TEXT("\n")), Field.GetProperty()->GetToolTipText(), FText::FromString(Field.GetProperty()->GetCPPType()));
+	}
+	return FText::GetEmpty();
+}
+
 } // namespace Private
 
 void SFieldEntry::Construct(const FArguments& InArgs)
 {
 	Field = InArgs._Field;
-	OnValidateField = InArgs._OnValidateField;
 
 	ChildSlot
 	[
@@ -58,24 +70,6 @@ void SFieldEntry::Construct(const FArguments& InArgs)
 
 void SFieldEntry::Refresh()
 {
-	FText ToolTipText = FText::GetEmpty();
-	bool bEnabled = true;
-	if (OnValidateField.IsBound())
-	{
-		TValueOrError<bool, FString> Result = OnValidateField.Execute(Field);
-		if (Result.HasError())
-		{
-			ToolTipText = FText::FromString(Result.GetError());
-			bEnabled = false;
-		}
-		else
-		{
-			bEnabled = true;
-		}
-	}
-
-	SetEnabled(bEnabled);
-
 	UE::MVVM::FMVVMConstFieldVariant Variant;
 
 	TArray<UE::MVVM::FMVVMConstFieldVariant> Fields = Field.GetFields();
@@ -84,14 +78,7 @@ void SFieldEntry::Refresh()
 		Variant = Fields.Last();
 	}
 
-	if (ToolTipText.IsEmpty())
-	{
-		ToolTipText = Variant.IsFunction() ? Variant.GetFunction()->GetToolTipText() :
-			Variant.IsProperty() ? Variant.GetProperty()->GetToolTipText() :
-			FText::GetEmpty();
-	}
-
-	SetToolTipText(ToolTipText);
+	SetToolTipText(Private::GetFieldToolTip(Variant));
 
 	Icon->RefreshBinding(Variant);
 	Label->SetText(Private::GetFieldDisplayName(Variant));
diff --git a/Engine/Plugins/Runtime/ModelViewViewModel/Source/ModelViewViewModelEditor/Private/Widgets/SMVVMFieldEntry.h b/Engine/Plugins/Runtime/ModelViewViewModel/Source/ModelViewViewModelEditor/Private/Widgets/SMVVMFieldEntry.h
--- a/Engine/Plugins/Runtime/ModelViewViewModel/Source/ModelViewViewModelEditor/Private/Widgets/SMVVMFieldEntry.h
+++ b/Engine/Plugins/Runtime/ModelViewViewModel/Source/ModelViewViewModelEditor/Private/Widgets/SMVVMFieldEntry.h
@@ -29,7 +29,6 @@ public:
 	{}
 		SLATE_STYLE_ARGUMENT(FTextBlockStyle, TextStyle)
 		SLATE_ARGUMENT(FMVVMBlueprintPropertyPath, Field)
-		SLATE_EVENT(FIsFieldValid, OnValidateField)
 	SLATE_END_ARGS()
 
 	void Construct(const FArguments& InArgs);
@@ -38,7 +37,6 @@ public:
 
 private:
 	FMVVMBlueprintPropertyPath Field;
-	FIsFieldValid OnValidateField;
 	TSharedPtr<SFieldIcon> Icon;
 	TSharedPtr<STextBlock> Label;
 };
diff --git a/Engine/Plugins/Runtime/ModelViewViewModel/Source/ModelViewViewModelEditor/Private/Widgets/SMVVMFieldSelector.cpp b/Engine/Plugins/Runtime/ModelViewViewModel/Source/ModelViewViewModelEditor/Private/Widgets/SMVVMFieldSelector.cpp
--- a/Engine/Plugins/Runtime/ModelViewViewModel/Source/ModelViewViewModelEditor/Private/Widgets/SMVVMFieldSelector.cpp
+++ b/Engine/Plugins/Runtime/ModelViewViewModel/Source/ModelViewViewModelEditor/Private/Widgets/SMVVMFieldSelector.cpp
@@ -7,7 +7,6 @@
 #include "Hierarchy/SReadOnlyHierarchyView.h"
 #include "Modules/ModuleManager.h"
 #include "MVVMEditorSubsystem.h"
-#include "SMVVMSelectViewModel.h"
 #include "SNegativeActionButton.h"
 #include "SPositiveActionButton.h"
 #include "SPrimaryButton.h"
@@ -46,9 +45,9 @@ FBindingSource GetSourceFromPath(const UWidgetBlueprint* WidgetBlueprint, const
 
 void SFieldSelector::Construct(const FArguments& InArgs, const UWidgetBlueprint* InWidgetBlueprint, bool bInViewModelProperty)
 {
-	OnValidateFieldDelegate = InArgs._OnValidateField;
 	OnSelectionChangedDelegate = InArgs._OnSelectionChanged;
 	check(OnSelectionChangedDelegate.IsBound());
+
 	WidgetBlueprint = InWidgetBlueprint;
 	check(WidgetBlueprint);
 
@@ -61,64 +60,76 @@ void SFieldSelector::Construct(const FArguments& InArgs, const UWidgetBlueprint*
 
 	bViewModelProperty = bInViewModelProperty;
 
+	AssignableTo = InArgs._AssignableTo;
+
 	TextStyle = InArgs._TextStyle;
 
+	TSharedPtr<SHorizontalBox> HBox;
+
 	ChildSlot
 	[
-		SNew(SBox)
-		.HAlign(HAlign_Left)
-		.VAlign(VAlign_Center)
+		SAssignNew(ComboButton, SComboButton)
+		.ComboButtonStyle(FMVVMEditorStyle::Get(), "FieldSelector.ComboButton")
+		.OnGetMenuContent(this, &SFieldSelector::OnGetMenuContent)
+		.ContentPadding(FMargin(4, 2))
+		.ButtonContent()
 		[
-			SAssignNew(ComboButton, SComboButton)
-			.ComboButtonStyle(FMVVMEditorStyle::Get(), "FieldSelector.ComboButton")
-			.OnGetMenuContent(this, &SFieldSelector::OnGetMenuContent)
-			.ContentPadding(FMargin(4, 2))
-			.ButtonContent()
+			SNew(SOverlay)
+			+ SOverlay::Slot()
 			[
-				SNew(SOverlay)
-				+ SOverlay::Slot()
-				[
-					SNew(SHorizontalBox)
-					.Visibility_Lambda([this]() { return CachedSelectedField.IsEmpty() ? EVisibility::Collapsed : EVisibility::Visible; })
-					+ SHorizontalBox::Slot()
-					.Padding(8, 0, 0, 0)
-					.HAlign(HAlign_Fill)
-					.FillWidth(1)
-					[
-						SAssignNew(SelectedSourceWidget, SSourceEntry)
-					]
-					+ SHorizontalBox::Slot()
-					.Padding(8, 0)
-					.VAlign(VAlign_Center)
-					.AutoWidth()
-					[
-						SNew(SImage)
-						.Image(FAppStyle::Get().GetBrush("Icons.ChevronRight"))
-					]
-					+ SHorizontalBox::Slot()
-					.Padding(0, 0, 8, 0)
-					.HAlign(HAlign_Fill)
-					.FillWidth(1)
-					[
-						SAssignNew(SelectedEntryWidget, SFieldEntry)
-						.TextStyle(TextStyle)
-						.OnValidateField(this, &SFieldSelector::ValidateField)
-					]
-				]
-				+ SOverlay::Slot()
+				SAssignNew(HBox, SHorizontalBox)
+				.Visibility_Lambda([this]() { return CachedSelectedField.IsEmpty() ? EVisibility::Collapsed : EVisibility::Visible; })
+					
+			]
+			+ SOverlay::Slot()
+			[
+				SNew(SBox)
+				.Padding(FMargin(8, 0, 8, 0))
 				[
-					SNew(SBox)
-					.Padding(FMargin(8, 0, 8, 0))
-					[
-						SNew(STextBlock)
-						.Visibility_Lambda([this]() { return CachedSelectedField.IsEmpty() ? EVisibility::Visible : EVisibility::Collapsed; })
-						.TextStyle(FAppStyle::Get(), "HintText")
-						.Text(LOCTEXT("None", "No field selected"))
-					]
+					SNew(STextBlock)
+					.Visibility_Lambda([this]() { return CachedSelectedField.IsEmpty() ? EVisibility::Visible : EVisibility::Collapsed; })
+					.TextStyle(FAppStyle::Get(), "HintText")
+					.Text(LOCTEXT("None", "No field selected"))
 				]
 			]
 		]
 	];
+	
+	if (InArgs._ShowSource)
+	{
+		HBox->AddSlot()
+			.Padding(8, 0, 0, 0)
+			.HAlign(HAlign_Left)
+			.VAlign(VAlign_Center)
+			.AutoWidth()
+			[
+				SAssignNew(SelectedSourceWidget, SSourceEntry)
+			];
+
+		HBox->AddSlot()
+			.Padding(8, 0)
+			.HAlign(HAlign_Left)
+			.VAlign(VAlign_Center)
+			.AutoWidth()
+			[
+				SNew(SImage)
+				.Image(FAppStyle::Get().GetBrush("Icons.ChevronRight"))
+			];					
+	}
+	else
+	{
+		FixedSource = InArgs._Source;
+	}
+
+	HBox->AddSlot()
+		.Padding(0, 0, 8, 0)
+		.HAlign(HAlign_Left)
+		.VAlign(VAlign_Center)
+		.AutoWidth()
+		[
+			SAssignNew(SelectedEntryWidget, SFieldEntry)
+			.TextStyle(TextStyle)
+		];
 
 	Refresh();
 }
@@ -138,22 +149,11 @@ void SFieldSelector::Refresh()
 	}
 }
 
-TValueOrError<bool, FString> SFieldSelector::ValidateField(FMVVMBlueprintPropertyPath Field) const
-{
-	if (OnValidateFieldDelegate.IsBound())
-	{
-		return OnValidateFieldDelegate.Execute(Field);
-	}
-
-	return MakeValue(true);
-}
-
 TSharedRef<SWidget> SFieldSelector::OnGenerateFieldWidget(FMVVMBlueprintPropertyPath Path) const
 {
 	return SNew(SFieldEntry)
 		.TextStyle(TextStyle)
-		.Field(Path)
-		.OnValidateField(this, &SFieldSelector::ValidateField);
+		.Field(Path);
 }
 
 bool SFieldSelector::IsSelectEnabled() const
@@ -215,21 +215,33 @@ void SFieldSelector::OnSearchTextChanged(const FText& NewText)
 
 void SFieldSelector::OnViewModelSelected(FBindingSource Source, ESelectInfo::Type)
 {
-	TArray<FBindingSource> Selection = ViewModelList->GetSelectedItems();
 	if (BindingList.IsValid())
 	{
+		TArray<FBindingSource> Selection = ViewModelList->GetSelectedItems();
+
 		BindingList->Clear();
 		BindingList->AddSources(Selection);
 	}
 }
 
+void SFieldSelector::OnWidgetSelected(FName WidgetName, ESelectInfo::Type)
+{
+	if (BindingList.IsValid())
+	{
+		FBindingSource Source = FBindingSource::CreateForWidget(WidgetBlueprint, WidgetName);
+
+		BindingList->Clear();
+		BindingList->AddSource(Source);
+	}
+}
+
 TSharedRef<ITableRow> SFieldSelector::GenerateRowForViewModel(FBindingSource ViewModel, const TSharedRef<STableViewBase>& OwnerTable) const
 {
 	return SNew(STableRow<FBindingSource>, OwnerTable)
 		[
 			SNew(SHorizontalBox)
 			+ SHorizontalBox::Slot()
-			.Padding(0, 0, 2, 0)
+			.Padding(4)
 			.HAlign(HAlign_Left)
 			.VAlign(VAlign_Center)
 			.AutoWidth()
@@ -258,61 +270,137 @@ TSharedRef<SWidget> SFieldSelector::OnGetMenuContent()
 	BindingList = SNew(SSourceBindingList, WidgetBlueprint)
 		.ShowSearchBox(false)
 		.OnDoubleClicked(this, &SFieldSelector::SetSelection)
-		.FieldVisibilityFlags(GetFieldVisibilityFlags());
+		.FieldVisibilityFlags(GetFieldVisibilityFlags())
+		.AssignableTo(AssignableTo);
+
+	TSharedRef<SVerticalBox> VBox = SNew(SVerticalBox)
+		+ SVerticalBox::Slot()
+		.HAlign(HAlign_Fill)
+		.VAlign(VAlign_Top)
+		.Padding(0, 4, 0, 4)
+		.AutoHeight()
+		[
+			SNew(SSearchBox)
+			.OnTextChanged(this, &SFieldSelector::OnSearchTextChanged)
+		];
 
-	TSharedPtr<SWidget> SourcePicker;
-	if (bViewModelProperty)
+	// single fixed source
+	if (FixedSource.IsSet())
 	{
-		FBindingSource SelectedSource;
+		BindingList->AddSource(FixedSource.GetValue());
 
-		UMVVMEditorSubsystem* EditorSubsystem = GEditor->GetEditorSubsystem<UMVVMEditorSubsystem>();
-		if (const UMVVMBlueprintView* View = EditorSubsystem->GetView(WidgetBlueprint))
+		VBox->AddSlot()
+			[
+				BindingList.ToSharedRef()
+			];
+	}
+	else
+	{ 
+		// show source picker
+		TSharedPtr<SWidget> SourcePicker;
+		if (bViewModelProperty)
 		{
-			const TArrayView<const FMVVMBlueprintViewModelContext> ViewModels = View->GetViewModels();
-			ViewModelSources.Reserve(ViewModels.Num());
+			FBindingSource SelectedSource;
 
-			for (const FMVVMBlueprintViewModelContext& ViewModel : ViewModels)
+			UMVVMEditorSubsystem* EditorSubsystem = GEditor->GetEditorSubsystem<UMVVMEditorSubsystem>();
+			if (const UMVVMBlueprintView* View = EditorSubsystem->GetView(WidgetBlueprint))
 			{
-				FBindingSource& Source = ViewModelSources.AddDefaulted_GetRef();
-				Source.Class = ViewModel.GetViewModelClass();
-				Source.Name = ViewModel.GetViewModelName();
-				Source.DisplayName = ViewModel.GetDisplayName();
-				Source.ViewModelId = ViewModel.GetViewModelId();
+				const TArrayView<const FMVVMBlueprintViewModelContext> ViewModels = View->GetViewModels();
+				ViewModelSources.Reserve(ViewModels.Num());
 
-				if (Source.ViewModelId == CachedSelectedField.GetViewModelId())
+				for (const FMVVMBlueprintViewModelContext& ViewModel : ViewModels)
 				{
-					SelectedSource = Source;
+					FBindingSource& Source = ViewModelSources.AddDefaulted_GetRef();
+					Source.Class = ViewModel.GetViewModelClass();
+					Source.Name = ViewModel.GetViewModelName();
+					Source.DisplayName = ViewModel.GetDisplayName();
+					Source.ViewModelId = ViewModel.GetViewModelId();
+
+					if (Source.ViewModelId == CachedSelectedField.GetViewModelId())
+					{
+						SelectedSource = Source;
+					}
 				}
 			}
-		}
 
-		ViewModelList = SNew(SListView<FBindingSource>)
-			.OnGenerateRow(this, &SFieldSelector::GenerateRowForViewModel)
-			.SelectionMode(ESelectionMode::Multi)
-			.ListItemsSource(&ViewModelSources)
-			.OnSelectionChanged(this, &SFieldSelector::OnViewModelSelected);
+			ViewModelList = SNew(SListView<FBindingSource>)
+				.OnGenerateRow(this, &SFieldSelector::GenerateRowForViewModel)
+				.SelectionMode(ESelectionMode::Multi)
+				.ListItemsSource(&ViewModelSources)
+				.OnSelectionChanged(this, &SFieldSelector::OnViewModelSelected);
 
-		ViewModelList->SetItemSelection(SelectedSource, true);
+			ViewModelList->SetItemSelection(SelectedSource, true);
 
-		SourcePicker = ViewModelList;
-	}
-	else
-	{
-		auto OnWidgetSelected = [this](FName WidgetName, ESelectInfo::Type)
+			SourcePicker = ViewModelList;
+		}
+		else
 		{
-			BindingList->Clear();
+			WidgetList = SNew(SReadOnlyHierarchyView, WidgetBlueprint)
+				.OnSelectionChanged(this, &SFieldSelector::OnWidgetSelected)
+				.SelectionMode(ESelectionMode::Multi)
+				.ShowSearch(false);
 
-			FBindingSource Source = FBindingSource::CreateForWidget(WidgetBlueprint, WidgetName);
-			if (Source.IsValid())
-			{
-				BindingList->AddSources(MakeArrayView(&Source, 1));
-			}
-		};
+			SourcePicker = WidgetList;
+		}
 
-		SourcePicker = SNew(SReadOnlyHierarchyView, WidgetBlueprint)
-			.OnSelectionChanged_Lambda(OnWidgetSelected)
-			.ShowSearch(false);
+		VBox->AddSlot()
+			.FillHeight(1.0f)
+			[
+				SNew(SSplitter)
+				.PhysicalSplitterHandleSize(4.0f)
+				+ SSplitter::Slot()
+				.Value(0.5f)
+				[
+					SourcePicker.ToSharedRef()
+				]
+				+ SSplitter::Slot()
+				.Value(0.5f)
+				[
+					BindingList.ToSharedRef()
+				]
+			];
 	}
+				
+	VBox->AddSlot()
+		.Padding(4, 4, 4, 0)
+		.HAlign(HAlign_Right)
+		.AutoHeight()
+		[
+			SNew(SHorizontalBox)
+			+ SHorizontalBox::Slot()
+			.AutoWidth()
+			[
+				SNew(SPrimaryButton)
+				.OnClicked(this, &SFieldSelector::OnSelectProperty)
+				.IsEnabled(this, &SFieldSelector::IsSelectEnabled)
+				.Text(LOCTEXT("Select", "Select"))
+			]
+			+ SHorizontalBox::Slot()
+			.Padding(FAppStyle::Get().GetMargin("StandardDialog.SlotPadding"))
+			.AutoWidth()
+			[
+				SNew(SButton)
+				.OnClicked(this, &SFieldSelector::OnClearBinding)
+				.IsEnabled(this, &SFieldSelector::IsClearEnabled)
+				.HAlign(HAlign_Center)
+				[
+					SNew(STextBlock)
+					.Text(LOCTEXT("Clear", "Clear"))
+				]
+			]
+			+ SHorizontalBox::Slot()
+			.Padding(FAppStyle::Get().GetMargin("StandardDialog.SlotPadding"))
+			.AutoWidth()
+			[
+				SNew(SButton)
+				.OnClicked(this, &SFieldSelector::OnCancel)
+				.HAlign(HAlign_Center)
+				[
+					SNew(STextBlock)
+					.Text(LOCTEXT("Cancel", "Cancel"))
+				]
+			]
+		];
 
 	TSharedRef<SWidget> MenuWidget = 
 		SNew(SBox)
@@ -323,67 +411,7 @@ TSharedRef<SWidget> SFieldSelector::OnGetMenuContent()
 			.BorderImage(FAppStyle::Get().GetBrush("ToolPanel.GroupBorder"))
 			.Padding(FMargin(8, 2, 8, 3))
 			[
-				SNew(SVerticalBox)
-				+ SVerticalBox::Slot()
-				.HAlign(HAlign_Fill)
-				.VAlign(VAlign_Top)
-				.AutoHeight()
-				[
-					SNew(SSearchBox)
-					.OnTextChanged(this, &SFieldSelector::OnSearchTextChanged)
-				]
-				+ SVerticalBox::Slot()
-				.FillHeight(1.0f)
-				[
-					SNew(SSplitter)
-					.PhysicalSplitterHandleSize(4.0f)
-					+ SSplitter::Slot()
-					.Value(0.5f)
-					[
-						SourcePicker.ToSharedRef()
-					]
-					+ SSplitter::Slot()
-					.Value(0.5f)
-					[
-						BindingList.ToSharedRef()
-					]
-				]
-				+ SVerticalBox::Slot()
-				.Padding(4, 4, 4, 0)
-				.AutoHeight()
-				[
-					SNew(SHorizontalBox)
-					+ SHorizontalBox::Slot()
-					[
-						SNew(SButton)
-						.OnClicked(this, &SFieldSelector::OnCancel)
-						.HAlign(HAlign_Center)
-						[
-							SNew(STextBlock)
-							.Text(LOCTEXT("Cancel", "Cancel"))
-						]
-					]
-					+ SHorizontalBox::Slot()
-					.Padding(4, 0, 0, 0)
-					[
-						SNew(SButton)
-						.OnClicked(this, &SFieldSelector::OnClearBinding)
-						.IsEnabled(this, &SFieldSelector::IsClearEnabled)
-						.HAlign(HAlign_Center)
-						[
-							SNew(STextBlock)
-							.Text(LOCTEXT("Clear", "Clear"))
-						]
-					]
-					+ SHorizontalBox::Slot()
-					.Padding(4, 0, 0, 0)
-					[
-						SNew(SPrimaryButton)
-						.OnClicked(this, &SFieldSelector::OnSelectProperty)
-						.IsEnabled(this, &SFieldSelector::IsSelectEnabled)
-						.Text(LOCTEXT("Select", "Select"))
-					]
-				]
+				VBox
 			]
 		];
 
diff --git a/Engine/Plugins/Runtime/ModelViewViewModel/Source/ModelViewViewModelEditor/Private/Widgets/SMVVMFieldSelector.h b/Engine/Plugins/Runtime/ModelViewViewModel/Source/ModelViewViewModelEditor/Private/Widgets/SMVVMFieldSelector.h
--- a/Engine/Plugins/Runtime/ModelViewViewModel/Source/ModelViewViewModelEditor/Private/Widgets/SMVVMFieldSelector.h
+++ b/Engine/Plugins/Runtime/ModelViewViewModel/Source/ModelViewViewModelEditor/Private/Widgets/SMVVMFieldSelector.h
@@ -29,7 +29,7 @@ class SSourceBindingList;
 class SFieldSelector : public SCompoundWidget
 {
 public:
-	DECLARE_DELEGATE_OneParam(FOnSelectionChanged, FMVVMBlueprintPropertyPath);
+	DECLARE_DELEGATE_OneParam(FOnSelectionChanged, const FMVVMBlueprintPropertyPath&);
 
 	SLATE_BEGIN_ARGS(SFieldSelector) :
 		_TextStyle(&FCoreStyle::Get().GetWidgetStyle<FTextBlockStyle>("NormalText"))
@@ -38,13 +38,26 @@ public:
 		SLATE_STYLE_ARGUMENT(FTextBlockStyle, TextStyle)
 		SLATE_ATTRIBUTE(FMVVMBlueprintPropertyPath, SelectedField)
 		SLATE_EVENT(FOnSelectionChanged, OnSelectionChanged)
+		
+		/**
+		 * Get the binding mode for this field. 
+		 * Required for validating the availability of functions, eg. SetFoo() can only be used as a setter.
+		 */
 		SLATE_ATTRIBUTE(EMVVMBindingMode, BindingMode)
 
+		/** 
+		 * Should we show the source as well as the field?
+		 * If this is set to false, then the Source attribute must be set.
+		 */
+		SLATE_ARGUMENT_DEFAULT(bool, ShowSource) = true;
+
 		/**
-		  * Would the given field be a valid entry for this combo box? 
-		  * The error string returned will be used as a tooltip.
-		  */
-		SLATE_EVENT(FIsFieldValid, OnValidateField)
+		 * Only show properties assignable to the given property.
+		 */
+		SLATE_ARGUMENT_DEFAULT(const FProperty*, AssignableTo) = nullptr;
+
+		/** The source to use. Only used if ShowSource is false. */
+		SLATE_ARGUMENT(FBindingSource, Source)
 	SLATE_END_ARGS()
 
 	void Construct(const FArguments& InArgs, const UWidgetBlueprint* InWidgetBlueprint, bool bInViewModelProperty);
@@ -70,6 +83,8 @@ private:
 	void OnViewModelSelected(FBindingSource ViewModel, ESelectInfo::Type);
 	TSharedRef<ITableRow> GenerateRowForViewModel(MVVM::FBindingSource ViewModel, const TSharedRef<STableViewBase>& OwnerTable) const;
 
+	void OnWidgetSelected(FName WidgetName, ESelectInfo::Type);
+
 	void OnSearchTextChanged(const FText& NewText);
 
 private:
@@ -82,7 +97,6 @@ private:
 	const FTextBlockStyle* TextStyle = nullptr;
 
 	FOnSelectionChanged OnSelectionChangedDelegate;
-	FIsFieldValid OnValidateFieldDelegate;
 
 	FMVVMBlueprintPropertyPath CachedSelectedField;
 
@@ -90,11 +104,16 @@ private:
 	TSharedPtr<SSearchBox> SearchBox;
 	TSharedPtr<SComboButton> ComboButton;
 
+	TOptional<FBindingSource> FixedSource;
+
 	const UWidgetBlueprint* WidgetBlueprint = nullptr;
 	bool bViewModelProperty = false;
 
 	TArray<FBindingSource> ViewModelSources;
 	TSharedPtr<SListView<FBindingSource>> ViewModelList;
+	TSharedPtr<SReadOnlyHierarchyView> WidgetList;
+
+	const FProperty* AssignableTo = nullptr;
 }; 
 
 } // namespace UE::MVVM
\ No newline at end of file
diff --git a/Engine/Plugins/Runtime/ModelViewViewModel/Source/ModelViewViewModelEditor/Private/Widgets/SMVVMSelectViewModel.cpp b/Engine/Plugins/Runtime/ModelViewViewModel/Source/ModelViewViewModelEditor/Private/Widgets/SMVVMSelectViewModel.cpp
--- a/Engine/Plugins/Runtime/ModelViewViewModel/Source/ModelViewViewModelEditor/Private/Widgets/SMVVMSelectViewModel.cpp
+++ b/Engine/Plugins/Runtime/ModelViewViewModel/Source/ModelViewViewModelEditor/Private/Widgets/SMVVMSelectViewModel.cpp
@@ -15,7 +15,6 @@
 #include "Styling/MVVMEditorStyle.h"
 #include "Widgets/Input/SButton.h"
 #include "Widgets/Layout/SSplitter.h"
-#include "Widgets/Layout/SUniformGridPanel.h"
 #include "Widgets/SMVVMViewModelBindingListWidget.h"
 
 #include "ClassViewerModule.h"
@@ -32,7 +31,7 @@ namespace Private
 
 	bool IsValidViewModel(const UClass* InClass)
 	{
-		if (InClass->IsChildOf(UWidget::StaticClass()))
+		if (InClass == nullptr || InClass->IsChildOf(UWidget::StaticClass()))
 		{
 			return false;
 		}
@@ -60,23 +59,6 @@ void SMVVMSelectViewModel::Construct(const FArguments& InArgs, const UWidgetBlue
 	OnCancel = InArgs._OnCancel;
 	OnViewModelCommitted = InArgs._OnViewModelCommitted;
 
-	TSharedRef<SWidget> ButtonsPanelContent = SNew(SUniformGridPanel)
-		.SlotPadding(FAppStyle::Get().GetMargin("StandardDialog.SlotPadding"))
-		+ SUniformGridPanel::Slot(0, 0)
-		[
-			SNew(SPrimaryButton)
-			.Text(LOCTEXT("ViewModelAddButtonText", "OK"))
-			.OnClicked(this, &SMVVMSelectViewModel::HandleAccepted)
-			.IsEnabled(this, &SMVVMSelectViewModel::HandleIsSelectionEnabled)
-		]
-		+ SUniformGridPanel::Slot(1, 0)
-		[
-			SNew(SButton)
-			.Text(LOCTEXT("ViewModelCancelButtonText", "Cancel"))
-			.HAlign(HAlign_Center)
-			.OnClicked(this, &SMVVMSelectViewModel::HandleCancel)
-		];
-
 	FClassViewerInitializationOptions ClassViewerOptions;
 	ClassViewerOptions.DisplayMode = EClassViewerDisplayMode::TreeView;
 	ClassViewerOptions.Mode = EClassViewerMode::ClassPicker;
@@ -88,41 +70,65 @@ void SMVVMSelectViewModel::Construct(const FArguments& InArgs, const UWidgetBlue
 	ChildSlot
 	[
 		SNew(SBorder)
-		.BorderImage(FAppStyle::GetBrush("ToolPanel.GroupBorder"))
+		.Padding(1)
 		[
-			SNew(SVerticalBox)
-			+ SVerticalBox::Slot()
-			.FillHeight(1.0f)
+			SNew(SBorder)
+			.BorderImage(FAppStyle::GetBrush("ToolPanel.GroupBorder"))
 			[
-				SNew(SSplitter)
-				.PhysicalSplitterHandleSize(2.0f)
-				+ SSplitter::Slot()
-				.Value(0.6f)
+				SNew(SVerticalBox)
+				+ SVerticalBox::Slot()
+				.FillHeight(1.0f)
 				[
-					SNew(SBorder)
-					.BorderImage(FStyleDefaults::GetNoBrush())
+					SNew(SSplitter)
+					.PhysicalSplitterHandleSize(2.0f)
+					+ SSplitter::Slot()
+					.Value(0.6f)
+					[
+						SNew(SBorder)
+						.BorderImage(FStyleDefaults::GetNoBrush())
+						[
+							ClassViewer.ToSharedRef()
+						]
+					]
+					+ SSplitter::Slot()
+					.Value(0.4f)
 					[
-						ClassViewer.ToSharedRef()
+						SNew(SBorder)
+						.BorderImage(FStyleDefaults::GetNoBrush())
+						.Padding(4, 30, 4, 30)
+						[
+							SAssignNew(BindingListWidget, SSourceBindingList, WidgetBlueprint)
+							.EnableSelection(false)
+						]
 					]
 				]
-				+ SSplitter::Slot()
-				.Value(0.4f)
+				+ SVerticalBox::Slot()
+				.HAlign(HAlign_Right)
+				.VAlign(VAlign_Bottom)
+				.AutoHeight()
+				.Padding(8)
 				[
-					SNew(SBorder)
-					.BorderImage(FStyleDefaults::GetNoBrush())
+					SNew(SHorizontalBox)
+					+ SHorizontalBox::Slot()
+					.Padding(FAppStyle::Get().GetMargin("StandardDialog.SlotPadding"))
+					.AutoWidth()
+					[
+						SNew(SPrimaryButton)
+						.Text(LOCTEXT("ViewModelSelectButtonText", "Select"))
+						.OnClicked(this, &SMVVMSelectViewModel::HandleSelected)
+						.IsEnabled(this, &SMVVMSelectViewModel::HandleIsSelectionEnabled)
+					]
+					+ SHorizontalBox::Slot()
+					.Padding(FAppStyle::Get().GetMargin("StandardDialog.SlotPadding"))
+					.AutoWidth()
 					[
-						SAssignNew(BindingListWidget, SSourceBindingList, WidgetBlueprint)
+						SNew(SButton)
+						.Text(LOCTEXT("ViewModelCancelButtonText", "Cancel"))
+						.HAlign(HAlign_Center)
+						.OnClicked(this, &SMVVMSelectViewModel::HandleCancel)
 					]
 				]
 			]
-			+ SVerticalBox::Slot()
-			.HAlign(HAlign_Right)
-			.VAlign(VAlign_Bottom)
-			.AutoHeight()
-			.Padding(8)
-			[
-				ButtonsPanelContent
-			]
 		]
 	];
 }
@@ -140,7 +146,7 @@ void SMVVMSelectViewModel::HandleClassPicked(UClass* ClassPicked)
 }
 
 
-FReply SMVVMSelectViewModel::HandleAccepted()
+FReply SMVVMSelectViewModel::HandleSelected()
 {
 	OnViewModelCommitted.ExecuteIfBound(SelectedClass.Get());
 	return FReply::Handled();
diff --git a/Engine/Plugins/Runtime/ModelViewViewModel/Source/ModelViewViewModelEditor/Private/Widgets/SMVVMSelectViewModel.h b/Engine/Plugins/Runtime/ModelViewViewModel/Source/ModelViewViewModelEditor/Private/Widgets/SMVVMSelectViewModel.h
--- a/Engine/Plugins/Runtime/ModelViewViewModel/Source/ModelViewViewModelEditor/Private/Widgets/SMVVMSelectViewModel.h
+++ b/Engine/Plugins/Runtime/ModelViewViewModel/Source/ModelViewViewModelEditor/Private/Widgets/SMVVMSelectViewModel.h
@@ -37,7 +37,7 @@ public:
 
 private:
 	void HandleClassPicked(UClass* ClassPicked);
-	FReply HandleAccepted();
+	FReply HandleSelected();
 	FReply HandleCancel();
 	bool HandleIsSelectionEnabled() const;
 
diff --git a/Engine/Plugins/Runtime/ModelViewViewModel/Source/ModelViewViewModelEditor/Private/Widgets/SMVVMSourceSelector.cpp b/Engine/Plugins/Runtime/ModelViewViewModel/Source/ModelViewViewModelEditor/Private/Widgets/SMVVMSourceSelector.cpp
--- a/Engine/Plugins/Runtime/ModelViewViewModel/Source/ModelViewViewModelEditor/Private/Widgets/SMVVMSourceSelector.cpp
+++ b/Engine/Plugins/Runtime/ModelViewViewModel/Source/ModelViewViewModelEditor/Private/Widgets/SMVVMSourceSelector.cpp
@@ -5,6 +5,7 @@
 #include "Algo/Transform.h"
 #include "Editor.h"
 #include "MVVMEditorSubsystem.h"
+#include "SPrimaryButton.h"
 #include "SSimpleButton.h"
 #include "Widgets/Layout/SBox.h"
 #include "Widgets/SMVVMFieldIcon.h"
@@ -18,9 +19,12 @@ namespace UE::MVVM
 void SSourceSelector::Construct(const FArguments& Args, const UWidgetBlueprint* InWidgetBlueprint)
 {
 	TextStyle = Args._TextStyle;
+
 	SelectedSourceAttribute = Args._SelectedSource;
 	check(SelectedSourceAttribute.IsSet());
 
+	InitialSource = SelectedSourceAttribute.Get();
+
 	WidgetBlueprint = InWidgetBlueprint;
 	check(InWidgetBlueprint);
 
@@ -85,31 +89,31 @@ void SSourceSelector::Construct(const FArguments& Args, const UWidgetBlueprint*
 	];
 
 	OnSelectionChangedDelegate = Args._OnSelectionChanged;
+	check(OnSelectionChangedDelegate.IsBound());
 }
 
 TSharedRef<SWidget> SSourceSelector::OnGetMenuContent()
 {
-	TSharedRef<SBox> TopLevelBox = SNew(SBox)
-		.MinDesiredWidth(200)
-		.MinDesiredHeight(400);
+	TSharedRef<SVerticalBox> VBox = SNew(SVerticalBox);
 
 	if (bViewModels)
 	{
 		ViewModelList = SNew(SListView<FBindingSource>)
-			.OnGenerateRow_Lambda([this](FBindingSource Source, TSharedRef<STableViewBase> OwnerTable)
-				{
-					return SNew(STableRow<FBindingSource>, OwnerTable)
-					[
-						SNew(SSourceEntry)
-						.Source(Source)
-						.TextStyle(TextStyle)
-					];
-				})
+			.OnGenerateRow(this, &SSourceSelector::OnGenerateRow)
 			.OnSelectionChanged(this, &SSourceSelector::OnViewModelSelectionChanged)
 			.SelectionMode(ESelectionMode::Single)
 			.ListItemsSource(&ViewModelSources);
 
-		TopLevelBox->SetContent(ViewModelList.ToSharedRef());
+		VBox->AddSlot()
+			.Padding(2)
+			[
+				SNew(SBorder)
+				.BorderImage(FAppStyle::Get().GetBrush("Brushes.Recessed"))
+				.Padding(4)
+				[
+					ViewModelList.ToSharedRef()
+				]
+			];
 	}
 	else
 	{
@@ -117,35 +121,95 @@ TSharedRef<SWidget> SSourceSelector::OnGetMenuContent()
 			.OnSelectionChanged(this, &SSourceSelector::OnWidgetSelectionChanged)
 			.ShowSearch(false);
 
-		TopLevelBox->SetContent(WidgetHierarchy.ToSharedRef());
+		VBox->AddSlot()
+			.Padding(2)
+			[
+				SNew(SBorder)
+				.BorderImage(FAppStyle::Get().GetBrush("Brushes.Recessed"))
+				.Padding(4)
+				[
+					WidgetHierarchy.ToSharedRef()
+				]
+			];
 	}
 
-	return TopLevelBox;
+	VBox->AddSlot()
+		.Padding(4,0,4,4)
+		.HAlign(HAlign_Right)
+		.VAlign(VAlign_Bottom)
+		.AutoHeight()
+		[
+			SNew(SHorizontalBox)
+			+ SHorizontalBox::Slot()
+			.Padding(FAppStyle::Get().GetMargin("StandardDialog.SlotPadding"))
+			.AutoWidth()
+			[
+				SNew(SPrimaryButton)
+				.Text(LOCTEXT("Select", "Select"))
+				.OnClicked(this, &SSourceSelector::HandleSelect)
+				.IsEnabled(this, &SSourceSelector::IsSelectEnabled)
+			]
+			+ SHorizontalBox::Slot()
+			.Padding(FAppStyle::Get().GetMargin("StandardDialog.SlotPadding"))
+			.AutoWidth()
+			[
+				SNew(SButton)
+				.Text(LOCTEXT("Cancel", "Cancel"))
+				.HAlign(HAlign_Center)
+				.OnClicked(this, &SSourceSelector::HandleCancel)
+			]
+		];
+
+	return SNew(SBox)
+		.MinDesiredWidth(300)
+		.Padding(1)
+		[
+			SNew(SBorder)
+			.BorderImage(FAppStyle::Get().GetBrush("ToolPanel.GroupBorder"))
+			.Padding(4)
+			[
+				VBox
+			]
+		];
 }
 
-void SSourceSelector::OnViewModelSelectionChanged(FBindingSource Selected, ESelectInfo::Type SelectionType)
+bool SSourceSelector::IsSelectEnabled() const
 {
-	SelectedSource = Selected;
+	return SelectedSource.IsValid();
+}
 
-	if (SelectedSourceWidget.IsValid())
-	{
-		SelectedSourceWidget->RefreshSource(Selected);
-	}
+FReply SSourceSelector::HandleCancel()
+{
+	OnViewModelSelectionChanged(InitialSource, ESelectInfo::Direct);
 
-	OnSelectionChangedDelegate.ExecuteIfBound(Selected);
+	MenuAnchor->SetIsOpen(false);
+	return FReply::Handled();
 }
 
-void SSourceSelector::OnWidgetSelectionChanged(FName SelectedName, ESelectInfo::Type SelectionType)
+FReply SSourceSelector::HandleSelect()
 {
-	FBindingSource Source = FBindingSource::CreateForWidget(WidgetBlueprint.Get(), SelectedName);
-
 	if (SelectedSourceWidget.IsValid())
 	{
-		SelectedSourceWidget->RefreshSource(Source);
+		SelectedSourceWidget->RefreshSource(SelectedSource);
 	}
 
-	OnSelectionChangedDelegate.ExecuteIfBound(Source);
+	OnSelectionChangedDelegate.ExecuteIfBound(SelectedSource);
+
+	MenuAnchor->SetIsOpen(false);
 
+	return FReply::Handled();
+}
+
+void SSourceSelector::OnViewModelSelectionChanged(FBindingSource Selected, ESelectInfo::Type SelectionType)
+{
+	SelectedSource = Selected;
+}
+
+void SSourceSelector::OnWidgetSelectionChanged(FName SelectedName, ESelectInfo::Type SelectionType)
+{
+	FBindingSource Selected = FBindingSource::CreateForWidget(WidgetBlueprint.Get(), SelectedName);
+
+	SelectedSource = Selected;
 }
 
 void SSourceSelector::Refresh()
@@ -193,7 +257,6 @@ void SSourceSelector::Refresh()
 			ViewModelList->ClearSelection();
 		}
 	}
-
 }
 
 EVisibility SSourceSelector::GetClearVisibility() const
@@ -211,6 +274,20 @@ FReply SSourceSelector::OnClearSource()
 	return FReply::Handled();
 }
 
+TSharedRef<ITableRow> SSourceSelector::OnGenerateRow(FBindingSource Source, const TSharedRef<STableViewBase>& OwnerTable)
+{
+	return SNew(STableRow<FBindingSource>, OwnerTable)
+	[
+		SNew(SBox)
+		.Padding(4, 3)
+		[
+			SNew(SSourceEntry)
+			.Source(Source)
+			.TextStyle(TextStyle)
+		]
+	];
+}
+
 } // namespace UE::MVVM
 
 #undef LOCTEXT_NAMESPACE
diff --git a/Engine/Plugins/Runtime/ModelViewViewModel/Source/ModelViewViewModelEditor/Private/Widgets/SMVVMSourceSelector.h b/Engine/Plugins/Runtime/ModelViewViewModel/Source/ModelViewViewModelEditor/Private/Widgets/SMVVMSourceSelector.h
--- a/Engine/Plugins/Runtime/ModelViewViewModel/Source/ModelViewViewModelEditor/Private/Widgets/SMVVMSourceSelector.h
+++ b/Engine/Plugins/Runtime/ModelViewViewModel/Source/ModelViewViewModelEditor/Private/Widgets/SMVVMSourceSelector.h
@@ -40,9 +40,15 @@ private:
 	void OnWidgetSelectionChanged(FName SelectedName, ESelectInfo::Type SelectionType);
 	TSharedRef<SWidget> OnGetMenuContent();
 
+	FReply HandleSelect();
+	FReply HandleCancel();
+	bool IsSelectEnabled() const;
+
 	EVisibility GetClearVisibility() const;
 	FReply OnClearSource();
 
+	TSharedRef<ITableRow> OnGenerateRow(FBindingSource Source, const TSharedRef<STableViewBase>& OwnerTable);
+
 private:
 	TWeakObjectPtr<const UWidgetBlueprint> WidgetBlueprint;
 	FSelectionChanged OnSelectionChangedDelegate;
@@ -52,6 +58,7 @@ private:
 	TAttribute<FBindingSource> SelectedSourceAttribute;
 
 	TArray<FBindingSource> ViewModelSources;
+	FBindingSource InitialSource;
 	FBindingSource SelectedSource;
 
 	TSharedPtr<SMenuAnchor> MenuAnchor;
diff --git a/Engine/Plugins/Runtime/ModelViewViewModel/Source/ModelViewViewModelEditor/Private/Widgets/SMVVMViewBindingListView.cpp b/Engine/Plugins/Runtime/ModelViewViewModel/Source/ModelViewViewModelEditor/Private/Widgets/SMVVMViewBindingListView.cpp
--- a/Engine/Plugins/Runtime/ModelViewViewModel/Source/ModelViewViewModelEditor/Private/Widgets/SMVVMViewBindingListView.cpp
+++ b/Engine/Plugins/Runtime/ModelViewViewModel/Source/ModelViewViewModelEditor/Private/Widgets/SMVVMViewBindingListView.cpp
@@ -4,14 +4,15 @@
 
 #include "Algo/Transform.h"
 #include "Bindings/MVVMBindingHelper.h"
+#include "BlueprintEditor.h"
 #include "Blueprint/WidgetTree.h"
 #include "Dialog/SCustomDialog.h"
 #include "Editor.h"
-#include "Styling/AppStyle.h"
 #include "Engine/Engine.h"
 #include "Features/IModularFeatures.h"
 #include "Framework/MultiBox/MultiBoxBuilder.h"
 #include "IPropertyAccessEditor.h"
+#include "K2Node_CallFunction.h"
 #include "MVVMBlueprintView.h"
 #include "MVVMBlueprintViewBinding.h"
 #include "MVVMEditorSubsystem.h"
@@ -20,14 +21,17 @@
 #include "SEnumCombo.h"
 #include "SPrimaryButton.h"
 #include "SSimpleButton.h" 
+#include "Styling/AppStyle.h"
 #include "Styling/MVVMEditorStyle.h"
 #include "Styling/StyleColors.h"
 #include "Widgets/SMVVMConversionPath.h"
 #include "Widgets/SMVVMFieldSelector.h"
+#include "Widgets/SMVVMFunctionParameter.h"
 #include "Widgets/SMVVMPropertyPath.h"
 #include "Widgets/SMVVMSourceSelector.h"
 #include "Widgets/SMVVMViewBindingPanel.h"
 #include "Widgets/SNullWidget.h"
+#include "Widgets/Layout/SSpacer.h"
 #include "Widgets/Images/SImage.h"
 #include "Widgets/Input/SCheckBox.h"
 #include "Widgets/Input/SComboBox.h"
@@ -44,9 +48,14 @@ namespace UE::MVVM
 // a wrapper around either a widget row or a binding row
 struct FBindingEntry
 {
-	FBindingEntry(FName InWidgetName) : WidgetName(InWidgetName) {}
-	FBindingEntry(int32 InIndex) : BindingIndex(InIndex) { check(InIndex != INDEX_NONE); }
-	
+	enum class ERowType
+	{
+		None,
+		Widget,
+		Binding,
+		Parameter
+	};
+
 	FMVVMBlueprintViewBinding* GetBinding(UMVVMBlueprintView* View) const
 	{
 		return View->GetBindingAt(BindingIndex);
@@ -57,14 +66,41 @@ struct FBindingEntry
 		return View->GetBindingAt(BindingIndex);
 	}
 
+	ERowType GetRowType() const
+	{
+		return RowType;
+	}
+
 	int32 GetBindingIndex() const
 	{
 		return BindingIndex;
 	}
 
-	FName GetWidgetName() const
+	void SetBindingIndex(int32 Index)
+	{
+		check(RowType == ERowType::None);
+		RowType = ERowType::Binding;
+		BindingIndex = Index;
+	}
+
+	void SetWidgetName(FName WidgetName)
+	{
+		check(RowType == ERowType::None);
+		RowType = ERowType::Widget;
+		Name = WidgetName;
+	}
+
+	void SetParameterName(int32 Index, FName ParameterName)
+	{
+		check(RowType == ERowType::None);
+		RowType = ERowType::Parameter;
+		BindingIndex = Index;
+		Name = ParameterName;
+	}
+
+	FName GetName() const
 	{
-		return WidgetName;
+		return Name;
 	}
 
 	TConstArrayView<TSharedPtr<FBindingEntry>> GetChildren() const
@@ -77,8 +113,16 @@ struct FBindingEntry
 		Children.Add(Child);
 	}
 
+	bool operator==(const FBindingEntry& Other) const
+	{
+		return RowType == Other.RowType &&
+			Name == Other.Name &&
+			BindingIndex == Other.BindingIndex;
+	}
+
 private:
-	FName WidgetName;
+	ERowType RowType = ERowType::None;
+	FName Name;
 	int32 BindingIndex = INDEX_NONE;
 	TArray<TSharedPtr<FBindingEntry>> Children;
 };
@@ -136,19 +180,20 @@ public:
 			[
 				SNew(SBox)
 				.HeightOverride(30)
+				.MinDesiredWidth(200)
+				.HAlign(HAlign_Left)
 				[
 					SNew(SHorizontalBox)
 					+ SHorizontalBox::Slot()
 					.Padding(2, 1)
-					.HAlign(HAlign_Left)
 					.VAlign(VAlign_Fill)
-					.AutoWidth()
 					[
 						SNew(SSourceSelector, InWidgetBlueprint)
 						.ShowClear(false)
 						.AutoRefresh(true)
 						.ViewModels(false)
 						.SelectedSource(this, &SWidgetRow::GetSelectedWidget)
+						.OnSelectionChanged(this, &SWidgetRow::SetSelectedWidget)
 					]
 					+ SHorizontalBox::Slot()
 					.HAlign(HAlign_Left)
@@ -157,7 +202,7 @@ public:
 					[
 						SNew(SSimpleButton)
 						.Icon(FAppStyle::Get().GetBrush("Icons.Plus"))
-						.IsEnabled_Lambda([this]() { return !Entry->GetWidgetName().IsNone(); })
+						.IsEnabled_Lambda([this]() { return !Entry->GetName().IsNone(); })
 						.OnClicked(this, &SWidgetRow::AddBinding)
 
 					]
@@ -171,7 +216,27 @@ private:
 
 	FBindingSource GetSelectedWidget() const
 	{
-		return FBindingSource::CreateForWidget(WidgetBlueprintWeak.Get(), Entry->GetWidgetName());
+		return FBindingSource::CreateForWidget(WidgetBlueprintWeak.Get(), Entry->GetName());
+	}
+
+	void SetSelectedWidget(FBindingSource Source)
+	{
+		if (UWidgetBlueprint* WidgetBlueprint = WidgetBlueprintWeak.Get())
+		{
+			UMVVMEditorSubsystem* EditorSubsystem = GEditor->GetEditorSubsystem<UMVVMEditorSubsystem>();
+			UMVVMBlueprintView* View = EditorSubsystem->GetView(WidgetBlueprint);
+
+			for (const TSharedPtr<FBindingEntry>& ChildEntry : Entry->GetChildren())
+			{
+				if (FMVVMBlueprintViewBinding* Binding = ChildEntry->GetBinding(View))
+				{
+					FMVVMBlueprintPropertyPath CurrentPath = Binding->WidgetPath;
+					CurrentPath.SetWidgetName(Source.Name);
+
+					EditorSubsystem->SetWidgetPropertyForBinding(WidgetBlueprint, *Binding, CurrentPath);
+				}
+			}
+		}
 	}
 
 	FReply AddBinding() const
@@ -181,7 +246,7 @@ private:
 			UMVVMEditorSubsystem* EditorSubsystem = GEditor->GetEditorSubsystem<UMVVMEditorSubsystem>();
 			FMVVMBlueprintViewBinding& Binding = EditorSubsystem->AddBinding(WidgetBlueprint);
 			FMVVMBlueprintPropertyPath Path;
-			Path.SetWidgetName(Entry->GetWidgetName());
+			Path.SetWidgetName(Entry->GetName());
 			EditorSubsystem->SetWidgetPropertyForBinding(WidgetBlueprint, Binding, Path);
 		}
 
@@ -199,15 +264,15 @@ public:
 	SLATE_BEGIN_ARGS(SBindingRow) {}
 	SLATE_END_ARGS()
 
-	void Construct(const FArguments& Args, const TSharedRef<STableViewBase>& OwnerTableView, const TSharedPtr<FBindingEntry>& InEntry, UMVVMBlueprintView* InBlueprintView, UWidgetBlueprint* InWidgetBlueprint)
+	void Construct(const FArguments& Args, const TSharedRef<STableViewBase>& OwnerTableView, const TSharedPtr<FBindingEntry>& InEntry, UWidgetBlueprint* InWidgetBlueprint)
 	{
 		Entry = InEntry;
 		WidgetBlueprintWeak = InWidgetBlueprint;
-		BlueprintView = InBlueprintView;
 
 		OnBlueprintChangedHandle = InWidgetBlueprint->OnChanged().AddSP(this, &SBindingRow::HandleBlueprintChanged);
-
-		FMVVMBlueprintViewBinding* ViewModelBinding = GetThisViewBinding();
+		
+		FMVVMBlueprintViewBinding* ViewBinding = GetThisViewBinding();
+		FBindingSource WidgetSource = FBindingSource::CreateForWidget(InWidgetBlueprint, ViewBinding->WidgetPath.GetWidgetName());
 
 		STableRow<TSharedPtr<FBindingEntry>>::Construct(
 			STableRow<TSharedPtr<FBindingEntry>>::FArguments()
@@ -242,40 +307,18 @@ public:
 				+ SHorizontalBox::Slot()
 				.Padding(4, 0)
 				.VAlign(VAlign_Center)
-				.HAlign(HAlign_Fill)
+				.HAlign(HAlign_Left)
 				.AutoWidth()
 				[
 					SNew(SBox)
-					.HeightOverride(28)
+					.MinDesiredWidth(200)
 					[
-						SAssignNew(WidgetMenuAnchor, SMenuAnchor)
-						.OnGetMenuContent(this, &SBindingRow::OnGetWidgetMenuContent)
-						[
-							SNew(SButton)
-							.ButtonStyle(&FAppStyle::Get().GetWidgetStyle<FComboButtonStyle>("ComboButton").ButtonStyle)
-							.OnClicked(this, &SBindingRow::OnWidgetMenuAnchorClicked)
-							[
-								SNew(SHorizontalBox)
-								+ SHorizontalBox::Slot()
-								.HAlign(HAlign_Left)
-								.VAlign(VAlign_Fill)
-								.AutoWidth()
-								[
-									SAssignNew(SelectedWidget, SFieldEntry)
-									.Field(GetSelectedWidgetProperty())
-								]
-								+ SHorizontalBox::Slot()
-								.Padding(4, 0, 0, 0)
-								.HAlign(HAlign_Right)
-								.VAlign(VAlign_Center)
-								.AutoWidth()
-								[
-									SNew(SImage)
-									.Image(FAppStyle::Get().GetBrush("Icons.ChevronDown"))
-									.ColorAndOpacity(FSlateColor::UseForeground())
-								]
-							]
-						]
+						SNew(SFieldSelector, InWidgetBlueprint, false)
+						.SelectedField(this, &SBindingRow::GetSelectedWidgetProperty)
+						.OnSelectionChanged(this, &SBindingRow::OnWidgetPropertySelected)
+						.BindingMode(this, &SBindingRow::GetCurrentBindingMode)
+						.ShowSource(false)
+						.Source(WidgetSource)
 					]
 				]
 
@@ -287,7 +330,7 @@ public:
 				[
 					SNew(SComboBox<FName>)
 					.OptionsSource(Private::GetBindingModeNames())
-					.InitiallySelectedItem(StaticEnum<EMVVMBindingMode>()->GetNameByValue((int64) ViewModelBinding->BindingType))
+					.InitiallySelectedItem(StaticEnum<EMVVMBindingMode>()->GetNameByValue((int64) ViewBinding->BindingType))
 					.OnSelectionChanged(this, &SBindingRow::OnModeSelectionChanged)
 					.OnGenerateWidget(this, &SBindingRow::GenerateModeWidget)
 					.Content()
@@ -305,21 +348,30 @@ public:
 				]
 				
 				+ SHorizontalBox::Slot()
-				.Padding(2, 0)
+				.Padding(4, 0, 2, 0)
 				.VAlign(VAlign_Center)
-				.HAlign(HAlign_Fill) 
+				.HAlign(HAlign_Left)
 				.AutoWidth()
 				[
-					SAssignNew(ViewModelFieldSelector, SFieldSelector, InWidgetBlueprint, true)
-					.SelectedField(this, &SBindingRow::GetSelectedViewModelProperty)
-					.BindingMode(this, &SBindingRow::GetCurrentBindingMode)
-					.OnSelectionChanged(this, &SBindingRow::OnViewModelPropertySelectionChanged)
+					SNew(SBox)
+					.MinDesiredWidth(200)
+					[
+						SAssignNew(ViewModelFieldSelector, SFieldSelector, InWidgetBlueprint, true)
+						.SelectedField(this, &SBindingRow::GetSelectedViewModelProperty)
+						.BindingMode(this, &SBindingRow::GetCurrentBindingMode)
+						.OnSelectionChanged(this, &SBindingRow::OnViewModelPropertySelected)
+					]
+				]
+
+				+ SHorizontalBox::Slot()
+				[
+					SNew(SSpacer)
 				]
 
 				+ SHorizontalBox::Slot()
 				.Padding(2, 1)
 				.VAlign(VAlign_Fill)
-				.HAlign(HAlign_Left)
+				.HAlign(HAlign_Right)
 				.AutoWidth()
 				[
 					SNew(SEnumComboBox, StaticEnum<EMVVMViewBindingUpdateMode>())
@@ -331,7 +383,7 @@ public:
 				+ SHorizontalBox::Slot()
 				.Padding(2, 0)
 				.VAlign(VAlign_Center)
-				.HAlign(HAlign_Left)
+				.HAlign(HAlign_Right)
 				.AutoWidth()
 				[
 					SNew(SMVVMConversionPath, InWidgetBlueprint, false)
@@ -342,34 +394,13 @@ public:
 				+ SHorizontalBox::Slot()
 				.Padding(2, 0)
 				.VAlign(VAlign_Center)
-				.HAlign(HAlign_Left)
+				.HAlign(HAlign_Right)
 				.AutoWidth()
 				[
 					SNew(SMVVMConversionPath, InWidgetBlueprint, true)
 					.Bindings(this, &SBindingRow::GetThisViewBindingAsArray)
 					.OnFunctionChanged(this, &SBindingRow::OnConversionFunctionChanged, true)
 				]
-
-				+ SHorizontalBox::Slot() 
-				.Padding(2, 0)
-				.VAlign(VAlign_Center)
-				.HAlign(HAlign_Right)
-				.AutoWidth()
-				[
-					SAssignNew(ContextMenuOptionHelper, SButton)
-					.ToolTipText(LOCTEXT("DropDownOptionsToolTip", "Context Menu for Binding"))
-					.ButtonStyle(&FAppStyle::Get().GetWidgetStyle<FButtonStyle>("SimpleButton"))
-					.OnClicked(this, &SBindingRow::HandleDropDownOptionsPressed)
-					[
-						SNew(SBox)
-						.Padding(FMargin(3, 0))
-						[
-							SNew(SImage)
-							.Image(&FAppStyle::Get().GetWidgetStyle<FComboButtonStyle>("SegmentedCombo.Right").DownArrowImage)
-							.ColorAndOpacity(FSlateColor::UseForeground())
-						]
-					]
-				]
 			],
 			OwnerTableView
 		);
@@ -385,12 +416,12 @@ public:
 
 	FMVVMBlueprintViewBinding* GetThisViewBinding() const
 	{
-		if (UMVVMBlueprintView* BlueprintViewPtr = BlueprintView.Get())
+		UMVVMEditorSubsystem* EditorSubsystem = GEditor->GetEditorSubsystem<UMVVMEditorSubsystem>();
+		if (UMVVMBlueprintView* BlueprintViewPtr = EditorSubsystem->GetView(WidgetBlueprintWeak.Get()))
 		{
 			FMVVMBlueprintViewBinding* ViewBinding = Entry->GetBinding(BlueprintViewPtr);
 			return ViewBinding;
 		}
-
 		return nullptr;
 	}
 
@@ -434,7 +465,6 @@ private:
 		static const FText NewLineText = FText::FromString(TEXT("\n"));
 		FText HintText = LOCTEXT("ErrorButtonText", "Errors: (Click to show in a separate window)");
 		FText ErrorsText = FText::Join(NewLineText, GetThisViewBinding()->Errors);
-
 		return FText::Join(NewLineText, HintText, ErrorsText);
 	}
 
@@ -473,88 +503,6 @@ private:
 		return ViewModelBinding->BindingType;
 	}
 
-	FReply OnWidgetMenuAnchorClicked()
-	{
-		WidgetMenuAnchor->SetIsOpen(!WidgetMenuAnchor->IsOpen());
-		return FReply::Handled();
-	}
-
-	TSharedRef<SWidget> OnGetWidgetMenuContent()
-	{
-		FMVVMBlueprintViewBinding* ViewBinding = GetThisViewBinding();
-
-		UWidgetBlueprint* WidgetBlueprint = WidgetBlueprintWeak.Get();
-
-		EFieldVisibility Flags = EFieldVisibility::None;
-		if (UE::MVVM::IsForwardBinding(ViewBinding->BindingType))
-		{
-			Flags |= EFieldVisibility::Writable;
-		}
-
-		if (UE::MVVM::IsBackwardBinding(ViewBinding->BindingType))
-		{
-			Flags |= EFieldVisibility::Readable;
-
-			if (!UE::MVVM::IsOneTimeBinding(ViewBinding->BindingType))
-			{
-				Flags |= EFieldVisibility::Notify;
-			}
-		}
-
-		FBindingSource Source = FBindingSource::CreateForWidget(WidgetBlueprint, ViewBinding->WidgetPath.GetWidgetName());
-
-		WidgetFieldIterator = MakeShared<UE::MVVM::FFieldIterator_Bindable>(WidgetBlueprint, Flags);
-		
-		return SNew(SBorder)
-			.BorderImage(FAppStyle::Get().GetBrush("ToolPanel.GroupBorder"))
-			.Padding(FMargin(8, 8, 8, 4))
-			[
-				SNew(SVerticalBox)
-				+ SVerticalBox::Slot()
-				[
-					SAssignNew(WidgetBindingList, SSourceBindingList, WidgetBlueprint)
-					.InitialSource(Source)
-					.ShowSearchBox(false)
-					.OnDoubleClicked(this, &SBindingRow::OnWidgetPropertySelectionChanged)
-					.FieldVisibilityFlags(Flags)
-				]
-				+ SVerticalBox::Slot()
-				.Padding(4, 4, 4, 0)
-				.AutoHeight()
-				[
-					SNew(SHorizontalBox)
-					+ SHorizontalBox::Slot()
-					[
-						SNew(SButton)
-						.OnClicked(this, &SBindingRow::OnWidgetButtonClicked, false)
-						.HAlign(HAlign_Center)
-						[
-							SNew(STextBlock)
-							.Text(LOCTEXT("Cancel", "Cancel"))
-						]
-					]
-					+ SHorizontalBox::Slot()
-					.Padding(4, 0, 0, 0)
-					[
-						SNew(SButton)
-						.OnClicked(this, &SBindingRow::OnWidgetClearButtonClicked)
-						.HAlign(HAlign_Center)
-						[
-							SNew(STextBlock)
-							.Text(LOCTEXT("Clear", "Clear"))
-						]
-					]
-					+ SHorizontalBox::Slot()
-					.Padding(4, 0, 0, 0)
-					[
-						SNew(SPrimaryButton)
-						.OnClicked(this, &SBindingRow::OnWidgetButtonClicked, true)
-						.Text(LOCTEXT("Select", "Select"))
-					]
-				]
-			];
-	}
-
 	TSharedRef<ITableRow> OnGenerateErrorRow(TSharedPtr<FText> Text, const TSharedRef<STableViewBase>& TableView) const
 	{
 		return SNew(STableRow<TSharedPtr<FText>>, TableView)
@@ -595,7 +543,6 @@ private:
 		{
 			return ViewBinding->ViewModelPath;
 		}
-
 		return FMVVMBlueprintPropertyPath();
 	}
 
@@ -605,11 +552,10 @@ private:
 		{
 			return ViewBinding->WidgetPath;
 		}
-
 		return FMVVMBlueprintPropertyPath();
 	}
 
-	void OnViewModelPropertySelectionChanged(FMVVMBlueprintPropertyPath SelectedField)
+	void OnViewModelPropertySelected(const FMVVMBlueprintPropertyPath& SelectedField)
 	{
 		if (FMVVMBlueprintViewBinding* ViewBinding = GetThisViewBinding())
 		{
@@ -621,7 +567,7 @@ private:
 		}
 	}
 
-	void OnWidgetPropertySelectionChanged(const FMVVMBlueprintPropertyPath& SelectedField)
+	void OnWidgetPropertySelected(const FMVVMBlueprintPropertyPath& SelectedField)
 	{
 		if (FMVVMBlueprintViewBinding* ViewBinding = GetThisViewBinding())
 		{
@@ -643,25 +589,6 @@ private:
 		}
 	}
 
-	FReply OnWidgetButtonClicked(bool bSelect)
-	{
-		if (bSelect)
-		{
-			OnWidgetPropertySelectionChanged(WidgetBindingList->GetSelectedProperty());
-		}
-
-		WidgetMenuAnchor->SetIsOpen(false);
-		return FReply::Handled();
-	}
-
-	FReply OnWidgetClearButtonClicked()
-	{
-		FMVVMBlueprintPropertyPath Path;
-		Path.SetWidgetName(Entry->GetWidgetName());
-		OnWidgetPropertySelectionChanged(Path);
-		return FReply::Handled();
-	}
-
 	void OnUpdateModeSelectionChanged(int32 Value, ESelectInfo::Type)
 	{
 		if (FMVVMBlueprintViewBinding* ViewBinding = GetThisViewBinding())
@@ -673,7 +600,11 @@ private:
 
 	int32 GetUpdateModeValue() const
 	{
-		return (int32) GetThisViewBinding()->UpdateMode;
+		if (FMVVMBlueprintViewBinding* ViewBinding = GetThisViewBinding())
+		{
+			return (int32) ViewBinding->UpdateMode;
+		}
+		return (int32) EMVVMViewBindingUpdateMode::Immediate;
 	}
 
 	void OnIsBindingEnableChanged(ECheckBoxState NewState)
@@ -711,9 +642,9 @@ private:
 		case EMVVMBindingMode::OneTimeToDestination:
 			return FMVVMEditorStyle::Get().GetBrush("BindingMode.OneTime");
 		case EMVVMBindingMode::OneWayToDestination:
-			return FMVVMEditorStyle::Get().GetBrush("BindingMode.OneWay");
-		case EMVVMBindingMode::OneWayToSource:
 			return FMVVMEditorStyle::Get().GetBrush("BindingMode.OneWayToSource");
+		case EMVVMBindingMode::OneWayToSource:
+			return FMVVMEditorStyle::Get().GetBrush("BindingMode.OneWay");
 		case EMVVMBindingMode::OneTimeToSource:
 			return nullptr;
 		case EMVVMBindingMode::TwoWay:
@@ -725,7 +656,11 @@ private:
 
 	const FSlateBrush* GetCurrentModeBrush() const
 	{
-		return GetModeBrush(GetThisViewBinding()->BindingType);
+		if (FMVVMBlueprintViewBinding* ViewBinding = GetThisViewBinding())
+		{
+			return GetModeBrush(ViewBinding->BindingType);
+		}
+		return nullptr;
 	}
 
 	const FText& GetModeLabel(EMVVMBindingMode BindingMode) const
@@ -799,28 +734,6 @@ private:
 		}
 	}
 
-	FReply HandleDropDownOptionsPressed() 
-	{
-		if (TSharedPtr<ITypedTableView<TSharedPtr<FBindingEntry>>> ListView = OwnerTablePtr.Pin())
-		{
-			if (TSharedPtr<SBindingsList> ParentList = StaticCastSharedPtr<SBindingsList>(ListView->AsWidget()->GetParentWidget()))
-			{
-				// Get the context menu content. If invalid, don't open a menu.
-				ListView->Private_SetItemSelection(Entry, true);
-				TSharedPtr<SWidget> MenuContent = ParentList->OnSourceConstructContextMenu();
-
-				if (MenuContent.IsValid())
-				{
-					const FVector2D& SummonLocation = ContextMenuOptionHelper->GetCachedGeometry().GetRenderBoundingRect().GetBottomLeft();
-					FWidgetPath WidgetPath;
-					FSlateApplication::Get().PushMenu(ParentList->AsShared(), WidgetPath, MenuContent.ToSharedRef(), SummonLocation, FPopupTransitionEffect(FPopupTransitionEffect::ContextMenu));
-				}
-			}
-		}
-
-		return FReply::Handled();
-	}
-
 	void HandleBlueprintChanged(UBlueprint* Blueprint)
 	{
 		ViewModelFieldSelector->Refresh();
@@ -829,7 +742,6 @@ private:
 private:
 	TSharedPtr<FBindingEntry> Entry;
 	TWeakObjectPtr<UWidgetBlueprint> WidgetBlueprintWeak;
-	TWeakObjectPtr<UMVVMBlueprintView> BlueprintView;
 	TSharedPtr<SFieldSelector> ViewModelFieldSelector;
 	TSharedPtr<SWidget> ContextMenuOptionHelper;
 	TSharedPtr<SCustomDialog> ErrorDialog;
@@ -841,6 +753,103 @@ private:
 	TSharedPtr<SFieldEntry> SelectedWidget;
 };
 
+class SFunctionParameterRow : public STableRow<TSharedPtr<FBindingEntry>>
+{
+	SLATE_BEGIN_ARGS(SFunctionParameterRow) {}
+	SLATE_END_ARGS()
+
+	void Construct(const FArguments& Args, const TSharedRef<STableViewBase>& OwnerTableView, const TSharedPtr<FBindingEntry>& InEntry, UWidgetBlueprint* InWidgetBlueprint)
+	{
+		Entry = InEntry;
+		check(Entry->GetRowType() == FBindingEntry::ERowType::Parameter);
+
+		WidgetBlueprint = InWidgetBlueprint;
+
+		UMVVMEditorSubsystem* EditorSubsystem = GEditor->GetEditorSubsystem<UMVVMEditorSubsystem>();
+		UMVVMBlueprintView* View = EditorSubsystem->GetView(InWidgetBlueprint);
+		FMVVMBlueprintViewBinding* Binding = Entry->GetBinding(View);
+		UEdGraphPin* Pin = EditorSubsystem->FindConversionFunctionArgumentPin(InWidgetBlueprint, *Binding, Entry->GetName(), UE::MVVM::IsForwardBinding(Binding->BindingType));
+
+		FSlateColor PrimaryColor, SecondaryColor;
+		const FSlateBrush* SecondaryBrush = nullptr;
+		const FSlateBrush* PrimaryBrush = FBlueprintEditor::GetVarIconAndColorFromPinType(Pin->PinType, PrimaryColor, SecondaryBrush, SecondaryColor);
+
+		STableRow<TSharedPtr<FBindingEntry>>::Construct(
+			STableRow<TSharedPtr<FBindingEntry>>::FArguments()
+		[
+			SNew(SBox)
+			.HeightOverride(30)
+			[
+				SNew(SHorizontalBox)
+				+ SHorizontalBox::Slot()
+				.HAlign(HAlign_Left)
+				.VAlign(VAlign_Center)
+				.Padding(0, 0, 4, 0)
+				.AutoWidth()
+				[
+					SNew(SBox)
+					.MinDesiredWidth(100)
+					[
+						SNew(SHorizontalBox)
+						+ SHorizontalBox::Slot()
+						.VAlign(VAlign_Center)
+						.Padding(0, 0, 8, 0)
+						.AutoWidth()
+						[
+							SNew(SImage)
+							.Image(PrimaryBrush)
+							.ColorAndOpacity(PrimaryColor)
+							.DesiredSizeOverride(FVector2D(16, 16))
+						]
+						+ SHorizontalBox::Slot()
+						[
+							SNew(STextBlock)
+							.Text(Pin->GetDisplayName())
+						]
+					]
+				]
+				+ SHorizontalBox::Slot()
+				.HAlign(HAlign_Left)
+				.VAlign(VAlign_Center)
+				.AutoWidth()
+				[
+					SNew(SBox)
+					.MinDesiredWidth(200)
+					[
+						SNew(SFunctionParameter)
+						.OnGetBindingMode(this, &SFunctionParameterRow::OnGetBindingMode)
+						.WidgetBlueprint(InWidgetBlueprint)
+						.Binding(Binding)
+						.ParameterName(Entry->GetName())
+						.SourceToDestination(UE::MVVM::IsForwardBinding(Binding->BindingType))
+					]
+				]
+			]
+		], OwnerTableView);
+	}
+
+private:
+
+	EMVVMBindingMode OnGetBindingMode() const
+	{
+		if (UWidgetBlueprint* WidgetBlueprintPtr = WidgetBlueprint.Get())
+		{
+			UMVVMEditorSubsystem* EditorSubsystem = GEditor->GetEditorSubsystem<UMVVMEditorSubsystem>();
+			UMVVMBlueprintView* View = EditorSubsystem->GetView(WidgetBlueprintPtr);
+
+			if (FMVVMBlueprintViewBinding* Binding = Entry->GetBinding(View))
+			{
+				return Binding->BindingType;
+			}
+		}
+		return EMVVMBindingMode::OneWayToDestination;
+	}
+
+private:
+	TWeakObjectPtr<UWidgetBlueprint> WidgetBlueprint;
+	TSharedPtr<FBindingEntry> Entry;
+};
+
 void SBindingsList::Construct(const FArguments& InArgs, TSharedPtr<SBindingsPanel> Owner, UMVVMWidgetBlueprintExtension_View* InMVVMExtension)
 {
 	BindingPanel = Owner;
@@ -889,30 +898,32 @@ void SBindingsList::Refresh()
 	UMVVMBlueprintView* BlueprintView = MVVMExtensionPtr ? MVVMExtensionPtr->GetBlueprintView() : nullptr;
 
 	// store the current binding index
-	TArray<FMVVMBlueprintViewBinding*, TInlineAllocator<8>> SelectedBindings;
-
+	TArray<TSharedPtr<FBindingEntry>> PreviousSelectedEntries;
 	if (TreeView.IsValid() && BlueprintView)
 	{
-		TArray<TSharedPtr<FBindingEntry>> SelectedItems = TreeView->GetSelectedItems();
-		for (const TSharedPtr<FBindingEntry>& Entry : SelectedItems)
-		{
-			FMVVMBlueprintViewBinding* Binding = Entry->GetBinding(BlueprintView);
-			if (Binding != nullptr)
-			{
-				SelectedBindings.Add(Binding);
-			}
-		}
+		PreviousSelectedEntries = TreeView->GetSelectedItems();
 	}
 
 	RootWidgets.Reset();
 
-	TArray<TSharedPtr<FBindingEntry>> SelectedEntries;
+	TArray<TSharedPtr<FBindingEntry>> NewSelectedEntries;
 
 	// generate our entries
 	// for each widget with bindings, create an entry at the root level
 	// then add all bindings that reference that widget as its children
 	if (BlueprintView)
 	{
+		auto ReselectIfRequired = [&PreviousSelectedEntries, &NewSelectedEntries](const TSharedPtr<FBindingEntry>& ToFind)
+		{
+			if (PreviousSelectedEntries.ContainsByPredicate([ToFind](const TSharedPtr<FBindingEntry>& Entry) -> bool
+				{
+					return Entry.Get() == ToFind.Get();
+				}))
+			{
+				NewSelectedEntries.Add(ToFind);
+			}
+		};
+
 		TArrayView<FMVVMBlueprintViewBinding> Bindings = BlueprintView->GetBindings();
 		for (int32 BindingIndex = 0; BindingIndex < Bindings.Num(); ++BindingIndex)
 		{
@@ -923,7 +934,7 @@ void SBindingsList::Refresh()
 			TSharedPtr<FBindingEntry> ExistingWidget;
 			for (TSharedPtr<FBindingEntry> Widget : RootWidgets)
 			{
-				if (WidgetName == Widget->GetWidgetName())
+				if (WidgetName == Widget->GetName())
 				{
 					ExistingWidget = Widget;
 					break;
@@ -932,15 +943,45 @@ void SBindingsList::Refresh()
 
 			if (!ExistingWidget.IsValid())
 			{
-				ExistingWidget = RootWidgets.Add_GetRef(MakeShared<FBindingEntry>(WidgetName));
+				TSharedRef<FBindingEntry> NewWidget = MakeShared<FBindingEntry>();
+				NewWidget->SetWidgetName(WidgetName);
+				ExistingWidget = NewWidget;
+
+				RootWidgets.Add(NewWidget);
+
+				ReselectIfRequired(NewWidget);
 			}
 
-			TSharedRef<FBindingEntry> NewBindingEntry = MakeShared<FBindingEntry>(BindingIndex);
+			TSharedRef<FBindingEntry> NewBindingEntry = MakeShared<FBindingEntry>();
+			NewBindingEntry->SetBindingIndex(BindingIndex);
 			ExistingWidget->AddChild(NewBindingEntry);
+			ReselectIfRequired(NewBindingEntry);
 
-			if (SelectedBindings.Contains(&Binding))
+			// create entries for conversion function parameters
+			UMVVMEditorSubsystem* EditorSubsystem = GEditor->GetEditorSubsystem<UMVVMEditorSubsystem>();
+			UEdGraph* Graph = EditorSubsystem->GetConversionFunctionGraph(MVVMExtensionPtr->GetWidgetBlueprint(), Binding, UE::MVVM::IsForwardBinding(Binding.BindingType));
+			if (Graph != nullptr)
 			{
-				SelectedEntries.Add(NewBindingEntry);
+				TArray<UK2Node_CallFunction*> FunctionNodes;
+				Graph->GetNodesOfClass<UK2Node_CallFunction>(FunctionNodes);
+
+				if (FunctionNodes.Num() == 1)
+				{
+					UK2Node_CallFunction* CallFunctionNode = FunctionNodes[0];
+					for (UEdGraphPin* Pin : CallFunctionNode->Pins)
+					{
+						// skip output and exec pins, the rest should be arguments
+						if (Pin->Direction != EGPD_Input || Pin->PinName == UEdGraphSchema_K2::PN_Execute || Pin->PinName == UEdGraphSchema_K2::PN_Self)
+						{
+							continue;
+						}
+
+						TSharedRef<FBindingEntry> NewArgumentEntry = MakeShared<FBindingEntry>();
+						NewArgumentEntry->SetParameterName(BindingIndex, Pin->PinName);
+						NewBindingEntry->AddChild(NewArgumentEntry);
+						ReselectIfRequired(NewArgumentEntry);
+					}
+				}
 			}
 		}
 	}
@@ -949,7 +990,7 @@ void SBindingsList::Refresh()
 	{
 		TreeView->RequestTreeRefresh();
 
-		TreeView->SetItemSelection(SelectedEntries, true);
+		TreeView->SetItemSelection(NewSelectedEntries, true);
 
 		for (const TSharedPtr<FBindingEntry>& Entry : RootWidgets)
 		{
@@ -958,20 +999,29 @@ void SBindingsList::Refresh()
 	}
 }
 
-
 TSharedRef<ITableRow> SBindingsList::GenerateEntryRow(TSharedPtr<FBindingEntry> Entry, const TSharedRef<STableViewBase>& OwnerTable) const
 {
 	TSharedPtr<ITableRow> Row;
 
 	if (UMVVMWidgetBlueprintExtension_View* MVVMExtensionPtr = MVVMExtension.Get())
 	{
-		if (Entry->GetBinding(MVVMExtensionPtr->GetBlueprintView()) != nullptr)
-		{
-			Row = SNew(SBindingRow, OwnerTable, Entry, MVVMExtensionPtr->GetBlueprintView(), MVVMExtensionPtr->GetWidgetBlueprint());
-		}
-		else
+		switch (Entry->GetRowType())
 		{
-			Row = SNew(SWidgetRow, OwnerTable, Entry, MVVMExtensionPtr->GetWidgetBlueprint());
+			case FBindingEntry::ERowType::Widget:
+			{
+				Row = SNew(SWidgetRow, OwnerTable, Entry, MVVMExtensionPtr->GetWidgetBlueprint());
+				break;
+			}
+			case FBindingEntry::ERowType::Binding:
+			{
+				Row = SNew(SBindingRow, OwnerTable, Entry, MVVMExtensionPtr->GetWidgetBlueprint());
+				break;
+			}
+			case FBindingEntry::ERowType::Parameter:
+			{
+				Row = SNew(SFunctionParameterRow, OwnerTable, Entry, MVVMExtensionPtr->GetWidgetBlueprint());
+				break;
+			}
 		}
 
 		return Row.ToSharedRef();
@@ -981,6 +1031,22 @@ TSharedRef<ITableRow> SBindingsList::GenerateEntryRow(TSharedPtr<FBindingEntry>
 	return SNew(STableRow<TSharedPtr<FBindingEntry>>, OwnerTable);
 }
 
+namespace Private
+{
+	void RemoveAllBindings(UMVVMBlueprintView* BlueprintView, const TConstArrayView<TSharedPtr<FBindingEntry>>& Entries)
+	{
+		for (const TSharedPtr<FBindingEntry>& Entry : Entries)
+		{
+			RemoveAllBindings(BlueprintView, Entry->GetChildren());
+
+			if (FMVVMBlueprintViewBinding* Binding = Entry->GetBinding(BlueprintView))
+			{
+				BlueprintView->RemoveBinding(Binding);
+			}
+		}
+	}
+}
+
 TSharedPtr<SWidget> SBindingsList::OnSourceConstructContextMenu()
 {
 	const bool bShouldCloseWindowAfterMenuSelection = true;
@@ -1001,14 +1067,14 @@ TSharedPtr<SWidget> SBindingsList::OnSourceConstructContextMenu()
 
 					if (UMVVMBlueprintView* BlueprintView = MVVMExtensionPtr->GetBlueprintView())
 					{
-						for (const TSharedPtr<FBindingEntry>& Entry : Selection)
-						{
-							BlueprintView->RemoveBinding(Entry->GetBinding(BlueprintView));
-						}
+						Private::RemoveAllBindings(BlueprintView, Selection);
 					}
 				}
 			});
-		MenuBuilder.AddMenuEntry(LOCTEXT("RemoveBinding", "Remove Binding"), LOCTEXT("RemoveBindingTooltip", "Remove this binding."), FSlateIcon(FAppStyle::GetAppStyleSetName(), "Icons.Delete"), RemoveAction);
+		MenuBuilder.AddMenuEntry(LOCTEXT("RemoveBinding", "Remove Binding"), 
+			LOCTEXT("RemoveBindingTooltip", "Remove this binding."), 
+			FSlateIcon(FAppStyle::GetAppStyleSetName(), "Icons.Delete"), 
+			RemoveAction);
 	}
 
 	return MenuBuilder.MakeWidget();
diff --git a/Engine/Plugins/Runtime/ModelViewViewModel/Source/ModelViewViewModelEditor/Private/Widgets/SMVVMViewModelBindingListWidget.cpp b/Engine/Plugins/Runtime/ModelViewViewModel/Source/ModelViewViewModelEditor/Private/Widgets/SMVVMViewModelBindingListWidget.cpp
--- a/Engine/Plugins/Runtime/ModelViewViewModel/Source/ModelViewViewModelEditor/Private/Widgets/SMVVMViewModelBindingListWidget.cpp
+++ b/Engine/Plugins/Runtime/ModelViewViewModel/Source/ModelViewViewModelEditor/Private/Widgets/SMVVMViewModelBindingListWidget.cpp
@@ -24,32 +24,33 @@ using UE::PropertyViewer::SPropertyViewer;
 namespace UE::MVVM
 {
 
-	FFieldIterator_Bindable::FFieldIterator_Bindable(const UWidgetBlueprint* InWidgetBlueprint, EFieldVisibility InVisibilityFlags)
-	: WidgetBlueprint(InWidgetBlueprint)
-	, FieldVisibilityFlags(InVisibilityFlags)
-{ }
-
+FFieldIterator_Bindable::FFieldIterator_Bindable(const UWidgetBlueprint* InWidgetBlueprint, EFieldVisibility InVisibilityFlags, const FProperty* InAssignableTo) :
+	WidgetBlueprint(InWidgetBlueprint),
+	FieldVisibilityFlags(InVisibilityFlags),
+	AssignableTo(InAssignableTo)
+{
+}
 
 TArray<FFieldVariant> FFieldIterator_Bindable::GetFields(const UStruct* Struct) const
 {
 	TArray<FFieldVariant> Result;
 
-	auto AddResult = [Flags = FieldVisibilityFlags, &Result, Struct](const TArray<FMVVMAvailableBinding>& AvailableBindingsList)
+	auto AddResult = [this, &Result, Struct](const TArray<FMVVMAvailableBinding>& AvailableBindingsList)
 	{
 		Result.Reserve(AvailableBindingsList.Num());
 		for (const FMVVMAvailableBinding& Value : AvailableBindingsList)
 		{
-			if (EnumHasAllFlags(Flags, EFieldVisibility::Readable) && !Value.IsReadable())
+			if (EnumHasAllFlags(FieldVisibilityFlags, EFieldVisibility::Readable) && !Value.IsReadable())
 			{
 				continue;
 			} 
 
-			if (EnumHasAllFlags(Flags, EFieldVisibility::Writable) && !Value.IsWritable())
+			if (EnumHasAllFlags(FieldVisibilityFlags, EFieldVisibility::Writable) && !Value.IsWritable())
 			{
 				continue;
 			}
 			
-			if (EnumHasAllFlags(Flags, EFieldVisibility::Notify) && !Value.HasNotify())
+			if (EnumHasAllFlags(FieldVisibilityFlags, EFieldVisibility::Notify) && !Value.HasNotify())
 			{
 				continue;
 			}
@@ -57,10 +58,20 @@ TArray<FFieldVariant> FFieldIterator_Bindable::GetFields(const UStruct* Struct)
 			FMVVMFieldVariant FieldVariant = BindingHelper::FindFieldByName(Struct, Value.GetBindingName());
 			if (FieldVariant.IsFunction())
 			{
+				if (AssignableTo != nullptr && !BindingHelper::ArePropertiesCompatible(FieldVariant.GetFunction()->GetReturnProperty(), AssignableTo))
+				{
+					continue;
+				}
+
 				Result.Add(FFieldVariant(FieldVariant.GetFunction()));
 			}
 			else if (FieldVariant.IsProperty())
 			{
+				if (AssignableTo != nullptr && !BindingHelper::ArePropertiesCompatible(FieldVariant.GetProperty(), AssignableTo))
+				{
+					continue;
+				}
+
 				Result.Add(FFieldVariant(FieldVariant.GetProperty()));
 			}
 		}
@@ -115,7 +126,7 @@ TArray<FFieldVariant> FFieldIterator_Bindable::GetFields(const UStruct* Struct)
 /** */
 void SSourceBindingList::Construct(const FArguments& InArgs, const UWidgetBlueprint* WidgetBlueprint)
 {
-	FieldIterator = MakeUnique<FFieldIterator_Bindable>(WidgetBlueprint, InArgs._FieldVisibilityFlags);
+	FieldIterator = MakeUnique<FFieldIterator_Bindable>(WidgetBlueprint, InArgs._FieldVisibilityFlags, InArgs._AssignableTo);
 
 	OnDoubleClicked = InArgs._OnDoubleClicked;
 
@@ -126,6 +137,7 @@ void SSourceBindingList::Construct(const FArguments& InArgs, const UWidgetBluepr
 		.PropertyVisibility(SPropertyViewer::EPropertyVisibility::Hidden)
 		.bShowFieldIcon(true)
 		.bSanitizeName(true)
+		.SelectionMode(InArgs._EnableSelection ? ESelectionMode::Single : ESelectionMode::None)
 		.bShowSearchBox(InArgs._ShowSearchBox)
 		.OnSelectionChanged(this, &SSourceBindingList::HandleSelectionChanged)
 		.OnDoubleClicked(this, &SSourceBindingList::HandleDoubleClicked)
diff --git a/Engine/Plugins/Runtime/ModelViewViewModel/Source/ModelViewViewModelEditor/Private/Widgets/SMVVMViewModelBindingListWidget.h b/Engine/Plugins/Runtime/ModelViewViewModel/Source/ModelViewViewModelEditor/Private/Widgets/SMVVMViewModelBindingListWidget.h
--- a/Engine/Plugins/Runtime/ModelViewViewModel/Source/ModelViewViewModelEditor/Private/Widgets/SMVVMViewModelBindingListWidget.h
+++ b/Engine/Plugins/Runtime/ModelViewViewModel/Source/ModelViewViewModelEditor/Private/Widgets/SMVVMViewModelBindingListWidget.h
@@ -35,12 +35,13 @@ namespace UE::MVVM
 	class FFieldIterator_Bindable : public UE::PropertyViewer::FFieldIterator_BlueprintVisible
 	{
 	public:
-		FFieldIterator_Bindable(const UWidgetBlueprint* WidgetBlueprint, EFieldVisibility InVisibilityFlags);
+		FFieldIterator_Bindable(const UWidgetBlueprint* WidgetBlueprint, EFieldVisibility InVisibilityFlags, const FProperty* InAssignableTo = nullptr);
 		virtual TArray<FFieldVariant> GetFields(const UStruct*) const override;
 
 	private:
 		TWeakObjectPtr<const UWidgetBlueprint> WidgetBlueprint; 
 		EFieldVisibility FieldVisibilityFlags = EFieldVisibility::All;
+		const FProperty* AssignableTo = nullptr;
 	};
 
 	/** 
@@ -49,13 +50,14 @@ namespace UE::MVVM
 	class SSourceBindingList : public SCompoundWidget
 	{
 	public:
-
 		DECLARE_DELEGATE_OneParam(FOnDoubleClicked, const FMVVMBlueprintPropertyPath&);
 	
 		SLATE_BEGIN_ARGS(SSourceBindingList) {}
 			SLATE_ARGUMENT(FBindingSource, InitialSource)
 			SLATE_ARGUMENT_DEFAULT(bool, ShowSearchBox) = false;
 			SLATE_ARGUMENT_DEFAULT(EFieldVisibility, FieldVisibilityFlags) = EFieldVisibility::All;
+			SLATE_ARGUMENT_DEFAULT(bool, EnableSelection) = true;
+			SLATE_ARGUMENT_DEFAULT(const FProperty*, AssignableTo) = nullptr;
 			SLATE_EVENT(FOnDoubleClicked, OnDoubleClicked)
 		SLATE_END_ARGS()
 
