Missed new files from previous submit.

#rnx

diff --git a/Engine/Plugins/Runtime/ModelViewViewModel/Source/ModelViewViewModelEditor/Private/Widgets/SMVVMFunctionParameter.cpp b/Engine/Plugins/Runtime/ModelViewViewModel/Source/ModelViewViewModelEditor/Private/Widgets/SMVVMFunctionParameter.cpp
new file mode 100644
--- /dev/null
+++ b/Engine/Plugins/Runtime/ModelViewViewModel/Source/ModelViewViewModelEditor/Private/Widgets/SMVVMFunctionParameter.cpp
@@ -0,0 +1,177 @@
+// Copyright Epic Games, Inc. All Rights Reserved.
+
+#include "SMVVMFunctionParameter.h" 
+#include "EdGraph/EdGraph.h"
+#include "EdGraph/EdGraphPin.h"
+#include "Editor.h"
+#include "K2Node_CallFunction.h"
+#include "MVVMEditorSubsystem.h"
+#include "NodeFactory.h"
+#include "SGraphPin.h"
+#include "Types/MVVMBindingMode.h"
+#include "WidgetBlueprint.h"
+#include "Widgets/SBoxPanel.h"
+#include "Widgets/SMVVMFieldSelector.h"
+#include "Widgets/Input/SCheckBox.h"
+#include "Widgets/Layout/SBox.h"
+
+#define LOCTEXT_NAMESPACE "MVVMFieldBinding"
+
+namespace UE::MVVM
+{
+
+void SFunctionParameter::Construct(const FArguments& InArgs)
+{
+	WidgetBlueprint = InArgs._WidgetBlueprint;
+	check(InArgs._WidgetBlueprint != nullptr);
+
+	Binding = InArgs._Binding;
+	check(Binding != nullptr);
+
+	ParameterName = InArgs._ParameterName;
+	check(!ParameterName.IsNone());
+	
+	bSourceToDestination = InArgs._SourceToDestination;
+
+	GetBindingModeDelegate = InArgs._OnGetBindingMode;
+	check(GetBindingModeDelegate.IsBound());
+
+	UMVVMEditorSubsystem* EditorSubsystem = GEditor->GetEditorSubsystem<UMVVMEditorSubsystem>();
+	UEdGraphPin* Pin = EditorSubsystem->FindConversionFunctionArgumentPin(InArgs._WidgetBlueprint, *Binding, ParameterName, bSourceToDestination);
+	check(Pin);
+
+	const UFunction* ConversionFunction = EditorSubsystem->GetConversionFunction(InArgs._WidgetBlueprint, *Binding, bSourceToDestination);
+	const FProperty* Property = ConversionFunction->FindPropertyByName(ParameterName);
+	check(Property);
+
+	// create a new pin widget so that we can get the default value widget out of it
+	TSharedPtr<SGraphPin> PinWidget = FNodeFactory::CreateK2PinWidget(Pin);
+
+	TSharedRef<SWidget> ValueWidget = SNullWidget::NullWidget;
+	if (PinWidget.IsValid())
+	{
+		GraphPin = PinWidget;
+		ValueWidget = PinWidget->GetDefaultValueWidget();
+	}
+
+	if (ValueWidget == SNullWidget::NullWidget)
+	{
+		ValueWidget = SNew(STextBlock)
+			.Text(LOCTEXT("DefaultValue", "Default Value"))
+			.TextStyle(FAppStyle::Get(), "HintText");
+	}
+
+	ValueWidget->SetVisibility(TAttribute<EVisibility>::CreateSP(this, &SFunctionParameter::OnGetVisibility, true));
+
+	bool bFromViewModel;
+	if (bSourceToDestination)
+	{
+		bFromViewModel = UE::MVVM::IsForwardBinding(Binding->BindingType);
+	}
+	else
+	{
+		bFromViewModel = UE::MVVM::IsBackwardBinding(Binding->BindingType);
+	}
+
+	ChildSlot
+	[
+		SNew(SBox)
+		.MinDesiredWidth(100)
+		[
+			SNew(SHorizontalBox)
+			+ SHorizontalBox::Slot()
+			.Padding(0, 2)
+			.VAlign(VAlign_Center)
+			[
+				SNew(SOverlay)
+				+ SOverlay::Slot()
+				[
+					ValueWidget
+				]
+				+ SOverlay::Slot()
+				[
+					SAssignNew(FieldSelector, UE::MVVM::SFieldSelector, WidgetBlueprint.Get(), bFromViewModel)
+					.Visibility(this, &SFunctionParameter::OnGetVisibility, false)
+					.SelectedField(this, &SFunctionParameter::OnGetSelectedField)
+					.BindingMode_Lambda([this]() { return GetBindingModeDelegate.Execute(); })
+					.OnSelectionChanged(this, &SFunctionParameter::OnSelectionChanged)
+					.AssignableTo(Property)
+				]
+			]
+			+ SHorizontalBox::Slot()
+			.HAlign(HAlign_Right)
+			.VAlign(VAlign_Center)
+			.Padding(8, 0, 0, 0)
+			.AutoWidth()
+			[
+				SNew(SCheckBox)
+				.ToolTipText(LOCTEXT("BindArgument", "Bind this argument to a property."))
+				.Style(FAppStyle::Get(), "ToggleButtonCheckbox")
+				.IsChecked(this, &SFunctionParameter::OnGetIsBindArgumentChecked)
+				.OnCheckStateChanged(this, &SFunctionParameter::OnBindArgumentChecked)
+				.Padding(FMargin(4))
+				[
+					SNew(SImage)
+					.DesiredSizeOverride(FVector2D(16, 16))
+					.Image_Lambda([this]()
+						{
+							ECheckBoxState CheckState = OnGetIsBindArgumentChecked();
+							return (CheckState == ECheckBoxState::Checked) ? FAppStyle::GetBrush("Icons.Link") : FAppStyle::GetBrush("Icons.Unlink");
+						})
+					]
+				]
+			]
+		];
+}
+
+EVisibility SFunctionParameter::OnGetVisibility(bool bDefaultValue) const
+{
+	FMVVMBlueprintPropertyPath Path = OnGetSelectedField();
+
+	// if we're not bound then show the default value widget, otherwise show the binding widget
+	return Path.IsEmpty() == bDefaultValue ? EVisibility::Visible : EVisibility::Collapsed;
+}
+
+ECheckBoxState SFunctionParameter::OnGetIsBindArgumentChecked() const 
+{
+	UMVVMEditorSubsystem* EditorSubsystem = GEditor->GetEditorSubsystem<UMVVMEditorSubsystem>();
+	FMVVMBlueprintPropertyPath CurrentPath = EditorSubsystem->GetPathForConversionFunctionArgument(WidgetBlueprint.Get(), *Binding, ParameterName, bSourceToDestination);
+	return CurrentPath.IsEmpty() ? ECheckBoxState::Unchecked : ECheckBoxState::Checked;
+}
+
+void SFunctionParameter::OnBindArgumentChecked(ECheckBoxState Checked)
+{
+	UMVVMEditorSubsystem* EditorSubsystem = GEditor->GetEditorSubsystem<UMVVMEditorSubsystem>();
+
+	FMVVMBlueprintPropertyPath Path;
+
+	if (Checked == ECheckBoxState::Checked)
+	{
+		// HACK: Just set a placeholder viewmodel reference
+		if (const UMVVMBlueprintView* View = EditorSubsystem->GetView(WidgetBlueprint.Get()))
+		{
+			const TArrayView<const FMVVMBlueprintViewModelContext> ViewModels = View->GetViewModels();
+			if (ViewModels.Num() > 0)
+			{
+				Path.SetViewModelId(ViewModels[0].GetViewModelId());
+			}
+		}
+	}
+
+	EditorSubsystem->SetPathForConversionFunctionArgument(WidgetBlueprint.Get(), *Binding, ParameterName, Path, bSourceToDestination);
+}
+
+FMVVMBlueprintPropertyPath SFunctionParameter::OnGetSelectedField() const
+{
+	UMVVMEditorSubsystem* EditorSubsystem = GEditor->GetEditorSubsystem<UMVVMEditorSubsystem>();
+	return EditorSubsystem->GetPathForConversionFunctionArgument(WidgetBlueprint.Get(), *Binding, ParameterName, bSourceToDestination);
+}
+
+void SFunctionParameter::OnSelectionChanged(const FMVVMBlueprintPropertyPath& Selected)
+{
+	UMVVMEditorSubsystem* EditorSubsystem = GEditor->GetEditorSubsystem<UMVVMEditorSubsystem>();
+	EditorSubsystem->SetPathForConversionFunctionArgument(WidgetBlueprint.Get(), *Binding, ParameterName, Selected, bSourceToDestination);
+}
+
+}
+#undef LOCTEXT_NAMESPACE
diff --git a/Engine/Plugins/Runtime/ModelViewViewModel/Source/ModelViewViewModelEditor/Private/Widgets/SMVVMFunctionParameter.h b/Engine/Plugins/Runtime/ModelViewViewModel/Source/ModelViewViewModelEditor/Private/Widgets/SMVVMFunctionParameter.h
new file mode 100644
--- /dev/null
+++ b/Engine/Plugins/Runtime/ModelViewViewModel/Source/ModelViewViewModelEditor/Private/Widgets/SMVVMFunctionParameter.h
@@ -0,0 +1,60 @@
+// Copyright Epic Games, Inc. All Rights Reserved.
+
+#pragma once
+
+#include "Delegates/DelegateCombinations.h"
+#include "MVVMPropertyPath.h"
+#include "Widgets/SCompoundWidget.h"
+
+class SGraphPin;
+class UK2Node_CallFunction;
+class UWidgetBlueprint;
+
+struct FMVVMBlueprintViewBinding;
+
+enum class ECheckBoxState : uint8;
+enum class EMVVMBindingMode : uint8;
+
+namespace UE::MVVM
+{
+
+class SFieldSelector;
+
+class SFunctionParameter : public SCompoundWidget
+{
+public:
+	DECLARE_DELEGATE_RetVal(EMVVMBindingMode, FGetBindingMode)
+
+	SLATE_BEGIN_ARGS(SFunctionParameter) :
+		_WidgetBlueprint(nullptr),
+		_Binding(nullptr)
+		{}
+		SLATE_EVENT(FGetBindingMode, OnGetBindingMode)
+		SLATE_ARGUMENT(UWidgetBlueprint*, WidgetBlueprint)
+		SLATE_ARGUMENT(FMVVMBlueprintViewBinding*, Binding)
+		SLATE_ARGUMENT(FName, ParameterName)
+		SLATE_ARGUMENT(bool, SourceToDestination)
+	SLATE_END_ARGS()
+
+	void Construct(const FArguments& InArgs);
+
+private:
+
+	ECheckBoxState OnGetIsBindArgumentChecked() const;
+	void OnBindArgumentChecked(ECheckBoxState Checked);
+	EVisibility OnGetVisibility(bool bDefaultValue) const;
+	FMVVMBlueprintPropertyPath OnGetSelectedField() const;
+	void OnSelectionChanged(const FMVVMBlueprintPropertyPath& Selected);
+
+private:
+
+	TWeakObjectPtr<UWidgetBlueprint> WidgetBlueprint;
+	FMVVMBlueprintViewBinding* Binding = nullptr;
+	FName ParameterName;
+	TSharedPtr<SGraphPin> GraphPin;
+	TSharedPtr<UE::MVVM::SFieldSelector> FieldSelector;
+	bool bSourceToDestination = true;
+	FGetBindingMode GetBindingModeDelegate;
+};
+
+}
