UMG: Add PreConstruct and Construct call back to extension. Do not remove extension in the destruction call. They maybe be added out of the construct/destruct loop. Fix bad test in EvaluateObjectProperty.
#cherrypick 21739076

diff --git a/Engine/Plugins/Runtime/ModelViewViewModel/Source/ModelViewViewModel/Private/Bindings/MVVMFieldPathHelper.cpp b/Engine/Plugins/Runtime/ModelViewViewModel/Source/ModelViewViewModel/Private/Bindings/MVVMFieldPathHelper.cpp
--- a/Engine/Plugins/Runtime/ModelViewViewModel/Source/ModelViewViewModel/Private/Bindings/MVVMFieldPathHelper.cpp
+++ b/Engine/Plugins/Runtime/ModelViewViewModel/Source/ModelViewViewModel/Private/Bindings/MVVMFieldPathHelper.cpp
@@ -415,7 +415,7 @@ TValueOrError<UObject*, void> EvaluateObjectProperty(const FFieldContext& InSour
 		check(InSource.GetObjectVariant().GetUObject());
 
 		UFunction* Function = InSource.GetFieldVariant().GetFunction();
-		if (InSource.GetObjectVariant().GetUObject()->GetClass()->IsChildOf(Function->GetOuterUClass()))
+		if (!InSource.GetObjectVariant().GetUObject()->GetClass()->IsChildOf(Function->GetOuterUClass()))
 		{
 			return MakeError();
 		}
diff --git a/Engine/Plugins/Runtime/ModelViewViewModel/Source/ModelViewViewModel/Private/View/MVVMView.cpp b/Engine/Plugins/Runtime/ModelViewViewModel/Source/ModelViewViewModel/Private/View/MVVMView.cpp
--- a/Engine/Plugins/Runtime/ModelViewViewModel/Source/ModelViewViewModel/Private/View/MVVMView.cpp
+++ b/Engine/Plugins/Runtime/ModelViewViewModel/Source/ModelViewViewModel/Private/View/MVVMView.cpp
@@ -142,7 +142,7 @@ void UMVVMView::SetViewModel(FName ViewModelName, UMVVMViewModelBase* NewValue)
 						EnableLibraryBinding(Binding, Index);
 					}
 				}
-				else
+				else if (bConstructed)
 				{
 					// Enabled the default bindings
 					for (int32 Index = 0; Index < CompiledBindings.Num(); ++Index)
diff --git a/Engine/Plugins/Runtime/ModelViewViewModel/Source/ModelViewViewModel/Private/View/MVVMViewClass.cpp b/Engine/Plugins/Runtime/ModelViewViewModel/Source/ModelViewViewModel/Private/View/MVVMViewClass.cpp
--- a/Engine/Plugins/Runtime/ModelViewViewModel/Source/ModelViewViewModel/Private/View/MVVMViewClass.cpp
+++ b/Engine/Plugins/Runtime/ModelViewViewModel/Source/ModelViewViewModel/Private/View/MVVMViewClass.cpp
@@ -167,8 +167,7 @@ FString FMVVMViewClass_CompiledBinding::ToString() const
 
 void UMVVMViewClass::Construct(UUserWidget* UserWidget)
 {
-	check(UserWidget->GetExtension<UMVVMView>() == nullptr);
-
+	ensure(UserWidget->GetExtension<UMVVMView>() == nullptr);
 	UMVVMView* View = UserWidget->AddExtension<UMVVMView>();
 	if (ensure(View))
 	{
