MVVM - Editor subsystem now correctly handles calls for simple conversion functions that don't create a wrapper UEdGraph.

#rnx

diff --git a/Engine/Plugins/Runtime/ModelViewViewModel/Source/ModelViewViewModelEditor/Private/MVVMEditorSubsystem.cpp b/Engine/Plugins/Runtime/ModelViewViewModel/Source/ModelViewViewModelEditor/Private/MVVMEditorSubsystem.cpp
--- a/Engine/Plugins/Runtime/ModelViewViewModel/Source/ModelViewViewModelEditor/Private/MVVMEditorSubsystem.cpp
+++ b/Engine/Plugins/Runtime/ModelViewViewModel/Source/ModelViewViewModelEditor/Private/MVVMEditorSubsystem.cpp
@@ -297,19 +297,19 @@ UEdGraph* UMVVMEditorSubsystem::GetConversionFunctionGraph(const UWidgetBlueprin
 
 UFunction* UMVVMEditorSubsystem::GetConversionFunction(const UWidgetBlueprint* WidgetBlueprint, const FMVVMBlueprintViewBinding& Binding, bool bSourceToDestination) const
 {
-	const UEdGraph* Graph = GetConversionFunctionGraph(WidgetBlueprint, Binding, bSourceToDestination);
-
-	TArray<UK2Node_CallFunction*> FunctionNodes;
-	Graph->GetNodesOfClass<UK2Node_CallFunction>(FunctionNodes);
-
-	if (FunctionNodes.Num() != 1)
+	if (const UEdGraph* Graph = GetConversionFunctionGraph(WidgetBlueprint, Binding, bSourceToDestination))
 	{
-		// ambiguous result, no idea what our function node is
-		return nullptr;
-	}
+		TArray<UK2Node_CallFunction*> FunctionNodes;
+		Graph->GetNodesOfClass<UK2Node_CallFunction>(FunctionNodes);
 
-	const UK2Node_CallFunction* CallFunctionNode = FunctionNodes[0];
-	return CallFunctionNode->GetTargetFunction();
+		if (FunctionNodes.Num() == 1)
+		{
+			return FunctionNodes[0]->GetTargetFunction();
+		}
+	}
+	
+	FMemberReference ConversionFunctionReference = bSourceToDestination ? Binding.Conversion.SourceToDestinationFunction : Binding.Conversion.DestinationToSourceFunction;
+	return ConversionFunctionReference.ResolveMember<UFunction>(WidgetBlueprint->SkeletonGeneratedClass);
 }
 
 UEdGraph* UMVVMEditorSubsystem::CreateConversionFunctionWrapperGraph(UWidgetBlueprint* WidgetBlueprint, const FMVVMBlueprintViewBinding& Binding, const UFunction* ConversionFunction, bool bSourceToDestination)
@@ -387,6 +387,10 @@ UEdGraph* UMVVMEditorSubsystem::CreateConversionFunctionWrapperGraph(UWidgetBlue
 UEdGraphPin* UMVVMEditorSubsystem::FindConversionFunctionArgumentPin(const UWidgetBlueprint* WidgetBlueprint, const FMVVMBlueprintViewBinding& Binding, FName ParameterName, bool bSourceToDestination)
 {
 	const UEdGraph* Graph = GetConversionFunctionGraph(WidgetBlueprint, Binding, bSourceToDestination);
+	if (Graph == nullptr)
+	{
+		return nullptr;
+	}
 
 	TArray<UK2Node_CallFunction*> FunctionNodes;
 	Graph->GetNodesOfClass<UK2Node_CallFunction>(FunctionNodes);
@@ -565,13 +569,14 @@ void UMVVMEditorSubsystem::SetCompileForBinding(UWidgetBlueprint* WidgetBlueprin
 
 bool UMVVMEditorSubsystem::IsValidConversionFunction(const UFunction* Function, const FMVVMBlueprintPropertyPath& Source, const FMVVMBlueprintPropertyPath& Destination) const
 {
-	if (Source.IsEmpty() || Destination.IsEmpty())
+	TValueOrError<const FProperty*, FString> ReturnResult = UE::MVVM::BindingHelper::TryGetReturnTypeForConversionFunction(Function);
+	if (ReturnResult.HasError())
 	{
 		return false;
 	}
 
-	TValueOrError<const FProperty*, FString> ReturnResult = UE::MVVM::BindingHelper::TryGetReturnTypeForConversionFunction(Function);
-	if (ReturnResult.HasError())
+	const FProperty* ReturnProperty = ReturnResult.GetValue();
+	if (ReturnProperty == nullptr)
 	{
 		return false;
 	}
@@ -582,16 +587,21 @@ bool UMVVMEditorSubsystem::IsValidConversionFunction(const UFunction* Function,
 		return false;
 	}
 
+	const FProperty* SourceProperty = nullptr;
+
 	TArray<UE::MVVM::FMVVMConstFieldVariant> SourceFields = Source.GetFields();
-	TArray<UE::MVVM::FMVVMConstFieldVariant> DestFields = Destination.GetFields();
-	if (SourceFields.Num() == 0 || DestFields.Num() == 0)
+	if (SourceFields.Num() > 0)
 	{
-		return false;
+		SourceProperty = SourceFields.Last().IsProperty() ? SourceFields.Last().GetProperty() : UE::MVVM::BindingHelper::GetReturnProperty(SourceFields.Last().GetFunction());
 	}
 
-	const FProperty* ReturnProperty = ReturnResult.GetValue();
-	const FProperty* SourceProperty = SourceFields.Last().IsProperty() ? SourceFields.Last().GetProperty() : UE::MVVM::BindingHelper::GetReturnProperty(SourceFields.Last().GetFunction());
-	const FProperty* DestinationProperty = DestFields.Last().IsProperty() ? DestFields.Last().GetProperty() : UE::MVVM::BindingHelper::GetFirstArgumentProperty(DestFields.Last().GetFunction());
+	const FProperty* DestinationProperty = nullptr;
+
+	TArray<UE::MVVM::FMVVMConstFieldVariant> DestFields = Destination.GetFields();
+	if (DestFields.Num() > 0)
+	{
+		DestinationProperty = DestFields.Last().IsProperty() ? DestFields.Last().GetProperty() : UE::MVVM::BindingHelper::GetFirstArgumentProperty(DestFields.Last().GetFunction());
+	}
 
 	// check that at least one source -> argument binding is compatible
 	bool bAnyCompatible = false;
@@ -605,7 +615,8 @@ bool UMVVMEditorSubsystem::IsValidConversionFunction(const UFunction* Function,
 			return false;
 		}
 
-		if (UE::MVVM::BindingHelper::ArePropertiesCompatible(SourceProperty, ArgumentProperty))
+		if (SourceProperty == nullptr ||
+			UE::MVVM::BindingHelper::ArePropertiesCompatible(SourceProperty, ArgumentProperty))
 		{
 			bAnyCompatible = true;
 		}
@@ -617,7 +628,8 @@ bool UMVVMEditorSubsystem::IsValidConversionFunction(const UFunction* Function,
 	}
 
 	// check that the return -> dest is valid
-	if (!UE::MVVM::BindingHelper::ArePropertiesCompatible(ReturnProperty, DestinationProperty))
+	if (DestinationProperty != nullptr &&
+		!UE::MVVM::BindingHelper::ArePropertiesCompatible(ReturnProperty, DestinationProperty))
 	{
 		return false;
 	}
@@ -625,15 +637,33 @@ bool UMVVMEditorSubsystem::IsValidConversionFunction(const UFunction* Function,
 	return true;
 }
 
-TArray<UFunction*> UMVVMEditorSubsystem::GetAvailableConversionFunctions(const UWidgetBlueprint* WidgetBlueprint, const FMVVMBlueprintPropertyPath& Source, const FMVVMBlueprintPropertyPath& Destination) const
+bool UMVVMEditorSubsystem::IsSimpleConversionFunction(const UFunction* Function) const
 {
-	TArray<UFunction*> ConversionFunctions;
+	TValueOrError<const FProperty*, FString> ReturnResult = UE::MVVM::BindingHelper::TryGetReturnTypeForConversionFunction(Function);
+	if (ReturnResult.HasError())
+	{
+		return false;
+	}
 
-	if (Source.IsEmpty() || Destination.IsEmpty())
+	const FProperty* ReturnProperty = ReturnResult.GetValue();
+	if (ReturnProperty == nullptr)
 	{
-		return ConversionFunctions;
+		return false;
 	}
 
+	TValueOrError<TArray<const FProperty*>, FString> ArgumentsResult = UE::MVVM::BindingHelper::TryGetArgumentsForConversionFunction(Function);
+	if (ArgumentsResult.HasError())
+	{
+		return false;
+	}
+
+	return ArgumentsResult.GetValue().Num() == 1;
+}
+
+TArray<UFunction*> UMVVMEditorSubsystem::GetAvailableConversionFunctions(const UWidgetBlueprint* WidgetBlueprint, const FMVVMBlueprintPropertyPath& Source, const FMVVMBlueprintPropertyPath& Destination) const
+{
+	TArray<UFunction*> ConversionFunctions;
+
 	const UMVVMSubsystem* Subsystem = GEngine->GetEngineSubsystem<UMVVMSubsystem>();
 
 	FBlueprintActionDatabase& ActionDatabase = FBlueprintActionDatabase::Get();
@@ -676,18 +706,20 @@ FMVVMBlueprintPropertyPath UMVVMEditorSubsystem::GetPathForConversionFunctionArg
 	UEdGraph* WrapperGraph = GetConversionFunctionGraph(WidgetBlueprint, Binding, bSourceToDestination);
 	if (WrapperGraph == nullptr)
 	{
-		return FMVVMBlueprintPropertyPath();
+		if (bSourceToDestination)
+		{
+			return Binding.ViewModelPath;
+		}
+		else
+		{
+			return Binding.WidgetPath;
+		}
 	}
 	 
 	UEdGraphNode* ConversionNode = UE::MVVM::Private::FindConversionNode(WrapperGraph);
 
 	const UEdGraphPin* ArgumentPin = ConversionNode->FindPin(ArgumentName, EGPD_Input);
-	if (!ensureMsgf(ArgumentPin != nullptr, TEXT("Couldn't find argument pin named '%s'."), *ArgumentName.ToString()))
-	{
-		return FMVVMBlueprintPropertyPath();
-	}
-
-	if (ArgumentPin->LinkedTo.Num() == 0)
+	if (ArgumentPin == nullptr || ArgumentPin->LinkedTo.Num() == 0)
 	{
 		return FMVVMBlueprintPropertyPath();
 	}
@@ -756,8 +788,18 @@ FMVVMBlueprintPropertyPath UMVVMEditorSubsystem::GetPathForConversionFunctionArg
 void UMVVMEditorSubsystem::SetPathForConversionFunctionArgument(UWidgetBlueprint* WidgetBlueprint, FMVVMBlueprintViewBinding& Binding, FName ArgumentName, const FMVVMBlueprintPropertyPath& Path, bool bSourceToDestination) const
 {
 	UEdGraph* ConversionFunctionGraph = GetConversionFunctionGraph(WidgetBlueprint, Binding, bSourceToDestination);
-	if (!ensure(ConversionFunctionGraph != nullptr))
+	if (ConversionFunctionGraph == nullptr)
 	{
+		// simple conversion function
+		if (bSourceToDestination)
+		{
+			Binding.ViewModelPath = Path;
+		}
+		else
+		{
+			Binding.WidgetPath = Path;
+		}
+
 		return;
 	}
 
diff --git a/Engine/Plugins/Runtime/ModelViewViewModel/Source/ModelViewViewModelEditor/Public/MVVMEditorSubsystem.h b/Engine/Plugins/Runtime/ModelViewViewModel/Source/ModelViewViewModelEditor/Public/MVVMEditorSubsystem.h
--- a/Engine/Plugins/Runtime/ModelViewViewModel/Source/ModelViewViewModelEditor/Public/MVVMEditorSubsystem.h
+++ b/Engine/Plugins/Runtime/ModelViewViewModel/Source/ModelViewViewModelEditor/Public/MVVMEditorSubsystem.h
@@ -61,6 +61,9 @@ public:
 	UFUNCTION(BlueprintCallable, Category = "MVVM")
 	bool IsValidConversionFunction(const UFunction* Function, const FMVVMBlueprintPropertyPath& Source, const FMVVMBlueprintPropertyPath& Destination) const;
 
+	UFUNCTION(BlueprintCallable, Category = "MVVM")
+	bool IsSimpleConversionFunction(const UFunction* Function) const;
+
 	UFUNCTION(BlueprintCallable, Category = "MVVM")
 	UEdGraph* GetConversionFunctionGraph(const UWidgetBlueprint* WidgetBlueprint, const FMVVMBlueprintViewBinding& Binding, bool bSourceToDestination) const;
 
