MVVM - SFunctionParameter now correctly handles simple conversion function arguments.

#rnx

diff --git a/Engine/Plugins/Runtime/ModelViewViewModel/Source/ModelViewViewModelEditor/Private/Widgets/SMVVMFunctionParameter.cpp b/Engine/Plugins/Runtime/ModelViewViewModel/Source/ModelViewViewModelEditor/Private/Widgets/SMVVMFunctionParameter.cpp
--- a/Engine/Plugins/Runtime/ModelViewViewModel/Source/ModelViewViewModelEditor/Private/Widgets/SMVVMFunctionParameter.cpp
+++ b/Engine/Plugins/Runtime/ModelViewViewModel/Source/ModelViewViewModelEditor/Private/Widgets/SMVVMFunctionParameter.cpp
@@ -37,21 +37,21 @@ void SFunctionParameter::Construct(const FArguments& InArgs)
 	check(GetBindingModeDelegate.IsBound());
 
 	UMVVMEditorSubsystem* EditorSubsystem = GEditor->GetEditorSubsystem<UMVVMEditorSubsystem>();
-	UEdGraphPin* Pin = EditorSubsystem->FindConversionFunctionArgumentPin(InArgs._WidgetBlueprint, *Binding, ParameterName, bSourceToDestination);
-	check(Pin);
-
 	const UFunction* ConversionFunction = EditorSubsystem->GetConversionFunction(InArgs._WidgetBlueprint, *Binding, bSourceToDestination);
 	const FProperty* Property = ConversionFunction->FindPropertyByName(ParameterName);
 	check(Property);
 
-	// create a new pin widget so that we can get the default value widget out of it
-	TSharedPtr<SGraphPin> PinWidget = FNodeFactory::CreateK2PinWidget(Pin);
-
 	TSharedRef<SWidget> ValueWidget = SNullWidget::NullWidget;
-	if (PinWidget.IsValid())
+
+	UEdGraphPin* Pin = EditorSubsystem->FindConversionFunctionArgumentPin(InArgs._WidgetBlueprint, *Binding, ParameterName, bSourceToDestination);
+	if (Pin != nullptr)
 	{
-		GraphPin = PinWidget;
-		ValueWidget = PinWidget->GetDefaultValueWidget();
+		// create a new pin widget so that we can get the default value widget out of it
+		if (TSharedPtr<SGraphPin> PinWidget = FNodeFactory::CreateK2PinWidget(Pin))
+		{
+			GraphPin = PinWidget;
+			ValueWidget = PinWidget->GetDefaultValueWidget();
+		}
 	}
 
 	if (ValueWidget == SNullWidget::NullWidget)
@@ -94,7 +94,7 @@ void SFunctionParameter::Construct(const FArguments& InArgs)
 					.Visibility(this, &SFunctionParameter::OnGetVisibility, false)
 					.SelectedField(this, &SFunctionParameter::OnGetSelectedField)
 					.BindingMode_Lambda([this]() { return GetBindingModeDelegate.Execute(); })
-					.OnSelectionChanged(this, &SFunctionParameter::OnSelectionChanged)
+					.OnFieldSelectionChanged(this, &SFunctionParameter::OnFieldSelectionChanged)
 					.AssignableTo(Property)
 				]
 			]
@@ -122,6 +122,11 @@ void SFunctionParameter::Construct(const FArguments& InArgs)
 				]
 			]
 		];
+
+	if (Pin == nullptr)
+	{
+		OnBindArgumentChecked(ECheckBoxState::Checked);
+	}
 }
 
 EVisibility SFunctionParameter::OnGetVisibility(bool bDefaultValue) const
@@ -167,7 +172,7 @@ FMVVMBlueprintPropertyPath SFunctionParameter::OnGetSelectedField() const
 	return EditorSubsystem->GetPathForConversionFunctionArgument(WidgetBlueprint.Get(), *Binding, ParameterName, bSourceToDestination);
 }
 
-void SFunctionParameter::OnSelectionChanged(const FMVVMBlueprintPropertyPath& Selected)
+void SFunctionParameter::OnFieldSelectionChanged(FMVVMBlueprintPropertyPath Selected)
 {
 	UMVVMEditorSubsystem* EditorSubsystem = GEditor->GetEditorSubsystem<UMVVMEditorSubsystem>();
 	EditorSubsystem->SetPathForConversionFunctionArgument(WidgetBlueprint.Get(), *Binding, ParameterName, Selected, bSourceToDestination);
diff --git a/Engine/Plugins/Runtime/ModelViewViewModel/Source/ModelViewViewModelEditor/Private/Widgets/SMVVMFunctionParameter.h b/Engine/Plugins/Runtime/ModelViewViewModel/Source/ModelViewViewModelEditor/Private/Widgets/SMVVMFunctionParameter.h
--- a/Engine/Plugins/Runtime/ModelViewViewModel/Source/ModelViewViewModelEditor/Private/Widgets/SMVVMFunctionParameter.h
+++ b/Engine/Plugins/Runtime/ModelViewViewModel/Source/ModelViewViewModelEditor/Private/Widgets/SMVVMFunctionParameter.h
@@ -44,13 +44,14 @@ private:
 	void OnBindArgumentChecked(ECheckBoxState Checked);
 	EVisibility OnGetVisibility(bool bDefaultValue) const;
 	FMVVMBlueprintPropertyPath OnGetSelectedField() const;
-	void OnSelectionChanged(const FMVVMBlueprintPropertyPath& Selected);
+	void OnFieldSelectionChanged(FMVVMBlueprintPropertyPath Selected);
 
 private:
 
 	TWeakObjectPtr<UWidgetBlueprint> WidgetBlueprint;
 	FMVVMBlueprintViewBinding* Binding = nullptr;
 	FName ParameterName;
+	/** This reference is just to keep the default value widget alive. */
 	TSharedPtr<SGraphPin> GraphPin;
 	TSharedPtr<UE::MVVM::SFieldSelector> FieldSelector;
 	bool bSourceToDestination = true;
