Removed fully commented-out MVVMPropertyAccess files.

#rnx

diff --git a/Engine/Plugins/Runtime/ModelViewViewModel/Source/ModelViewViewModelEditor/Private/Customizations/MVVMBindPropertiesDetailView.cpp b/Engine/Plugins/Runtime/ModelViewViewModel/Source/ModelViewViewModelEditor/Private/Customizations/MVVMBindPropertiesDetailView.cpp
deleted file mode 100644
diff --git a/Engine/Plugins/Runtime/ModelViewViewModel/Source/ModelViewViewModelEditor/Private/Customizations/MVVMBindPropertiesDetailView.h b/Engine/Plugins/Runtime/ModelViewViewModel/Source/ModelViewViewModelEditor/Private/Customizations/MVVMBindPropertiesDetailView.h
deleted file mode 100644
diff --git a/Engine/Plugins/Runtime/ModelViewViewModel/Source/ModelViewViewModelEditor/Private/MVVMPropertyAccess.cpp b/Engine/Plugins/Runtime/ModelViewViewModel/Source/ModelViewViewModelEditor/Private/MVVMPropertyAccess.cpp
deleted file mode 100644
diff --git a/Engine/Plugins/Runtime/ModelViewViewModel/Source/ModelViewViewModelEditor/Private/MVVMPropertyAccess.h b/Engine/Plugins/Runtime/ModelViewViewModel/Source/ModelViewViewModelEditor/Private/MVVMPropertyAccess.h
deleted file mode 100644
diff --git a/Engine/Plugins/Runtime/ModelViewViewModel/Source/ModelViewViewModelEditor/Private/ModelViewViewModelEditorModule.cpp b/Engine/Plugins/Runtime/ModelViewViewModel/Source/ModelViewViewModelEditor/Private/ModelViewViewModelEditorModule.cpp
--- a/Engine/Plugins/Runtime/ModelViewViewModel/Source/ModelViewViewModelEditor/Private/ModelViewViewModelEditorModule.cpp
+++ b/Engine/Plugins/Runtime/ModelViewViewModel/Source/ModelViewViewModelEditor/Private/ModelViewViewModelEditorModule.cpp
@@ -5,7 +5,6 @@
 #include "AssetToolsModule.h"
 #include "BlueprintModes/WidgetBlueprintApplicationMode.h"
 #include "BlueprintModes/WidgetBlueprintApplicationModes.h"
-#include "Customizations/MVVMBindPropertiesDetailView.h"
 #include "Customizations/MVVMPropertyBindingExtension.h"
 #include "EdGraphSchema_K2.h"
 #include "Framework/Docking/LayoutExtender.h"
@@ -19,7 +18,6 @@
 #include "MVVMBlueprintViewModelContext.h"
 #include "MVVMEditorCommands.h"
 #include "MVVMWidgetBlueprintExtension_View.h"
-#include "PropertyEditorModule.h"
 #include "StatusBarSubsystem.h"
 #include "Styling/MVVMEditorStyle.h"
 #include "Styling/SlateColor.h"
@@ -44,7 +42,6 @@ const FName AnimationTabSummonerTabID(TEXT("Animations"));
 void FModelViewViewModelEditorModule::StartupModule()
 {
 	FMVVMEditorStyle::CreateInstance();
-	BindDetailView = MakeUnique<FMVVMBindPropertiesDetailView>();
 
 	IUMGEditorModule& UMGEditorModule = FModuleManager::LoadModuleChecked<IUMGEditorModule>("UMGEditor");
 	UMGEditorModule.OnRegisterTabsForEditor().AddRaw(this, &FModelViewViewModelEditorModule::HandleRegisterBlueprintEditorTab);
@@ -59,9 +56,6 @@ void FModelViewViewModelEditorModule::StartupModule()
 	PropertyBindingExtension = MakeShared<FMVVMPropertyBindingExtension>();
 	UMGEditorModule.GetPropertyBindingExtensibilityManager()->AddExtension(PropertyBindingExtension.ToSharedRef());
 
-	FPropertyEditorModule& PropertyEditorModule = FModuleManager::LoadModuleChecked<FPropertyEditorModule>("PropertyEditor");
-	PropertyEditorModule.GetGlobalRowExtensionDelegate().AddRaw(BindDetailView.Get(), &FMVVMBindPropertiesDetailView::CreatePropertyRowExtension);
-
 	FBlueprintEditorUtils::OnRenameVariableReferencesEvent.AddRaw(this, &FModelViewViewModelEditorModule::HandleRenameVariableReferences);
 
 	{
@@ -99,11 +93,6 @@ void FModelViewViewModelEditorModule::ShutdownModule()
 		AssetTools->Get().UnregisterAssetTypeActions(ViewModelBlueprintActions.ToSharedRef());
 	}
 
-	if (FPropertyEditorModule* PropertyEditorModule = FModuleManager::GetModulePtr<FPropertyEditorModule>("PropertyEditor"))
-	{
-		PropertyEditorModule->GetGlobalRowExtensionDelegate().RemoveAll(BindDetailView.Get());
-	}
-	BindDetailView.Reset();
 	FMVVMEditorStyle::DestroyInstance();
 
 	FMVVMEditorCommands::Unregister();
diff --git a/Engine/Plugins/Runtime/ModelViewViewModel/Source/ModelViewViewModelEditor/Private/ModelViewViewModelEditorModule.h b/Engine/Plugins/Runtime/ModelViewViewModel/Source/ModelViewViewModelEditor/Private/ModelViewViewModelEditorModule.h
--- a/Engine/Plugins/Runtime/ModelViewViewModel/Source/ModelViewViewModelEditor/Private/ModelViewViewModelEditorModule.h
+++ b/Engine/Plugins/Runtime/ModelViewViewModel/Source/ModelViewViewModelEditor/Private/ModelViewViewModelEditorModule.h
@@ -37,7 +37,6 @@ private:
 	void HandleActivateMode(FWidgetBlueprintApplicationMode& InDesignerMode);
 
 private:
-	TUniquePtr<FMVVMBindPropertiesDetailView> BindDetailView;
 	TSharedPtr<FMVVMPropertyBindingExtension> PropertyBindingExtension;
 	TSharedPtr<UE::MVVM::FAssetTypeActions_ViewModelBlueprint> ViewModelBlueprintActions;
 };
