MVVMBlueprintPropertyPath now has an AppendBasePropertyPath() that appends to the end of the path, allowing you to build complex paths.

#rnx

diff --git a/Engine/Plugins/Runtime/ModelViewViewModel/Source/ModelViewViewModelBlueprint/Public/MVVMPropertyPath.h b/Engine/Plugins/Runtime/ModelViewViewModel/Source/ModelViewViewModelBlueprint/Public/MVVMPropertyPath.h
--- a/Engine/Plugins/Runtime/ModelViewViewModel/Source/ModelViewViewModelBlueprint/Public/MVVMPropertyPath.h
+++ b/Engine/Plugins/Runtime/ModelViewViewModel/Source/ModelViewViewModelBlueprint/Public/MVVMPropertyPath.h
@@ -249,7 +249,12 @@ public:
 
 	void SetBasePropertyPath(UE::MVVM::FMVVMConstFieldVariant InField)
 	{
-		Paths.Reset();
+		ResetBasePropertyPath();
+		AppendBasePropertyPath(InField);
+	}
+
+	void AppendBasePropertyPath(UE::MVVM::FMVVMConstFieldVariant InField)
+	{
 		if (!InField.IsEmpty() && !InField.GetName().IsNone())
 		{
 			Paths.AddDefaulted_GetRef().SetBindingReference(InField);
