CIS - Fixed SMVVMFunctionParameter calling OnFieldSelectionChanged, which has not yet been submitted.

#rnx

diff --git a/Engine/Plugins/Runtime/ModelViewViewModel/Source/ModelViewViewModelEditor/Private/Widgets/SMVVMFunctionParameter.cpp b/Engine/Plugins/Runtime/ModelViewViewModel/Source/ModelViewViewModelEditor/Private/Widgets/SMVVMFunctionParameter.cpp
--- a/Engine/Plugins/Runtime/ModelViewViewModel/Source/ModelViewViewModelEditor/Private/Widgets/SMVVMFunctionParameter.cpp
+++ b/Engine/Plugins/Runtime/ModelViewViewModel/Source/ModelViewViewModelEditor/Private/Widgets/SMVVMFunctionParameter.cpp
@@ -94,7 +94,7 @@ void SFunctionParameter::Construct(const FArguments& InArgs)
 					.Visibility(this, &SFunctionParameter::OnGetVisibility, false)
 					.SelectedField(this, &SFunctionParameter::OnGetSelectedField)
 					.BindingMode_Lambda([this]() { return GetBindingModeDelegate.Execute(); })
-					.OnFieldSelectionChanged(this, &SFunctionParameter::OnFieldSelectionChanged)
+					.OnSelectionChanged(this, &SFunctionParameter::OnFieldSelectionChanged)
 					.AssignableTo(Property)
 				]
 			]
