CIS - Fixed SMVVMFunctionParameter::OnSelectionChanged() taking a value instead of a const& due to changes in SMVVMFieldSelector which haven't gone in yet.

#rnx

diff --git a/Engine/Plugins/Runtime/ModelViewViewModel/Source/ModelViewViewModelEditor/Private/Widgets/SMVVMFunctionParameter.cpp b/Engine/Plugins/Runtime/ModelViewViewModel/Source/ModelViewViewModelEditor/Private/Widgets/SMVVMFunctionParameter.cpp
--- a/Engine/Plugins/Runtime/ModelViewViewModel/Source/ModelViewViewModelEditor/Private/Widgets/SMVVMFunctionParameter.cpp
+++ b/Engine/Plugins/Runtime/ModelViewViewModel/Source/ModelViewViewModelEditor/Private/Widgets/SMVVMFunctionParameter.cpp
@@ -172,7 +172,7 @@ FMVVMBlueprintPropertyPath SFunctionParameter::OnGetSelectedField() const
 	return EditorSubsystem->GetPathForConversionFunctionArgument(WidgetBlueprint.Get(), *Binding, ParameterName, bSourceToDestination);
 }
 
-void SFunctionParameter::OnFieldSelectionChanged(FMVVMBlueprintPropertyPath Selected)
+void SFunctionParameter::OnFieldSelectionChanged(const FMVVMBlueprintPropertyPath& Selected)
 {
 	UMVVMEditorSubsystem* EditorSubsystem = GEditor->GetEditorSubsystem<UMVVMEditorSubsystem>();
 	EditorSubsystem->SetPathForConversionFunctionArgument(WidgetBlueprint.Get(), *Binding, ParameterName, Selected, bSourceToDestination);
diff --git a/Engine/Plugins/Runtime/ModelViewViewModel/Source/ModelViewViewModelEditor/Private/Widgets/SMVVMFunctionParameter.h b/Engine/Plugins/Runtime/ModelViewViewModel/Source/ModelViewViewModelEditor/Private/Widgets/SMVVMFunctionParameter.h
--- a/Engine/Plugins/Runtime/ModelViewViewModel/Source/ModelViewViewModelEditor/Private/Widgets/SMVVMFunctionParameter.h
+++ b/Engine/Plugins/Runtime/ModelViewViewModel/Source/ModelViewViewModelEditor/Private/Widgets/SMVVMFunctionParameter.h
@@ -44,7 +44,7 @@ private:
 	void OnBindArgumentChecked(ECheckBoxState Checked);
 	EVisibility OnGetVisibility(bool bDefaultValue) const;
 	FMVVMBlueprintPropertyPath OnGetSelectedField() const;
-	void OnFieldSelectionChanged(FMVVMBlueprintPropertyPath Selected);
+	void OnFieldSelectionChanged(const FMVVMBlueprintPropertyPath& Selected);
 
 private:
 
