MVVM - Field selector now displays valid conversion functions in the same dropdown menu as properties.

#rnx

diff --git a/Engine/Plugins/Runtime/ModelViewViewModel/Source/ModelViewViewModelEditor/Private/Customizations/MVVMPropertyPathCustomization.cpp b/Engine/Plugins/Runtime/ModelViewViewModel/Source/ModelViewViewModelEditor/Private/Customizations/MVVMPropertyPathCustomization.cpp
--- a/Engine/Plugins/Runtime/ModelViewViewModel/Source/ModelViewViewModelEditor/Private/Customizations/MVVMPropertyPathCustomization.cpp
+++ b/Engine/Plugins/Runtime/ModelViewViewModel/Source/ModelViewViewModelEditor/Private/Customizations/MVVMPropertyPathCustomization.cpp
@@ -104,7 +104,7 @@ void FPropertyPathCustomization::CustomizeHeader(TSharedRef<IPropertyHandle> InP
 		[
 			SAssignNew(FieldSelector, SFieldSelector, WidgetBlueprint, !bIsWidget)
 			.TextStyle(FAppStyle::Get(), "SmallText")
-			.OnSelectionChanged(this, &FPropertyPathCustomization::OnPropertySelectionChanged)
+			.OnFieldSelectionChanged(this, &FPropertyPathCustomization::OnFieldSelectionChanged)
 			.SelectedField(this, &FPropertyPathCustomization::OnGetSelectedField)
 			.BindingMode(this, &FPropertyPathCustomization::GetCurrentBindingMode)
 		];
@@ -127,7 +127,7 @@ EMVVMBindingMode FPropertyPathCustomization::GetCurrentBindingMode() const
 	return EMVVMBindingMode::OneWayToDestination;
 }
 
-void FPropertyPathCustomization::OnPropertySelectionChanged(const FMVVMBlueprintPropertyPath& Selected)
+void FPropertyPathCustomization::OnFieldSelectionChanged(FMVVMBlueprintPropertyPath Selected)
 {
 	if (bPropertySelectionChanging)
 	{
@@ -156,7 +156,7 @@ void FPropertyPathCustomization::OnPropertySelectionChanged(const FMVVMBlueprint
 
 void FPropertyPathCustomization::HandleBlueprintChanged(UBlueprint* Blueprint)
 {
-	FieldSelector->Refresh();
+	//FieldSelector->Refresh();
 }
 
 FBindingSource FPropertyPathCustomization::OnGetSelectedSource() const
@@ -217,6 +217,7 @@ FBindingSource FPropertyPathCustomization::OnGetSelectedSource() const
 
 FMVVMBlueprintPropertyPath FPropertyPathCustomization::OnGetSelectedField() const
 {
+	// TODO: This crashes when removing the viewmodel that the selected binding is pointed at
 	TArray<void*> RawData;
 	PropertyHandle->AccessRawData(RawData);
 
diff --git a/Engine/Plugins/Runtime/ModelViewViewModel/Source/ModelViewViewModelEditor/Private/Customizations/MVVMPropertyPathCustomization.h b/Engine/Plugins/Runtime/ModelViewViewModel/Source/ModelViewViewModelEditor/Private/Customizations/MVVMPropertyPathCustomization.h
--- a/Engine/Plugins/Runtime/ModelViewViewModel/Source/ModelViewViewModelEditor/Private/Customizations/MVVMPropertyPathCustomization.h
+++ b/Engine/Plugins/Runtime/ModelViewViewModel/Source/ModelViewViewModelEditor/Private/Customizations/MVVMPropertyPathCustomization.h
@@ -9,9 +9,9 @@
 
 namespace UE::MVVM
 {
- 
-struct FBindingSource;
+
 class SFieldSelector;
+struct FBindingSource;
 
 class FPropertyPathCustomization : public IPropertyTypeCustomization
 {
@@ -30,7 +30,7 @@ public:
 	}
 
 private:
-	void OnPropertySelectionChanged(const FMVVMBlueprintPropertyPath& Selected);
+	void OnFieldSelectionChanged(FMVVMBlueprintPropertyPath Selected);
 
 	void OnOtherPropertyChanged();
 
@@ -53,4 +53,4 @@ private:
 	FDelegateHandle OnBlueprintChangedHandle;
 };
 
-} // namespace UE::MVVM
\ No newline at end of file
+} // namespace UE::MVVM
diff --git a/Engine/Plugins/Runtime/ModelViewViewModel/Source/ModelViewViewModelEditor/Private/Widgets/SMVVMFieldSelector.cpp b/Engine/Plugins/Runtime/ModelViewViewModel/Source/ModelViewViewModelEditor/Private/Widgets/SMVVMFieldSelector.cpp
--- a/Engine/Plugins/Runtime/ModelViewViewModel/Source/ModelViewViewModelEditor/Private/Widgets/SMVVMFieldSelector.cpp
+++ b/Engine/Plugins/Runtime/ModelViewViewModel/Source/ModelViewViewModelEditor/Private/Widgets/SMVVMFieldSelector.cpp
@@ -1,7 +1,9 @@
 // Copyright Epic Games, Inc. All Rights Reserved.
 
 #include "Widgets/SMVVMFieldSelector.h"
- 
+
+#include "Algo/Transform.h"
+#include "Bindings/MVVMBindingHelper.h"
 #include "ClassViewerModule.h"
 #include "Editor.h"
 #include "Hierarchy/SReadOnlyHierarchyView.h"
@@ -45,26 +47,39 @@ FBindingSource GetSourceFromPath(const UWidgetBlueprint* WidgetBlueprint, const
 
 void SFieldSelector::Construct(const FArguments& InArgs, const UWidgetBlueprint* InWidgetBlueprint, bool bInViewModelProperty)
 {
-	OnSelectionChangedDelegate = InArgs._OnSelectionChanged;
-	check(OnSelectionChangedDelegate.IsBound());
-
 	WidgetBlueprint = InWidgetBlueprint;
 	check(WidgetBlueprint);
+	
+	BindingMode = InArgs._BindingMode;
+	check(BindingMode.IsSet());
+
+	bViewModelProperty = bInViewModelProperty;
+	AssignableTo = InArgs._AssignableTo;
+	TextStyle = InArgs._TextStyle;
+
+	OnFieldSelectionChangedDelegate = InArgs._OnFieldSelectionChanged;
+	check(OnFieldSelectionChangedDelegate.IsBound());
 
 	SelectedField = InArgs._SelectedField;
 	check(SelectedField.IsSet());
+
 	CachedSelectedField = SelectedField.Get();
 
-	BindingMode = InArgs._BindingMode;
-	check(BindingMode.IsSet());
+	ShowConversionFunctions = InArgs._ShowConversionFunctions;
 
-	bViewModelProperty = bInViewModelProperty;
+	const UFunction* CurrentFunction = nullptr;
+	if (ShowConversionFunctions.Get(false))
+	{
+		SelectedConversionFunction = InArgs._SelectedConversionFunction;
+		check(SelectedConversionFunction.IsSet());
 
-	AssignableTo = InArgs._AssignableTo;
+		CachedSelectedConversionFunction = SelectedConversionFunction.Get(nullptr);
 
-	TextStyle = InArgs._TextStyle;
+		OnConversionFunctionSelectionChangedDelegate = InArgs._OnConversionFunctionSelectionChanged;
+		check(OnConversionFunctionSelectionChangedDelegate.IsBound());
+	}
 
-	TSharedPtr<SHorizontalBox> HBox;
+	TSharedPtr<SHorizontalBox> SourceEntryBox;
 
 	ChildSlot
 	[
@@ -77,9 +92,30 @@ void SFieldSelector::Construct(const FArguments& InArgs, const UWidgetBlueprint*
 			SNew(SOverlay)
 			+ SOverlay::Slot()
 			[
-				SAssignNew(HBox, SHorizontalBox)
-				.Visibility_Lambda([this]() { return CachedSelectedField.IsEmpty() ? EVisibility::Collapsed : EVisibility::Visible; })
-					
+				SAssignNew(SourceEntryBox, SHorizontalBox)
+				.Visibility_Lambda([this]() { return CachedSelectedConversionFunction == nullptr && !CachedSelectedField.IsEmpty() ? EVisibility::Visible : EVisibility::Collapsed; })
+			]
+			+ SOverlay::Slot()
+			[
+				SNew(SHorizontalBox)
+				.Visibility_Lambda([this]() { return CachedSelectedConversionFunction != nullptr ? EVisibility::Visible : EVisibility::Collapsed; })
+				+ SHorizontalBox::Slot()
+				.HAlign(HAlign_Left)
+				.VAlign(VAlign_Center)
+				.Padding(0, 0, 4.0f, 0)
+				.AutoWidth()
+				[
+					SNew(SImage)
+					.DesiredSizeOverride(FVector2D(16.0f, 16.0f))
+					.Image(FAppStyle::Get().GetBrush("GraphEditor.Function_16x"))
+				]
+				+ SHorizontalBox::Slot()
+				.HAlign(HAlign_Fill)
+				.VAlign(VAlign_Center)
+				[
+					SNew(STextBlock)
+					.Text_Lambda([this]() { return CachedSelectedConversionFunction != nullptr ? CachedSelectedConversionFunction->GetDisplayNameText() : FText::GetEmpty(); })
+				]
 			]
 			+ SOverlay::Slot()
 			[
@@ -87,7 +123,7 @@ void SFieldSelector::Construct(const FArguments& InArgs, const UWidgetBlueprint*
 				.Padding(FMargin(8, 0, 8, 0))
 				[
 					SNew(STextBlock)
-					.Visibility_Lambda([this]() { return CachedSelectedField.IsEmpty() ? EVisibility::Visible : EVisibility::Collapsed; })
+					.Visibility_Lambda([this]() { return CachedSelectedField.IsEmpty() && CachedSelectedConversionFunction == nullptr ? EVisibility::Visible : EVisibility::Collapsed; })
 					.TextStyle(FAppStyle::Get(), "HintText")
 					.Text(LOCTEXT("None", "No field selected"))
 				]
@@ -97,7 +133,7 @@ void SFieldSelector::Construct(const FArguments& InArgs, const UWidgetBlueprint*
 	
 	if (InArgs._ShowSource)
 	{
-		HBox->AddSlot()
+		SourceEntryBox->AddSlot()
 			.Padding(8, 0, 0, 0)
 			.HAlign(HAlign_Left)
 			.VAlign(VAlign_Center)
@@ -106,7 +142,7 @@ void SFieldSelector::Construct(const FArguments& InArgs, const UWidgetBlueprint*
 				SAssignNew(SelectedSourceWidget, SSourceEntry)
 			];
 
-		HBox->AddSlot()
+		SourceEntryBox->AddSlot()
 			.Padding(8, 0)
 			.HAlign(HAlign_Left)
 			.VAlign(VAlign_Center)
@@ -121,7 +157,7 @@ void SFieldSelector::Construct(const FArguments& InArgs, const UWidgetBlueprint*
 		FixedSource = InArgs._Source;
 	}
 
-	HBox->AddSlot()
+	SourceEntryBox->AddSlot()
 		.Padding(0, 0, 8, 0)
 		.HAlign(HAlign_Left)
 		.VAlign(VAlign_Center)
@@ -137,6 +173,7 @@ void SFieldSelector::Construct(const FArguments& InArgs, const UWidgetBlueprint*
 void SFieldSelector::Refresh()
 {
 	CachedSelectedField = SelectedField.Get();
+	CachedSelectedConversionFunction = SelectedConversionFunction.Get(nullptr);
 
 	if (SelectedEntryWidget.IsValid())
 	{
@@ -161,33 +198,52 @@ bool SFieldSelector::IsSelectEnabled() const
 	if (BindingList.IsValid())
 	{
 		FMVVMBlueprintPropertyPath Path = BindingList->GetSelectedProperty();
-		return (Path.IsFromViewModel() || Path.IsFromWidget()) && !Path.GetBasePropertyPath().IsEmpty();
+		if ((Path.IsFromViewModel() || Path.IsFromWidget()) && !Path.GetBasePropertyPath().IsEmpty())
+		{
+			return true;
+		}
+	}
+
+	if (ConversionFunctionList.IsValid())
+	{
+		if (ConversionFunctionList->GetSelectedItems().Num() > 0)
+		{
+			return true;
+		}
 	}
+
 	return false;
 }
 
-FReply SFieldSelector::OnSelectProperty()
+FReply SFieldSelector::OnSelectClicked()
 {
 	if (BindingList.IsValid())
 	{
-		FMVVMBlueprintPropertyPath Path = BindingList->GetSelectedProperty();
-		SetSelection(Path);
+		SetPropertySelection(BindingList->GetSelectedProperty());
+	}
+
+	if (ConversionFunctionList.IsValid())
+	{
+		TArray<const UFunction*> Selection = ConversionFunctionList->GetSelectedItems();
+		SetConversionFunctionSelection(Selection.Num() > 0 ? Selection[0] : nullptr);
 	}
+
 	return FReply::Handled();
 }
 
 bool SFieldSelector::IsClearEnabled() const
 {
-	return !CachedSelectedField.IsEmpty();
+	return !CachedSelectedField.IsEmpty() || CachedSelectedConversionFunction != nullptr;
 }
 
-FReply SFieldSelector::OnClearBinding()
+FReply SFieldSelector::OnClearClicked()
 {
-	SetSelection(FMVVMBlueprintPropertyPath());
+	SetPropertySelection(FMVVMBlueprintPropertyPath());
+	SetConversionFunctionSelection(nullptr);
 	return FReply::Handled();
 }
 
-FReply SFieldSelector::OnCancel()
+FReply SFieldSelector::OnCancelClicked()
 {
 	if (ComboButton.IsValid())
 	{
@@ -197,20 +253,40 @@ FReply SFieldSelector::OnCancel()
 	return FReply::Handled();
 }
 
-void SFieldSelector::SetSelection(const FMVVMBlueprintPropertyPath& SelectedPath)
+void SFieldSelector::SetPropertySelection(const FMVVMBlueprintPropertyPath& SelectedPath)
 {
 	if (ComboButton.IsValid())
 	{
 		ComboButton->SetIsOpen(false);
 	}
 
-	OnSelectionChangedDelegate.Execute(SelectedPath);
+	OnFieldSelectionChangedDelegate.Execute(SelectedPath);
 	Refresh();
 }
 
-void SFieldSelector::OnSearchTextChanged(const FText& NewText)
+void SFieldSelector::OnSearchBoxTextChanged(const FText& NewText)
 {
-	BindingList->SetRawFilterText(NewText);
+	if (BindingList.IsValid())
+	{
+		BindingList->SetRawFilterText(NewText);
+	}
+
+	if (WidgetList.IsValid())
+	{
+		WidgetList->SetRawFilterText(NewText);
+	}
+
+	if (ConversionFunctionCategoryTree.IsValid())
+	{
+		FilterConversionFunctionCategories();
+		ConversionFunctionCategoryTree->RequestTreeRefresh();
+	}
+
+	if (ConversionFunctionList.IsValid())
+	{
+		FilterConversionFunctions();
+		ConversionFunctionList->RequestListRefresh();
+	}
 }
 
 void SFieldSelector::OnViewModelSelected(FBindingSource Source, ESelectInfo::Type)
@@ -222,6 +298,11 @@ void SFieldSelector::OnViewModelSelected(FBindingSource Source, ESelectInfo::Typ
 		BindingList->Clear();
 		BindingList->AddSources(Selection);
 	}
+
+	if (ConversionFunctionCategoryTree.IsValid())
+	{
+		ConversionFunctionCategoryTree->ClearSelection();
+	}
 }
 
 void SFieldSelector::OnWidgetSelected(FName WidgetName, ESelectInfo::Type)
@@ -233,9 +314,14 @@ void SFieldSelector::OnWidgetSelected(FName WidgetName, ESelectInfo::Type)
 		BindingList->Clear();
 		BindingList->AddSource(Source);
 	}
+
+	if (ConversionFunctionCategoryTree.IsValid())
+	{
+		ConversionFunctionCategoryTree->ClearSelection();
+	}
 }
 
-TSharedRef<ITableRow> SFieldSelector::GenerateRowForViewModel(FBindingSource ViewModel, const TSharedRef<STableViewBase>& OwnerTable) const
+TSharedRef<ITableRow> SFieldSelector::GenerateViewModelRow(FBindingSource ViewModel, const TSharedRef<STableViewBase>& OwnerTable) const
 {
 	return SNew(STableRow<FBindingSource>, OwnerTable)
 		[
@@ -258,109 +344,371 @@ TSharedRef<ITableRow> SFieldSelector::GenerateRowForViewModel(FBindingSource Vie
 				SNew(STextBlock)
 				.Text(ViewModel.DisplayName)
 				.ColorAndOpacity(FSlateColor::UseForeground())
+				.HighlightText_Lambda([this]() { return SearchBox.IsValid() ? SearchBox->GetText() : FText::GetEmpty(); })
 			]
 		];
 }
 
-TSharedRef<SWidget> SFieldSelector::OnGetMenuContent()
+bool SFieldSelector::FilterConversionFunctionCategoryChildren(const TArray<FString>& FilterStrings, const TArray<TSharedPtr<FConversionFunctionItem>>& SourceArray, TArray<TSharedPtr<FConversionFunctionItem>>& OutDestArray)
 {
-	ViewModelSources.Reset();
-	ViewModelList.Reset();
+	bool bAnyAdded = false;
 
-	BindingList = SNew(SSourceBindingList, WidgetBlueprint)
-		.ShowSearchBox(false)
-		.OnDoubleClicked(this, &SFieldSelector::SetSelection)
-		.FieldVisibilityFlags(GetFieldVisibilityFlags())
-		.AssignableTo(AssignableTo);
+	for (const TSharedPtr<FConversionFunctionItem>& SourceItem : SourceArray)
+	{
+		TArray<TSharedPtr<FConversionFunctionItem>> FilteredChildren;
 
-	TSharedRef<SVerticalBox> VBox = SNew(SVerticalBox)
-		+ SVerticalBox::Slot()
-		.HAlign(HAlign_Fill)
-		.VAlign(VAlign_Top)
-		.Padding(0, 4, 0, 4)
-		.AutoHeight()
-		[
-			SNew(SSearchBox)
-			.OnTextChanged(this, &SFieldSelector::OnSearchTextChanged)
-		];
+		// check if our name matches the filters
+		bool bMatchesFilters = true;
 
-	// single fixed source
-	if (FixedSource.IsSet())
+		const FString ItemName = SourceItem->Function != nullptr ? SourceItem->Function->GetName() : SourceItem->CategoryName;
+		
+		for (const FString& Filter : FilterStrings)
+		{
+			if (!ItemName.Contains(Filter))
+			{
+				bMatchesFilters = false;
+				break;
+			}
+		}
+
+		bool bChildrenMatch = false;
+
+		if (bMatchesFilters)
+		{
+			// if we match the filter, then we always want to add all our children
+			FilteredChildren = SourceItem->Children;
+		}
+		else
+		{
+			bChildrenMatch = FilterConversionFunctionCategoryChildren(FilterStrings, SourceItem->Children, FilteredChildren);
+		}
+
+		// then add this item to the destination array
+		if (bMatchesFilters || bChildrenMatch)
+		{
+			TSharedPtr<FConversionFunctionItem>& NewItem = OutDestArray.Add_GetRef(MakeShared<FConversionFunctionItem>());
+			NewItem->CategoryName = SourceItem->CategoryName;
+			NewItem->Function = SourceItem->Function;
+			NewItem->Children = FilteredChildren;
+
+			bAnyAdded = true;
+		}
+	}
+
+	return bAnyAdded;
+}
+
+void SFieldSelector::FilterConversionFunctionCategories()
+{
+	FilteredConversionFunctionCategories.Reset();
+
+	TArray<FString> FilterStrings;
+	if (SearchBox.IsValid())
 	{
-		BindingList->AddSource(FixedSource.GetValue());
+		SearchBox->GetText().ToString().ParseIntoArrayWS(FilterStrings);
+	}
 
-		VBox->AddSlot()
-			[
-				BindingList.ToSharedRef()
-			];
+	if (FilterStrings.IsEmpty())
+	{
+		// don't bother filtering if we don't have any search terms
+		FilteredConversionFunctionCategories = ConversionFunctionCategories;
+		return;
 	}
-	else
-	{ 
-		// show source picker
-		TSharedPtr<SWidget> SourcePicker;
-		if (bViewModelProperty)
+	
+	FilterConversionFunctionCategoryChildren(FilterStrings, ConversionFunctionCategories, FilteredConversionFunctionCategories);
+}
+
+void SFieldSelector::FilterViewModels()
+{
+	FilteredViewModelSources.Reset();
+
+	TArray<FString> FilterStrings;
+	if (SearchBox.IsValid())
+	{
+		SearchBox->GetText().ToString().ParseIntoArrayWS(FilterStrings);
+	}
+
+	if (FilterStrings.IsEmpty())
+	{
+		FilteredViewModelSources = ViewModelSources;
+		return;
+	}
+
+	for (const FBindingSource& ViewModel : ViewModelSources)
+	{
+		const FString DisplayName = ViewModel.DisplayName.ToString();
+		const FString ClassName = ViewModel.Class != nullptr ? ViewModel.Class->GetName() : FString();
+		const FString Name = ViewModel.Name.ToString();
+
+		bool bMatchesFilters = true;
+		
+		for (const FString& Filter : FilterStrings)
+		{
+			if (!DisplayName.Contains(Filter) && 
+				!ClassName.Contains(Filter) &&
+				!Name.Contains(Filter))
+			{
+				bMatchesFilters = false;
+				break;
+			}
+		}
+
+		if (bMatchesFilters)
 		{
-			FBindingSource SelectedSource;
+			FBindingSource& NewSource = FilteredViewModelSources.Add_GetRef(FBindingSource());
+			NewSource.Class = ViewModel.Class;
+			NewSource.Name = ViewModel.Name;
+			NewSource.DisplayName = ViewModel.DisplayName;
+			NewSource.ViewModelId = ViewModel.ViewModelId;
+		}
+	}
+}
 
-			UMVVMEditorSubsystem* EditorSubsystem = GEditor->GetEditorSubsystem<UMVVMEditorSubsystem>();
-			if (const UMVVMBlueprintView* View = EditorSubsystem->GetView(WidgetBlueprint))
+TSharedRef<SWidget> SFieldSelector::CreateSourcePanel()
+{
+	ConversionFunctionCategoryTree.Reset();
+	ConversionFunctionCategories.Reset();
+	ViewModelSources.Reset();
+	ViewModelList.Reset();
+	WidgetList.Reset();
+
+	// show source picker
+	TSharedPtr<SWidget> SourcePicker;
+	if (bViewModelProperty)
+	{
+		FBindingSource SelectedSource;
+
+		UMVVMEditorSubsystem* EditorSubsystem = GEditor->GetEditorSubsystem<UMVVMEditorSubsystem>();
+		if (const UMVVMBlueprintView* View = EditorSubsystem->GetView(WidgetBlueprint))
+		{
+			const TArrayView<const FMVVMBlueprintViewModelContext> ViewModels = View->GetViewModels();
+			ViewModelSources.Reserve(ViewModels.Num());
+
+			for (const FMVVMBlueprintViewModelContext& ViewModel : ViewModels)
 			{
-				const TArrayView<const FMVVMBlueprintViewModelContext> ViewModels = View->GetViewModels();
-				ViewModelSources.Reserve(ViewModels.Num());
+				FBindingSource& Source = ViewModelSources.AddDefaulted_GetRef();
+				Source.Class = ViewModel.GetViewModelClass();
+				Source.Name = ViewModel.GetViewModelName();
+				Source.DisplayName = ViewModel.GetDisplayName();
+				Source.ViewModelId = ViewModel.GetViewModelId();
 
-				for (const FMVVMBlueprintViewModelContext& ViewModel : ViewModels)
+				if (Source.ViewModelId == CachedSelectedField.GetViewModelId())
 				{
-					FBindingSource& Source = ViewModelSources.AddDefaulted_GetRef();
-					Source.Class = ViewModel.GetViewModelClass();
-					Source.Name = ViewModel.GetViewModelName();
-					Source.DisplayName = ViewModel.GetDisplayName();
-					Source.ViewModelId = ViewModel.GetViewModelId();
-
-					if (Source.ViewModelId == CachedSelectedField.GetViewModelId())
-					{
-						SelectedSource = Source;
-					}
+					SelectedSource = Source;
 				}
 			}
+		}
 
-			ViewModelList = SNew(SListView<FBindingSource>)
-				.OnGenerateRow(this, &SFieldSelector::GenerateRowForViewModel)
-				.SelectionMode(ESelectionMode::Multi)
-				.ListItemsSource(&ViewModelSources)
-				.OnSelectionChanged(this, &SFieldSelector::OnViewModelSelected);
+		FilterViewModels();
 
+		ViewModelList = SNew(SListView<FBindingSource>)
+			.ListItemsSource(&FilteredViewModelSources)
+			.OnGenerateRow(this, &SFieldSelector::GenerateViewModelRow)
+			.SelectionMode(ESelectionMode::Multi)
+			.OnSelectionChanged(this, &SFieldSelector::OnViewModelSelected);
+
+		if (SelectedSource.IsValid())
+		{
 			ViewModelList->SetItemSelection(SelectedSource, true);
+		}
+
+		SourcePicker = ViewModelList;
+	}
+	else
+	{
+		WidgetList = SNew(SReadOnlyHierarchyView, WidgetBlueprint)
+			.OnSelectionChanged(this, &SFieldSelector::OnWidgetSelected)
+			.SelectionMode(ESelectionMode::Multi)
+			.ShowSearch(false);
+
+		SourcePicker = WidgetList;
+	}
+
+	if (ShowConversionFunctions.Get(false))
+	{
+		UMVVMEditorSubsystem* Subsystem = GEditor->GetEditorSubsystem<UMVVMEditorSubsystem>();
+		TArray<UFunction*> AllConversionFunctions = Subsystem->GetAvailableConversionFunctions(WidgetBlueprint, FMVVMBlueprintPropertyPath(), FMVVMBlueprintPropertyPath());
 
-			SourcePicker = ViewModelList;
+		// remove all incompatible conversion functions
+		for (int32 Idx = 0; Idx < AllConversionFunctions.Num(); ++Idx)
+		{
+			const UFunction* Function = AllConversionFunctions[Idx];
+			const FProperty* ReturnProperty = BindingHelper::GetReturnProperty(Function);
+			const FProperty* AssignToProperty = AssignableTo.Get(nullptr);
+
+			if (AssignToProperty != nullptr && !BindingHelper::ArePropertiesCompatible(ReturnProperty, AssignToProperty))
+			{
+				AllConversionFunctions.RemoveAt(Idx);
+				--Idx;
+			}
 		}
-		else
+
+		TSharedPtr<FConversionFunctionItem> CurrentSelectedItem;
+
+		TArray<FString> SubCategories;
+		TArray<TSharedPtr<FConversionFunctionItem>> PathToSelection;
+
+		for (const UFunction* Function : AllConversionFunctions)
 		{
-			WidgetList = SNew(SReadOnlyHierarchyView, WidgetBlueprint)
-				.OnSelectionChanged(this, &SFieldSelector::OnWidgetSelected)
+			FText CategoryName = Function->GetMetaDataText("Category");
+
+			SubCategories.Reset();
+			CategoryName.ToString().ParseIntoArray(SubCategories, TEXT("|"));
+
+			TArray<TSharedPtr<FConversionFunctionItem>>* CurrentItems = &ConversionFunctionCategories;
+			TSharedPtr<FConversionFunctionItem> ParentItem;
+
+			// create items for the entire category path
+			// eg. "Math|Boolean|AND" 
+			// Math 
+			//   > Boolean
+			//     > AND
+			for (const FString& SubCategory : SubCategories)
+			{
+				const FString Trimmed = SubCategory.TrimStartAndEnd();
+
+				ParentItem = FindOrCreateItemForCategory(*CurrentItems, Trimmed);
+				CurrentItems = &ParentItem->Children;
+
+				if (Function == CachedSelectedConversionFunction)
+				{
+					PathToSelection.Add(ParentItem);
+				}
+			}
+
+			TSharedPtr<FConversionFunctionItem> Item = CurrentItems->Add_GetRef(MakeShared<FConversionFunctionItem>());
+			Item->Function = Function;
+		}
+
+		FilterConversionFunctionCategories();
+
+		ConversionFunctionCategoryRoot.Reset();
+		TSharedPtr<FConversionFunctionItem> RootItem = ConversionFunctionCategoryRoot.Add_GetRef(MakeShared<FConversionFunctionItem>());
+		RootItem->CategoryName = TEXT("Conversion Functions");
+
+		TSharedRef<SWidget> StackedSourcePicker =
+			SNew(SVerticalBox)
+			+ SVerticalBox::Slot()
+			.AutoHeight()
+			[
+				SourcePicker.ToSharedRef()
+			]
+			+ SVerticalBox::Slot()
+			.AutoHeight()
+			[
+				SAssignNew(ConversionFunctionCategoryTree, STreeView<TSharedPtr<FConversionFunctionItem>>)
 				.SelectionMode(ESelectionMode::Multi)
-				.ShowSearch(false);
+				.TreeItemsSource(&ConversionFunctionCategoryRoot)
+				.OnGenerateRow(this, &SFieldSelector::GenerateConversionFunctionCategoryRow)
+				.OnSelectionChanged(this, &SFieldSelector::OnConversionFunctionCategorySelected)
+				.OnGetChildren(this, &SFieldSelector::GetConversionFunctionCategoryChildren)
+			];
 
-			SourcePicker = WidgetList;
+		if (PathToSelection.Num() > 0)
+		{
+			ConversionFunctionCategoryTree->SetItemExpansion(RootItem, true);
+
+			for (const TSharedPtr<FConversionFunctionItem>& ParentItem : PathToSelection)
+			{
+				ConversionFunctionCategoryTree->SetItemExpansion(ParentItem, true);
+			}
+			ConversionFunctionCategoryTree->SetItemSelection(PathToSelection.Last(), true);
 		}
 
+		SourcePicker = StackedSourcePicker;
+	}
+
+	return SourcePicker.ToSharedRef();
+}
+
+TSharedRef<SWidget> SFieldSelector::OnGetMenuContent()
+{
+	TSharedRef<SVerticalBox> VBox = SNew(SVerticalBox)
+		+ SVerticalBox::Slot()
+		.HAlign(HAlign_Fill)
+		.VAlign(VAlign_Top)
+		.Padding(0, 4, 0, 4)
+		.AutoHeight()
+		[
+			SAssignNew(SearchBox, SSearchBox)
+			.OnTextChanged(this, &SFieldSelector::OnSearchBoxTextChanged)
+		];
+
+	// If we're showing conversion functions, we don't want to set the AssignableTo property of SSourceBindingList, because then it will only show exact matches, 
+	// and since we're also showing conversion functions we know that's not what the user wants.
+	// However, in the case we're not showing conversion functions, we want only exact matches.
+	const FProperty* AssignableToProperty = nullptr;
+	const bool bShowConversionFunctions = ShowConversionFunctions.Get(false);
+	if (!bShowConversionFunctions)
+	{
+		AssignableToProperty = AssignableTo.Get(nullptr);
+	}
+
+	BindingList = SNew(SSourceBindingList, WidgetBlueprint)
+		.ShowSearchBox(false)
+		.OnDoubleClicked(this, &SFieldSelector::SetPropertySelection)
+		.FieldVisibilityFlags(GetFieldVisibilityFlags())
+		.AssignableTo(AssignableToProperty);
+
+	if (FixedSource.IsSet())
+	{
+		// single fixed source, don't show the separate source panel
+		BindingList->AddSource(FixedSource.GetValue());
+
+		VBox->AddSlot()
+			[
+				BindingList.ToSharedRef()
+			];
+	} 
+	else
+	{
+		TSharedPtr<SVerticalBox> InnerVBox;
+
 		VBox->AddSlot()
 			.FillHeight(1.0f)
 			[
-				SNew(SSplitter)
-				.PhysicalSplitterHandleSize(4.0f)
-				+ SSplitter::Slot()
-				.Value(0.5f)
-				[
-					SourcePicker.ToSharedRef()
-				]
-				+ SSplitter::Slot()
-				.Value(0.5f)
+				SNew(SBorder)
+				.BorderImage(FAppStyle::Get().GetBrush("Brushes.Recessed"))
 				[
-					BindingList.ToSharedRef()
+					SNew(SSplitter)
+					.PhysicalSplitterHandleSize(4.0f)
+					+ SSplitter::Slot()
+					.Value(0.5f)
+					[
+						CreateSourcePanel()
+					]
+					+ SSplitter::Slot()
+					.Value(0.5f)
+					[
+						SAssignNew(InnerVBox, SVerticalBox)
+						+ SVerticalBox::Slot()
+						.AutoHeight()
+						[
+							BindingList.ToSharedRef()
+						]
+					]
 				]
 			];
+
+		if (bShowConversionFunctions)
+		{
+			InnerVBox->AddSlot()
+				.AutoHeight()
+				[
+					 SAssignNew(ConversionFunctionList, SListView<const UFunction*>)
+					.SelectionMode(ESelectionMode::Single)
+					.ListItemsSource(&FilteredConversionFunctions)
+					.OnMouseButtonDoubleClick(this, &SFieldSelector::SetConversionFunctionSelection)
+					.OnGenerateRow(this, &SFieldSelector::GenerateConversionFunctionRow)
+				];
+
+			if (CachedSelectedConversionFunction != nullptr)
+			{
+				ConversionFunctionList->SetItemSelection(CachedSelectedConversionFunction, true);
+			}
+		}
 	}
-				
+
 	VBox->AddSlot()
 		.Padding(4, 4, 4, 0)
 		.HAlign(HAlign_Right)
@@ -371,7 +719,7 @@ TSharedRef<SWidget> SFieldSelector::OnGetMenuContent()
 			.AutoWidth()
 			[
 				SNew(SPrimaryButton)
-				.OnClicked(this, &SFieldSelector::OnSelectProperty)
+				.OnClicked(this, &SFieldSelector::OnSelectClicked)
 				.IsEnabled(this, &SFieldSelector::IsSelectEnabled)
 				.Text(LOCTEXT("Select", "Select"))
 			]
@@ -380,7 +728,7 @@ TSharedRef<SWidget> SFieldSelector::OnGetMenuContent()
 			.AutoWidth()
 			[
 				SNew(SButton)
-				.OnClicked(this, &SFieldSelector::OnClearBinding)
+				.OnClicked(this, &SFieldSelector::OnClearClicked)
 				.IsEnabled(this, &SFieldSelector::IsClearEnabled)
 				.HAlign(HAlign_Center)
 				[
@@ -393,7 +741,7 @@ TSharedRef<SWidget> SFieldSelector::OnGetMenuContent()
 			.AutoWidth()
 			[
 				SNew(SButton)
-				.OnClicked(this, &SFieldSelector::OnCancel)
+				.OnClicked(this, &SFieldSelector::OnCancelClicked)
 				.HAlign(HAlign_Center)
 				[
 					SNew(STextBlock)
@@ -415,6 +763,7 @@ TSharedRef<SWidget> SFieldSelector::OnGetMenuContent()
 			]
 		];
 
+	ComboButton->SetMenuContentWidgetToFocus(SearchBox);
 	return MenuWidget;
 }
 
@@ -459,6 +808,230 @@ EFieldVisibility SFieldSelector::GetFieldVisibilityFlags() const
 	return Flags;
 }
 
+TSharedPtr<SFieldSelector::FConversionFunctionItem> SFieldSelector::FindOrCreateItemForCategory(TArray<TSharedPtr<FConversionFunctionItem>>& Items, const FString& Name)
+{
+	int32 Idx = 0;
+	for (; Idx < Items.Num(); ++Idx)
+	{
+		// found item
+		if (Items[Idx]->CategoryName == Name)
+		{
+			return Items[Idx];
+		}
+
+		// passed the place where it should have been, break out
+		if (Items[Idx]->CategoryName > Name)
+		{
+			break;
+		}
+	}
+
+	TSharedPtr<FConversionFunctionItem> NewItem = Items.Insert_GetRef(MakeShared<FConversionFunctionItem>(), Idx);
+	NewItem->CategoryName = Name;
+	return NewItem;
+}
+
+TSharedRef<ITableRow> SFieldSelector::GenerateConversionFunctionCategoryRow(TSharedPtr<FConversionFunctionItem> Item, const TSharedRef<STableViewBase>& OwnerTable)
+{
+	return SNew(STableRow<TSharedPtr<FConversionFunctionItem>>, OwnerTable)
+		[
+			SNew(SHorizontalBox)
+			+ SHorizontalBox::Slot()
+			.HAlign(HAlign_Left)
+			.VAlign(VAlign_Center)
+			.Padding(0, 2.0f, 4.0f, 2.0f)
+			.AutoWidth()
+			[
+				SNew(SImage)
+				.DesiredSizeOverride(FVector2D(16.0f, 16.0f))
+				.Image(FAppStyle::Get().GetBrush("GraphEditor.Function_16x"))
+			]
+			+ SHorizontalBox::Slot()
+			.HAlign(HAlign_Fill)
+			.VAlign(VAlign_Center)
+			[
+				SNew(STextBlock)
+				.Font(Item == ConversionFunctionCategoryRoot[0] ? FAppStyle::Get().GetFontStyle("NormalText") : FAppStyle::Get().GetFontStyle("BoldFont"))
+				.Text(FText::FromString(Item->CategoryName))
+				.HighlightText_Lambda([this]() { return SearchBox.IsValid() ? SearchBox->GetText() : FText::GetEmpty(); })
+			]
+		];
+}
+
+void SFieldSelector::GetConversionFunctionCategoryChildren(TSharedPtr<FConversionFunctionItem> Item, TArray<TSharedPtr<FConversionFunctionItem>>& OutItems) const
+{
+	if (ConversionFunctionCategoryRoot.Num() > 0)
+	{
+		if (Item == ConversionFunctionCategoryRoot[0])
+		{
+			Algo::TransformIf(FilteredConversionFunctionCategories, OutItems,
+				[](const TSharedPtr<FConversionFunctionItem>& Item)
+				{
+					return !Item->CategoryName.IsEmpty();
+				},
+				[](const TSharedPtr<FConversionFunctionItem>& Item)
+				{
+					return Item;
+				});
+			return;
+		}
+	}
+
+	Algo::TransformIf(Item->Children, OutItems,
+		[](const TSharedPtr<FConversionFunctionItem>& Item)
+		{
+			return !Item->CategoryName.IsEmpty();
+		},
+		[](const TSharedPtr<FConversionFunctionItem>& Item)
+		{
+			return Item;
+		});
+}
+
+TSharedRef<ITableRow> SFieldSelector::GenerateConversionFunctionRow(const UFunction* Function, const TSharedRef<STableViewBase>& OwnerTable)
+{
+	return SNew(STableRow<const UFunction*>, OwnerTable)
+		[
+			SNew(SHorizontalBox)
+			+ SHorizontalBox::Slot()
+			.HAlign(HAlign_Left)
+			.VAlign(VAlign_Center)
+			.Padding(0, 0, 4.0f, 0)
+			.AutoWidth()
+			[
+				SNew(SImage)
+				.DesiredSizeOverride(FVector2D(16.0f, 16.0f))
+				.Image(FAppStyle::Get().GetBrush("GraphEditor.Function_16x"))
+			]
+			+ SHorizontalBox::Slot()
+			.HAlign(HAlign_Fill)
+			.VAlign(VAlign_Center)
+			[
+				SNew(STextBlock)
+				.Text(Function->GetDisplayNameText())
+				.HighlightText_Lambda([this]() { return SearchBox.IsValid() ? SearchBox->GetText() : FText::GetEmpty(); })
+			]
+		];
+}
+
+void SFieldSelector::AddConversionFunctionChildrenRecursive(const TSharedPtr<FConversionFunctionItem>& Parent, TArray<const UFunction*>& OutFunctions)
+{
+	for (const TSharedPtr<FConversionFunctionItem>& Item : Parent->Children)
+	{
+		if (Item->Function != nullptr)
+		{
+			int32 Index = 0;
+			for (; Index < OutFunctions.Num(); ++Index)
+			{
+				if (OutFunctions[Index]->GetName() > Item->Function->GetName())
+				{
+					break;
+				}
+			}
+
+			OutFunctions.Insert(Item->Function, Index);
+		}
+		else
+		{
+			AddConversionFunctionChildrenRecursive(Item, OutFunctions);
+		}
+	}
+}
+
+void SFieldSelector::FilterConversionFunctions()
+{
+	TArray<FString> FilterStrings;
+	if (SearchBox.IsValid())
+	{
+		SearchBox->GetText().ToString().ParseIntoArrayWS(FilterStrings);
+	}
+
+	if (FilterStrings.IsEmpty())
+	{
+		FilteredConversionFunctions = ConversionFunctions;
+		return;
+	}
+	
+	for (const UFunction* Function : ConversionFunctions)
+	{
+		const FString FunctionName = Function->GetName();
+
+		bool bMatches = true;
+		for (const FString& Filter : FilterStrings)
+		{
+			if (!FunctionName.Contains(Filter))
+			{
+				bMatches = false;
+				break;
+			}
+		}
+
+		if (bMatches)
+		{
+			FilteredConversionFunctions.Add(Function);
+		}
+	}
+}
+
+void SFieldSelector::OnConversionFunctionCategorySelected(TSharedPtr<FConversionFunctionItem> SelectedItem, ESelectInfo::Type)
+{
+	ConversionFunctions.Reset();
+	FilteredConversionFunctions.Reset();
+
+	if (BindingList.IsValid())
+	{
+		BindingList->Clear();
+	}
+
+	if (ViewModelList.IsValid())
+	{
+		ViewModelList->ClearSelection();
+	}
+
+	if (WidgetList.IsValid())
+	{
+		WidgetList->ClearSelection();
+	}
+
+	TArray<TSharedPtr<FConversionFunctionItem>> SelectedItems = ConversionFunctionCategoryTree->GetSelectedItems();
+	if (SelectedItems.Num() == 0)
+	{
+		return;
+	}
+
+	if (SelectedItems.Contains(ConversionFunctionCategoryRoot[0]))
+	{ 
+		for (const TSharedPtr<FConversionFunctionItem>& Item : FilteredConversionFunctionCategories)
+		{
+			AddConversionFunctionChildrenRecursive(Item, ConversionFunctions);
+		}
+		return;
+	}
+
+	for (const TSharedPtr<FConversionFunctionItem>& Item : SelectedItems)
+	{
+		AddConversionFunctionChildrenRecursive(Item, ConversionFunctions);
+	}
+
+	FilterConversionFunctions();
+
+	if (ConversionFunctionList.IsValid())
+	{
+		ConversionFunctionList->RequestListRefresh();
+	}
+}
+
+void SFieldSelector::SetConversionFunctionSelection(const UFunction* SelectedFunction)
+{
+	if (ComboButton.IsValid())
+	{
+		ComboButton->SetIsOpen(false);
+	}
+
+	OnConversionFunctionSelectionChangedDelegate.ExecuteIfBound(SelectedFunction);
+	Refresh();
+}
+
 } // namespace UE::MVVM
 
 #undef LOCTEXT_NAMESPACE
diff --git a/Engine/Plugins/Runtime/ModelViewViewModel/Source/ModelViewViewModelEditor/Private/Widgets/SMVVMFieldSelector.h b/Engine/Plugins/Runtime/ModelViewViewModel/Source/ModelViewViewModelEditor/Private/Widgets/SMVVMFieldSelector.h
--- a/Engine/Plugins/Runtime/ModelViewViewModel/Source/ModelViewViewModelEditor/Private/Widgets/SMVVMFieldSelector.h
+++ b/Engine/Plugins/Runtime/ModelViewViewModel/Source/ModelViewViewModelEditor/Private/Widgets/SMVVMFieldSelector.h
@@ -2,6 +2,7 @@
 
 #pragma once
 
+#include "Misc/TVariant.h"
 #include "MVVMPropertyPath.h"
 #include "Styling/CoreStyle.h"
 #include "Templates/ValueOrError.h"
@@ -15,10 +16,12 @@
 #include "Widgets/SMVVMFieldEntry.h"
 #include "Widgets/SMVVMViewModelBindingListWidget.h"
 #include "Widgets/Text/STextBlock.h"
+#include "Widgets/Views/STreeView.h"
 
 class SComboButton;
 class SSearchBox;
 class SReadOnlyHierarchyView;
+class STableViewBase;
 
 namespace UE::MVVM
 {
@@ -29,15 +32,24 @@ class SSourceBindingList;
 class SFieldSelector : public SCompoundWidget
 {
 public:
-	DECLARE_DELEGATE_OneParam(FOnSelectionChanged, const FMVVMBlueprintPropertyPath&);
+	DECLARE_DELEGATE_OneParam(FOnFieldSelectionChanged, FMVVMBlueprintPropertyPath);
+	DECLARE_DELEGATE_OneParam(FOnConversionFunctionSelectionChanged, const UFunction*);
 
 	SLATE_BEGIN_ARGS(SFieldSelector) :
-		_TextStyle(&FCoreStyle::Get().GetWidgetStyle<FTextBlockStyle>("NormalText"))
+		_TextStyle(&FCoreStyle::Get().GetWidgetStyle<FTextBlockStyle>("NormalText")),
+		_ShowConversionFunctions(false)
 		{
 		}
 		SLATE_STYLE_ARGUMENT(FTextBlockStyle, TextStyle)
 		SLATE_ATTRIBUTE(FMVVMBlueprintPropertyPath, SelectedField)
-		SLATE_EVENT(FOnSelectionChanged, OnSelectionChanged)
+		
+		/** Handler for when a field is selected. Always required. */
+		SLATE_EVENT(FOnFieldSelectionChanged, OnFieldSelectionChanged)
+
+		SLATE_ATTRIBUTE(const UFunction*, SelectedConversionFunction)
+		
+		/** Handler for when a conversion function is selected. Required if ShowConversionFunctions is set. */
+		SLATE_EVENT(FOnConversionFunctionSelectionChanged, OnConversionFunctionSelectionChanged)
 		
 		/**
 		 * Get the binding mode for this field. 
@@ -50,14 +62,17 @@ public:
 		 * If this is set to false, then the Source attribute must be set.
 		 */
 		SLATE_ARGUMENT_DEFAULT(bool, ShowSource) = true;
+		
+		/** The source to use. Only used if ShowSource is false. */
+		SLATE_ARGUMENT(FBindingSource, Source)
 
 		/**
 		 * Only show properties assignable to the given property.
 		 */
-		SLATE_ARGUMENT_DEFAULT(const FProperty*, AssignableTo) = nullptr;
+		SLATE_ATTRIBUTE(const FProperty*, AssignableTo);
 
-		/** The source to use. Only used if ShowSource is false. */
-		SLATE_ARGUMENT(FBindingSource, Source)
+		/** Should we show all conversion functions with results compatible with the AssignableTo property? */
+		SLATE_ATTRIBUTE(bool, ShowConversionFunctions);
 	SLATE_END_ARGS()
 
 	void Construct(const FArguments& InArgs, const UWidgetBlueprint* InWidgetBlueprint, bool bInViewModelProperty);
@@ -65,40 +80,65 @@ public:
 
 private:
 	TSharedRef<SWidget> OnGenerateFieldWidget(FMVVMBlueprintPropertyPath Path) const;
-	TSharedRef<SWidget> OnGetMenuContent();
 
-	TValueOrError<bool, FString> ValidateField(FMVVMBlueprintPropertyPath Field) const;
+	TSharedRef<SWidget> CreateSourcePanel();
+	TSharedRef<SWidget> OnGetMenuContent();
 
 	bool IsClearEnabled() const;
-	FReply OnClearBinding();
+	FReply OnClearClicked();
 
 	bool IsSelectEnabled() const;
-	FReply OnSelectProperty();
+	FReply OnSelectClicked();
 
-	FReply OnCancel();
+	FReply OnCancelClicked();
 
-	void SetSelection(const FMVVMBlueprintPropertyPath& SelectedPath);
+	void SetPropertySelection(const FMVVMBlueprintPropertyPath& SelectedPath);
 	EFieldVisibility GetFieldVisibilityFlags() const;
 
 	void OnViewModelSelected(FBindingSource ViewModel, ESelectInfo::Type);
-	TSharedRef<ITableRow> GenerateRowForViewModel(MVVM::FBindingSource ViewModel, const TSharedRef<STableViewBase>& OwnerTable) const;
+	TSharedRef<ITableRow> GenerateViewModelRow(MVVM::FBindingSource ViewModel, const TSharedRef<STableViewBase>& OwnerTable) const;
+	void FilterViewModels();
 
 	void OnWidgetSelected(FName WidgetName, ESelectInfo::Type);
 
-	void OnSearchTextChanged(const FText& NewText);
+	void OnSearchBoxTextChanged(const FText& NewText);
+
+	struct FConversionFunctionItem
+	{
+		FString CategoryName;
+		const UFunction* Function = nullptr;
+		TArray<TSharedPtr<FConversionFunctionItem>> Children;
+	};
+
+	TSharedPtr<FConversionFunctionItem> FindOrCreateItemForCategory(TArray<TSharedPtr<FConversionFunctionItem>>& Items, const FString& Name);
+	TSharedRef<ITableRow> GenerateConversionFunctionCategoryRow(TSharedPtr<FConversionFunctionItem> Item, const TSharedRef<STableViewBase>& OwnerTable);
+	void FilterConversionFunctionCategories();
+	/** Recursively filter the items in SourceArray and place them into DestArray. Returns true if any items were added. */
+	bool FilterConversionFunctionCategoryChildren(const TArray<FString>& FilterStrings, const TArray<TSharedPtr<FConversionFunctionItem>>& SourceArray, TArray<TSharedPtr<FConversionFunctionItem>>& OutDestArray);
+
+	void FilterConversionFunctions();
+
+	void OnConversionFunctionCategorySelected(TSharedPtr<FConversionFunctionItem> Item, ESelectInfo::Type);
+	void GetConversionFunctionCategoryChildren(TSharedPtr<FConversionFunctionItem> Item, TArray<TSharedPtr<FConversionFunctionItem>>& OutItems) const;
+	
+	void AddConversionFunctionChildrenRecursive(const TSharedPtr<FConversionFunctionItem>& Parent, TArray<const UFunction*>& OutFunctions);
+
+	void SetConversionFunctionSelection(const UFunction* SelectedFunction);
+
+	TSharedRef<ITableRow> GenerateConversionFunctionRow(const UFunction* Function, const TSharedRef<STableViewBase>& OwnerTable);
 
 private:
-	TAttribute<FMVVMBlueprintPropertyPath> SelectedField;
 	TAttribute<EMVVMBindingMode> BindingMode;
 
+	TAttribute<FMVVMBlueprintPropertyPath> SelectedField;
+	FMVVMBlueprintPropertyPath CachedSelectedField;
+
 	TSharedPtr<SFieldEntry> SelectedEntryWidget;
 	TSharedPtr<SSourceEntry> SelectedSourceWidget;
 
 	const FTextBlockStyle* TextStyle = nullptr;
 
-	FOnSelectionChanged OnSelectionChangedDelegate;
-
-	FMVVMBlueprintPropertyPath CachedSelectedField;
+	FOnFieldSelectionChanged OnFieldSelectionChangedDelegate;
 
 	TSharedPtr<SSourceBindingList> BindingList;
 	TSharedPtr<SSearchBox> SearchBox;
@@ -110,10 +150,29 @@ private:
 	bool bViewModelProperty = false;
 
 	TArray<FBindingSource> ViewModelSources;
+	TArray<FBindingSource> FilteredViewModelSources;
+
 	TSharedPtr<SListView<FBindingSource>> ViewModelList;
 	TSharedPtr<SReadOnlyHierarchyView> WidgetList;
 
-	const FProperty* AssignableTo = nullptr;
+	TAttribute<const FProperty*> AssignableTo;
+	TAttribute<bool> ShowConversionFunctions;
+
+	TAttribute<const UFunction*> SelectedConversionFunction;
+	const UFunction* CachedSelectedConversionFunction = nullptr;
+
+	FOnConversionFunctionSelectionChanged OnConversionFunctionSelectionChangedDelegate;
+
+	TSharedPtr<STreeView<TSharedPtr<FConversionFunctionItem>>> ConversionFunctionCategoryTree;
+	TArray<TSharedPtr<FConversionFunctionItem>> ConversionFunctionCategories;
+	TArray<TSharedPtr<FConversionFunctionItem>> FilteredConversionFunctionCategories;
+	TArray<TSharedPtr<FConversionFunctionItem>> ConversionFunctionCategoryRoot;
+
+	TSharedPtr<SListView<const UFunction*>> ConversionFunctionList;
+
+	TArray<const UFunction*> ConversionFunctions;
+	TArray<const UFunction*> FilteredConversionFunctions;
+
 }; 
 
 } // namespace UE::MVVM
\ No newline at end of file
diff --git a/Engine/Plugins/Runtime/ModelViewViewModel/Source/ModelViewViewModelEditor/Private/Widgets/SMVVMFunctionParameter.cpp b/Engine/Plugins/Runtime/ModelViewViewModel/Source/ModelViewViewModelEditor/Private/Widgets/SMVVMFunctionParameter.cpp
--- a/Engine/Plugins/Runtime/ModelViewViewModel/Source/ModelViewViewModelEditor/Private/Widgets/SMVVMFunctionParameter.cpp
+++ b/Engine/Plugins/Runtime/ModelViewViewModel/Source/ModelViewViewModelEditor/Private/Widgets/SMVVMFunctionParameter.cpp
@@ -94,7 +94,7 @@ void SFunctionParameter::Construct(const FArguments& InArgs)
 					.Visibility(this, &SFunctionParameter::OnGetVisibility, false)
 					.SelectedField(this, &SFunctionParameter::OnGetSelectedField)
 					.BindingMode_Lambda([this]() { return GetBindingModeDelegate.Execute(); })
-					.OnSelectionChanged(this, &SFunctionParameter::OnFieldSelectionChanged)
+					.OnFieldSelectionChanged(this, &SFunctionParameter::OnFieldSelectionChanged)
 					.AssignableTo(Property)
 				]
 			]
@@ -172,7 +172,7 @@ FMVVMBlueprintPropertyPath SFunctionParameter::OnGetSelectedField() const
 	return EditorSubsystem->GetPathForConversionFunctionArgument(WidgetBlueprint.Get(), *Binding, ParameterName, bSourceToDestination);
 }
 
-void SFunctionParameter::OnFieldSelectionChanged(const FMVVMBlueprintPropertyPath& Selected)
+void SFunctionParameter::OnFieldSelectionChanged(FMVVMBlueprintPropertyPath Selected)
 {
 	UMVVMEditorSubsystem* EditorSubsystem = GEditor->GetEditorSubsystem<UMVVMEditorSubsystem>();
 	EditorSubsystem->SetPathForConversionFunctionArgument(WidgetBlueprint.Get(), *Binding, ParameterName, Selected, bSourceToDestination);
diff --git a/Engine/Plugins/Runtime/ModelViewViewModel/Source/ModelViewViewModelEditor/Private/Widgets/SMVVMFunctionParameter.h b/Engine/Plugins/Runtime/ModelViewViewModel/Source/ModelViewViewModelEditor/Private/Widgets/SMVVMFunctionParameter.h
--- a/Engine/Plugins/Runtime/ModelViewViewModel/Source/ModelViewViewModelEditor/Private/Widgets/SMVVMFunctionParameter.h
+++ b/Engine/Plugins/Runtime/ModelViewViewModel/Source/ModelViewViewModelEditor/Private/Widgets/SMVVMFunctionParameter.h
@@ -44,7 +44,7 @@ private:
 	void OnBindArgumentChecked(ECheckBoxState Checked);
 	EVisibility OnGetVisibility(bool bDefaultValue) const;
 	FMVVMBlueprintPropertyPath OnGetSelectedField() const;
-	void OnFieldSelectionChanged(const FMVVMBlueprintPropertyPath& Selected);
+	void OnFieldSelectionChanged(FMVVMBlueprintPropertyPath Selected);
 
 private:
 
diff --git a/Engine/Plugins/Runtime/ModelViewViewModel/Source/ModelViewViewModelEditor/Private/Widgets/SMVVMSourceSelector.cpp b/Engine/Plugins/Runtime/ModelViewViewModel/Source/ModelViewViewModelEditor/Private/Widgets/SMVVMSourceSelector.cpp
--- a/Engine/Plugins/Runtime/ModelViewViewModel/Source/ModelViewViewModelEditor/Private/Widgets/SMVVMSourceSelector.cpp
+++ b/Engine/Plugins/Runtime/ModelViewViewModel/Source/ModelViewViewModelEditor/Private/Widgets/SMVVMSourceSelector.cpp
@@ -34,43 +34,49 @@ void SSourceSelector::Construct(const FArguments& Args, const UWidgetBlueprint*
 	
 	Refresh();
 
-	TSharedRef<SHorizontalBox> HorizontalBox = SNew(SHorizontalBox)
-		+ SHorizontalBox::Slot()
+	MenuAnchor = SNew(SMenuAnchor)
+		.OnGetMenuContent(this, &SSourceSelector::OnGetMenuContent)
 		[
-			SAssignNew(MenuAnchor, SMenuAnchor)
-			.OnGetMenuContent(this, &SSourceSelector::OnGetMenuContent)
+			SNew(SButton)
+			.ButtonStyle(&FAppStyle::Get().GetWidgetStyle<FComboButtonStyle>("ComboButton").ButtonStyle)
+			.OnClicked_Lambda([this]() 
+			{
+				MenuAnchor->SetIsOpen(true); 
+				return FReply::Handled(); 
+			})
 			[
-				SNew(SButton)
-				.ButtonStyle(&FAppStyle::Get().GetWidgetStyle<FComboButtonStyle>("ComboButton").ButtonStyle)
-				.OnClicked_Lambda([this]() 
-				{
-					MenuAnchor->SetIsOpen(true); 
-					return FReply::Handled(); 
-				})
+				SNew(SHorizontalBox)
+				+ SHorizontalBox::Slot()
+				.HAlign(HAlign_Fill)
+				[
+					SAssignNew(SelectedSourceWidget, SSourceEntry)
+					.Source(SelectedSource)
+					.TextStyle(TextStyle)
+				]
+				+ SHorizontalBox::Slot()
+				.Padding(2, 0)
+				.HAlign(HAlign_Right)
+				.VAlign(VAlign_Center)
+				.AutoWidth()
 				[
-					SNew(SHorizontalBox)
-					+ SHorizontalBox::Slot()
-					[
-						SAssignNew(SelectedSourceWidget, SSourceEntry)
-						.Source(SelectedSource)
-						.TextStyle(TextStyle)
-					]
-					+ SHorizontalBox::Slot()
-					.Padding(2, 0)
-					.HAlign(HAlign_Right)
-					.VAlign(VAlign_Center)
-					.AutoWidth()
-					[
-						SNew(SImage)
-						.Image(FAppStyle::Get().GetBrush("Icons.ChevronDown"))
-					]
+					SNew(SImage)
+					.Image(FAppStyle::Get().GetBrush("Icons.ChevronDown"))
 				]
 			]
 		];
+		
 
 	if (bShowClear)
 	{
-		HorizontalBox->AddSlot()
+		ChildSlot
+		[
+			SNew(SHorizontalBox)
+			+ SHorizontalBox::Slot()
+			.HAlign(HAlign_Fill)
+			[
+				MenuAnchor.ToSharedRef()
+			]
+			+ SHorizontalBox::Slot()
 			.HAlign(HAlign_Right)
 			.VAlign(VAlign_Center)
 			.AutoWidth()
@@ -80,13 +86,16 @@ void SSourceSelector::Construct(const FArguments& Args, const UWidgetBlueprint*
 				.ToolTipText(LOCTEXT("ClearField", "Clear source selection."))
 				.Visibility(this, &SSourceSelector::GetClearVisibility)
 				.OnClicked(this, &SSourceSelector::OnClearSource)
-			];
+			]
+		];
+	}
+	else
+	{
+		ChildSlot
+		[
+			MenuAnchor.ToSharedRef()
+		];
 	}
-
-	ChildSlot
-	[
-		HorizontalBox
-	];
 
 	OnSelectionChangedDelegate = Args._OnSelectionChanged;
 	check(OnSelectionChangedDelegate.IsBound());
diff --git a/Engine/Plugins/Runtime/ModelViewViewModel/Source/ModelViewViewModelEditor/Private/Widgets/SMVVMViewBindingListView.cpp b/Engine/Plugins/Runtime/ModelViewViewModel/Source/ModelViewViewModelEditor/Private/Widgets/SMVVMViewBindingListView.cpp
--- a/Engine/Plugins/Runtime/ModelViewViewModel/Source/ModelViewViewModelEditor/Private/Widgets/SMVVMViewBindingListView.cpp
+++ b/Engine/Plugins/Runtime/ModelViewViewModel/Source/ModelViewViewModelEditor/Private/Widgets/SMVVMViewBindingListView.cpp
@@ -178,15 +178,15 @@ public:
 			STableRow<TSharedPtr<FBindingEntry>>::FArguments()
 			.Padding(1.0f)
 			[
-				SNew(SBox)
-				.HeightOverride(30)
-				.MinDesiredWidth(200)
+				SNew(SHorizontalBox)
+				+ SHorizontalBox::Slot()
+				.Padding(2, 1)
 				.HAlign(HAlign_Left)
+				.VAlign(VAlign_Fill)
+				.AutoWidth()
 				[
-					SNew(SHorizontalBox)
-					+ SHorizontalBox::Slot()
-					.Padding(2, 1)
-					.VAlign(VAlign_Fill)
+					SNew(SBox)
+					.MinDesiredWidth(150)
 					[
 						SNew(SSourceSelector, InWidgetBlueprint)
 						.ShowClear(false)
@@ -195,17 +195,17 @@ public:
 						.SelectedSource(this, &SWidgetRow::GetSelectedWidget)
 						.OnSelectionChanged(this, &SWidgetRow::SetSelectedWidget)
 					]
-					+ SHorizontalBox::Slot()
-					.HAlign(HAlign_Left)
-					.VAlign(VAlign_Center)
-					.AutoWidth()
-					[
-						SNew(SSimpleButton)
-						.Icon(FAppStyle::Get().GetBrush("Icons.Plus"))
-						.IsEnabled_Lambda([this]() { return !Entry->GetName().IsNone(); })
-						.OnClicked(this, &SWidgetRow::AddBinding)
+				]
+				+ SHorizontalBox::Slot()
+				.HAlign(HAlign_Left)
+				.VAlign(VAlign_Center)
+				.AutoWidth()
+				[
+					SNew(SSimpleButton)
+					.Icon(FAppStyle::Get().GetBrush("Icons.Plus"))
+					.IsEnabled_Lambda([this]() { return !Entry->GetName().IsNone(); })
+					.OnClicked(this, &SWidgetRow::AddBinding)
 
-					]
 				]
 			],
 			OwnerTableView
@@ -311,14 +311,18 @@ public:
 				.AutoWidth()
 				[
 					SNew(SBox)
-					.MinDesiredWidth(200)
+					.MinDesiredWidth(150)
 					[
 						SNew(SFieldSelector, InWidgetBlueprint, false)
-						.SelectedField(this, &SBindingRow::GetSelectedWidgetProperty)
-						.OnSelectionChanged(this, &SBindingRow::OnWidgetPropertySelected)
 						.BindingMode(this, &SBindingRow::GetCurrentBindingMode)
+						.SelectedField(this, &SBindingRow::GetSelectedWidgetProperty)
+						.OnFieldSelectionChanged(this, &SBindingRow::OnWidgetPropertySelected)
+						.SelectedConversionFunction(this, &SBindingRow::GetSelectedConversionFunction, false)
+						.OnConversionFunctionSelectionChanged(this, &SBindingRow::OnConversionFunctionChanged, false)
+						.ShowConversionFunctions(this, &SBindingRow::ShouldShowConversionFunctions, false)
 						.ShowSource(false)
 						.Source(WidgetSource)
+						.AssignableTo(this, &SBindingRow::GetAssignableToProperty, false)
 					]
 				]
 
@@ -354,12 +358,16 @@ public:
 				.AutoWidth()
 				[
 					SNew(SBox)
-					.MinDesiredWidth(200)
+					.MinDesiredWidth(150)
 					[
 						SAssignNew(ViewModelFieldSelector, SFieldSelector, InWidgetBlueprint, true)
-						.SelectedField(this, &SBindingRow::GetSelectedViewModelProperty)
 						.BindingMode(this, &SBindingRow::GetCurrentBindingMode)
-						.OnSelectionChanged(this, &SBindingRow::OnViewModelPropertySelected)
+						.SelectedField(this, &SBindingRow::GetSelectedViewModelProperty)
+						.OnFieldSelectionChanged(this, &SBindingRow::OnViewModelPropertySelected)
+						.SelectedConversionFunction(this, &SBindingRow::GetSelectedConversionFunction, true)
+						.OnConversionFunctionSelectionChanged(this, &SBindingRow::OnConversionFunctionChanged, true)
+						.ShowConversionFunctions(this, &SBindingRow::ShouldShowConversionFunctions, true)
+						.AssignableTo(this, &SBindingRow::GetAssignableToProperty, true)
 					]
 				]
 
@@ -379,28 +387,6 @@ public:
 					.OnEnumSelectionChanged(this, &SBindingRow::OnUpdateModeSelectionChanged)
 					.CurrentValue(this, &SBindingRow::GetUpdateModeValue)
 				]
-		
-				+ SHorizontalBox::Slot()
-				.Padding(2, 0)
-				.VAlign(VAlign_Center)
-				.HAlign(HAlign_Right)
-				.AutoWidth()
-				[
-					SNew(SMVVMConversionPath, InWidgetBlueprint, false)
-					.Bindings(this, &SBindingRow::GetThisViewBindingAsArray)
-					.OnFunctionChanged(this, &SBindingRow::OnConversionFunctionChanged, false)
-				]
-
-				+ SHorizontalBox::Slot()
-				.Padding(2, 0)
-				.VAlign(VAlign_Center)
-				.HAlign(HAlign_Right)
-				.AutoWidth()
-				[
-					SNew(SMVVMConversionPath, InWidgetBlueprint, true)
-					.Bindings(this, &SBindingRow::GetThisViewBindingAsArray)
-					.OnFunctionChanged(this, &SBindingRow::OnConversionFunctionChanged, true)
-				]
 			],
 			OwnerTableView
 		);
@@ -499,8 +485,61 @@ private:
 
 	EMVVMBindingMode GetCurrentBindingMode() const
 	{
-		const FMVVMBlueprintViewBinding* ViewModelBinding = GetThisViewBinding();
-		return ViewModelBinding->BindingType;
+		if (const FMVVMBlueprintViewBinding* ViewModelBinding = GetThisViewBinding())
+		{
+			return ViewModelBinding->BindingType;
+		}
+		return EMVVMBindingMode::OneWayToDestination;
+	}
+
+	bool ShouldShowConversionFunctions(bool bViewModel) const
+	{
+		EMVVMBindingMode Mode = GetCurrentBindingMode();
+		if (IsForwardBinding(Mode))
+		{
+			return bViewModel;
+		}
+		else if (IsBackwardBinding(Mode))
+		{
+			return !bViewModel;
+		}
+		
+		return false;
+	}
+
+	const FProperty* GetAssignableToProperty(bool bViewModel) const
+	{
+		if (FMVVMBlueprintViewBinding* ViewBinding = GetThisViewBinding())
+		{
+			TArray<FMVVMConstFieldVariant> Fields;
+			if (bViewModel)
+			{
+				Fields = ViewBinding->WidgetPath.GetFields();
+			}
+			else
+			{
+				Fields = ViewBinding->ViewModelPath.GetFields();
+			}
+
+			if (Fields.Num() > 0)
+			{
+				FMVVMConstFieldVariant LastField = Fields.Last();
+				if (LastField.IsProperty())
+				{
+					return LastField.GetProperty();
+				}
+				else if (LastField.IsFunction())
+				{
+					const UFunction* Function = LastField.GetFunction();
+					if (Function != nullptr)
+					{
+						return BindingHelper::GetFirstArgumentProperty(Function);
+					}
+				}
+			}
+		}
+
+		return nullptr;
 	}
 
 	TSharedRef<ITableRow> OnGenerateErrorRow(TSharedPtr<FText> Text, const TSharedRef<STableViewBase>& TableView) const
@@ -555,7 +594,18 @@ private:
 		return FMVVMBlueprintPropertyPath();
 	}
 
-	void OnViewModelPropertySelected(const FMVVMBlueprintPropertyPath& SelectedField)
+	const UFunction* GetSelectedConversionFunction(bool bSourceToDest) const
+	{
+		if (FMVVMBlueprintViewBinding* ViewBinding = GetThisViewBinding())
+		{
+			UMVVMEditorSubsystem* EditorSubsystem = GEditor->GetEditorSubsystem<UMVVMEditorSubsystem>();
+			return EditorSubsystem->GetConversionFunction(WidgetBlueprintWeak.Get(), *ViewBinding, bSourceToDest);
+		}
+
+		return nullptr;
+	}
+
+	void OnViewModelPropertySelected(FMVVMBlueprintPropertyPath SelectedField)
 	{
 		if (FMVVMBlueprintViewBinding* ViewBinding = GetThisViewBinding())
 		{
@@ -567,7 +617,7 @@ private:
 		}
 	}
 
-	void OnWidgetPropertySelected(const FMVVMBlueprintPropertyPath& SelectedField)
+	void OnWidgetPropertySelected(FMVVMBlueprintPropertyPath SelectedField)
 	{
 		if (FMVVMBlueprintViewBinding* ViewBinding = GetThisViewBinding())
 		{
@@ -768,17 +818,36 @@ class SFunctionParameterRow : public STableRow<TSharedPtr<FBindingEntry>>
 		UMVVMEditorSubsystem* EditorSubsystem = GEditor->GetEditorSubsystem<UMVVMEditorSubsystem>();
 		UMVVMBlueprintView* View = EditorSubsystem->GetView(InWidgetBlueprint);
 		FMVVMBlueprintViewBinding* Binding = Entry->GetBinding(View);
-		UEdGraphPin* Pin = EditorSubsystem->FindConversionFunctionArgumentPin(InWidgetBlueprint, *Binding, Entry->GetName(), UE::MVVM::IsForwardBinding(Binding->BindingType));
+		const bool bSourceToDestination = UE::MVVM::IsForwardBinding(Binding->BindingType);
 
 		FSlateColor PrimaryColor, SecondaryColor;
+		const FSlateBrush* PrimaryBrush = nullptr;
 		const FSlateBrush* SecondaryBrush = nullptr;
-		const FSlateBrush* PrimaryBrush = FBlueprintEditor::GetVarIconAndColorFromPinType(Pin->PinType, PrimaryColor, SecondaryBrush, SecondaryColor);
+		FText DisplayName, ToolTip;
+		
+		if (UEdGraphPin* Pin = EditorSubsystem->FindConversionFunctionArgumentPin(InWidgetBlueprint, *Binding, Entry->GetName(), bSourceToDestination))
+		{
+			PrimaryBrush = FBlueprintEditor::GetVarIconAndColorFromPinType(Pin->PinType, PrimaryColor, SecondaryBrush, SecondaryColor);
+			DisplayName = Pin->GetDisplayName();
+			ToolTip = FText::FromString(Pin->PinToolTip);
+		}
+		else if (const UFunction* Function = EditorSubsystem->GetConversionFunction(InWidgetBlueprint, *Binding, bSourceToDestination))
+		{
+			// no wrapper graph, this is a simple conversion function of the form: int32 Convert(float x)
+			if (const FProperty* Argument = BindingHelper::GetFirstArgumentProperty(Function))
+			{
+				PrimaryBrush = FBlueprintEditor::GetVarIconAndColorFromProperty(Argument, PrimaryColor, SecondaryBrush, SecondaryColor);
+				DisplayName = Argument->GetDisplayNameText();
+				ToolTip = Argument->GetToolTipText();
+			}
+		}
 
 		STableRow<TSharedPtr<FBindingEntry>>::Construct(
 			STableRow<TSharedPtr<FBindingEntry>>::FArguments()
 		[
 			SNew(SBox)
 			.HeightOverride(30)
+			.ToolTipText(ToolTip)
 			[
 				SNew(SHorizontalBox)
 				+ SHorizontalBox::Slot()
@@ -804,7 +873,7 @@ class SFunctionParameterRow : public STableRow<TSharedPtr<FBindingEntry>>
 						+ SHorizontalBox::Slot()
 						[
 							SNew(STextBlock)
-							.Text(Pin->GetDisplayName())
+							.Text(DisplayName)
 						]
 					]
 				]
@@ -821,7 +890,7 @@ class SFunctionParameterRow : public STableRow<TSharedPtr<FBindingEntry>>
 						.WidgetBlueprint(InWidgetBlueprint)
 						.Binding(Binding)
 						.ParameterName(Entry->GetName())
-						.SourceToDestination(UE::MVVM::IsForwardBinding(Binding->BindingType))
+						.SourceToDestination(bSourceToDestination)
 					]
 				]
 			]
@@ -959,25 +1028,16 @@ void SBindingsList::Refresh()
 
 			// create entries for conversion function parameters
 			UMVVMEditorSubsystem* EditorSubsystem = GEditor->GetEditorSubsystem<UMVVMEditorSubsystem>();
-			UEdGraph* Graph = EditorSubsystem->GetConversionFunctionGraph(MVVMExtensionPtr->GetWidgetBlueprint(), Binding, UE::MVVM::IsForwardBinding(Binding.BindingType));
-			if (Graph != nullptr)
+			const UFunction* Function = EditorSubsystem->GetConversionFunction(MVVMExtensionPtr->GetWidgetBlueprint(), Binding, UE::MVVM::IsForwardBinding(Binding.BindingType));
+			if (Function != nullptr)
 			{
-				TArray<UK2Node_CallFunction*> FunctionNodes;
-				Graph->GetNodesOfClass<UK2Node_CallFunction>(FunctionNodes);
-
-				if (FunctionNodes.Num() == 1)
+				TValueOrError<TArray<const FProperty*>, FString> ArgumentsResult = BindingHelper::TryGetArgumentsForConversionFunction(Function);
+				if (ArgumentsResult.HasValue())
 				{
-					UK2Node_CallFunction* CallFunctionNode = FunctionNodes[0];
-					for (UEdGraphPin* Pin : CallFunctionNode->Pins)
+					for (const FProperty* Argument : ArgumentsResult.GetValue())
 					{
-						// skip output and exec pins, the rest should be arguments
-						if (Pin->Direction != EGPD_Input || Pin->PinName == UEdGraphSchema_K2::PN_Execute || Pin->PinName == UEdGraphSchema_K2::PN_Self)
-						{
-							continue;
-						}
-
 						TSharedRef<FBindingEntry> NewArgumentEntry = MakeShared<FBindingEntry>();
-						NewArgumentEntry->SetParameterName(BindingIndex, Pin->PinName);
+						NewArgumentEntry->SetParameterName(BindingIndex, Argument->GetFName());
 						NewBindingEntry->AddChild(NewArgumentEntry);
 						ReselectIfRequired(NewArgumentEntry);
 					}
diff --git a/Engine/Plugins/Runtime/ModelViewViewModel/Source/ModelViewViewModelEditor/Private/Widgets/SMVVMViewModelBindingListWidget.cpp b/Engine/Plugins/Runtime/ModelViewViewModel/Source/ModelViewViewModelEditor/Private/Widgets/SMVVMViewModelBindingListWidget.cpp
--- a/Engine/Plugins/Runtime/ModelViewViewModel/Source/ModelViewViewModelEditor/Private/Widgets/SMVVMViewModelBindingListWidget.cpp
+++ b/Engine/Plugins/Runtime/ModelViewViewModel/Source/ModelViewViewModelEditor/Private/Widgets/SMVVMViewModelBindingListWidget.cpp
@@ -58,7 +58,7 @@ TArray<FFieldVariant> FFieldIterator_Bindable::GetFields(const UStruct* Struct)
 			FMVVMFieldVariant FieldVariant = BindingHelper::FindFieldByName(Struct, Value.GetBindingName());
 			if (FieldVariant.IsFunction())
 			{
-				if (AssignableTo != nullptr && !BindingHelper::ArePropertiesCompatible(FieldVariant.GetFunction()->GetReturnProperty(), AssignableTo))
+				if (AssignableTo != nullptr && !BindingHelper::ArePropertiesCompatible(BindingHelper::GetReturnProperty(FieldVariant.GetFunction()), AssignableTo))
 				{
 					continue;
 				}
@@ -81,12 +81,13 @@ TArray<FFieldVariant> FFieldIterator_Bindable::GetFields(const UStruct* Struct)
 	{
 		const UWidgetBlueprint* WidgetBlueprintPtr = WidgetBlueprint.Get();
 		TSubclassOf<UObject> AccessorClass = WidgetBlueprintPtr ? WidgetBlueprintPtr->GeneratedClass : nullptr;
-		AddResult(GEngine->GetEngineSubsystem<UMVVMSubsystem>()->GetAvailableBindings(const_cast<UClass*>(Class), AccessorClass));
-
+		TArray<FMVVMAvailableBinding> Bindings = GEngine->GetEngineSubsystem<UMVVMSubsystem>()->GetAvailableBindings(const_cast<UClass*>(Class), AccessorClass);
+		AddResult(Bindings);
 	}
 	else if (const UScriptStruct* ScriptStruct = Cast<const UScriptStruct>(Struct))
 	{
-		AddResult(GEngine->GetEngineSubsystem<UMVVMSubsystem>()->GetAvailableBindingsForStruct(ScriptStruct));
+		TArray<FMVVMAvailableBinding> Bindings = GEngine->GetEngineSubsystem<UMVVMSubsystem>()->GetAvailableBindingsForStruct(ScriptStruct);
+		AddResult(Bindings);
 	}
 
 	Result.Sort([](const FFieldVariant& A, const FFieldVariant& B)
@@ -130,9 +131,7 @@ void SSourceBindingList::Construct(const FArguments& InArgs, const UWidgetBluepr
 
 	OnDoubleClicked = InArgs._OnDoubleClicked;
 
-	ChildSlot
-	[
-		SAssignNew(PropertyViewer, SPropertyViewer)
+	PropertyViewer = SNew(SPropertyViewer)
 		.FieldIterator(FieldIterator.Get())
 		.PropertyVisibility(SPropertyViewer::EPropertyVisibility::Hidden)
 		.bShowFieldIcon(true)
@@ -140,13 +139,17 @@ void SSourceBindingList::Construct(const FArguments& InArgs, const UWidgetBluepr
 		.SelectionMode(InArgs._EnableSelection ? ESelectionMode::Single : ESelectionMode::None)
 		.bShowSearchBox(InArgs._ShowSearchBox)
 		.OnSelectionChanged(this, &SSourceBindingList::HandleSelectionChanged)
-		.OnDoubleClicked(this, &SSourceBindingList::HandleDoubleClicked)
-	];
+		.OnDoubleClicked(this, &SSourceBindingList::HandleDoubleClicked);
 
 	if (InArgs._InitialSource.IsValid())
 	{
 		AddSources(MakeArrayView(&InArgs._InitialSource, 1));
 	}
+
+	ChildSlot
+	[
+		PropertyViewer.ToSharedRef()
+	];
 }
 
 void SSourceBindingList::Clear()
@@ -216,10 +219,13 @@ void SSourceBindingList::AddSource(const FBindingSource& InSource)
 
 void SSourceBindingList::AddSources(TArrayView<const FBindingSource> InSources)
 {
-	Algo::Transform(InSources, Sources, [](const FBindingSource& Source)
+	for (const FBindingSource& Source : InSources)
+	{
+		if (Source.Class != nullptr)
 		{
-			return TPair<FBindingSource, SPropertyViewer::FHandle>(Source, SPropertyViewer::FHandle());
-		});
+			Sources.Add(TPair<FBindingSource, SPropertyViewer::FHandle>(Source, SPropertyViewer::FHandle()));
+		}
+	}
 
 	if (PropertyViewer)
 	{
@@ -257,11 +263,9 @@ FMVVMBlueprintPropertyPath SSourceBindingList::CreateBlueprintPropertyPath(SProp
 		PropertyPath.ResetBasePropertyPath();
 	}
 
-	// TODO (sebastiann): FMVVMBlueprintPropertyPath doesn't support long paths yet
-	FFieldVariant LastField = FieldPath.Num() > 0 ? FieldPath.Last() : FFieldVariant();
-	if (LastField.IsValid())
+	for (const FFieldVariant& Field : FieldPath)
 	{
-		PropertyPath.SetBasePropertyPath(FMVVMConstFieldVariant(LastField));
+		PropertyPath.AppendBasePropertyPath(FMVVMConstFieldVariant(Field));
 	}
 
 	return PropertyPath;
@@ -296,4 +300,4 @@ FMVVMBlueprintPropertyPath SSourceBindingList::GetSelectedProperty() const
 
 } // namespace UE::MVVM
 
-#undef LOCTEXT_NAMESPACE
\ No newline at end of file
+#undef LOCTEXT_NAMESPACE
diff --git a/Engine/Plugins/Runtime/ModelViewViewModel/Source/ModelViewViewModelEditor/Private/Widgets/SMVVMViewModelBindingListWidget.h b/Engine/Plugins/Runtime/ModelViewViewModel/Source/ModelViewViewModelEditor/Private/Widgets/SMVVMViewModelBindingListWidget.h
--- a/Engine/Plugins/Runtime/ModelViewViewModel/Source/ModelViewViewModelEditor/Private/Widgets/SMVVMViewModelBindingListWidget.h
+++ b/Engine/Plugins/Runtime/ModelViewViewModel/Source/ModelViewViewModelEditor/Private/Widgets/SMVVMViewModelBindingListWidget.h
@@ -13,8 +13,8 @@
 #include "Widgets/PropertyViewer/SPropertyViewer.h"
 
 class UWidgetBlueprint;
+
 struct FMVVMBlueprintViewModelContext;
-class UWidgetBlueprint;
 
 namespace UE::MVVM
 {
@@ -57,7 +57,14 @@ namespace UE::MVVM
 			SLATE_ARGUMENT_DEFAULT(bool, ShowSearchBox) = false;
 			SLATE_ARGUMENT_DEFAULT(EFieldVisibility, FieldVisibilityFlags) = EFieldVisibility::All;
 			SLATE_ARGUMENT_DEFAULT(bool, EnableSelection) = true;
+
+			/** 
+			 * Show only properties that are assignable to the given property. 
+			 * This means that they are either the exact same type, or that they are trivially convertible to the target type. 
+			 * Eg. bool -> bool, int32 -> int64, float -> double
+			 */
 			SLATE_ARGUMENT_DEFAULT(const FProperty*, AssignableTo) = nullptr;
+
 			SLATE_EVENT(FOnDoubleClicked, OnDoubleClicked)
 		SLATE_END_ARGS()
 
