MVVM: Add BP binding event to Viewmodel. Fix FieldVariant::IsEmpty now check if the pointer is nullptr (the function or property is not valid anymore). GetAvailableConversionFunctions now check for the widget parent class.

diff --git a/Engine/Plugins/Runtime/ModelViewViewModel/Source/ModelViewViewModel/Private/MVVMViewModelBase.cpp b/Engine/Plugins/Runtime/ModelViewViewModel/Source/ModelViewViewModel/Private/MVVMViewModelBase.cpp
--- a/Engine/Plugins/Runtime/ModelViewViewModel/Source/ModelViewViewModel/Private/MVVMViewModelBase.cpp
+++ b/Engine/Plugins/Runtime/ModelViewViewModel/Source/ModelViewViewModel/Private/MVVMViewModelBase.cpp
@@ -23,6 +23,24 @@ FDelegateHandle UMVVMViewModelBase::AddFieldValueChangedDelegate(UE::FieldNotifi
 }
 
 
+void UMVVMViewModelBase::K2_AddFieldValueChangedDelegate(FFieldNotificationId InFieldId, FFieldValueChangedDynamicDelegate InDelegate)
+{
+	if (InFieldId.IsValid())
+	{
+		const UE::FieldNotification::FFieldId FieldId = GetFieldNotificationDescriptor().GetField(GetClass(), InFieldId.FieldName);
+		if (ensureMsgf(FieldId.IsValid(), TEXT("The field should be compiled correctly.")))
+		{
+			FDelegateHandle Result = Delegates.Add(this, FieldId, InDelegate);
+			if (Result.IsValid())
+			{
+				EnabledFieldNotifications.PadToNum(FieldId.GetIndex() + 1, false);
+				EnabledFieldNotifications[FieldId.GetIndex()] = true;
+			}
+		}
+	}
+}
+
+
 bool UMVVMViewModelBase::RemoveFieldValueChangedDelegate(UE::FieldNotification::FFieldId InFieldId, FDelegateHandle InHandle)
 {
 	bool bResult = false;
@@ -36,6 +54,23 @@ bool UMVVMViewModelBase::RemoveFieldValueChangedDelegate(UE::FieldNotification::
 }
 
 
+void UMVVMViewModelBase::K2_RemoveFieldValueChangedDelegate(FFieldNotificationId InFieldId, FFieldValueChangedDynamicDelegate InDelegate)
+{
+	if (InFieldId.IsValid())
+	{
+		const UE::FieldNotification::FFieldId FieldId = GetFieldNotificationDescriptor().GetField(GetClass(), InFieldId.FieldName);
+		if (ensureMsgf(FieldId.IsValid(), TEXT("The field should be compiled correctly.")))
+		{
+			if (EnabledFieldNotifications.IsValidIndex(FieldId.GetIndex()) && EnabledFieldNotifications[FieldId.GetIndex()])
+			{
+				UE::FieldNotification::FFieldMulticastDelegate::FRemoveFromResult RemoveResult = Delegates.RemoveFrom(this, FieldId, InDelegate);
+				EnabledFieldNotifications[FieldId.GetIndex()] = RemoveResult.bHasOtherBoundDelegates;
+			}
+		}
+	}
+}
+
+
 int32 UMVVMViewModelBase::RemoveAllFieldValueChangedDelegates(const void* InUserObject)
 {
 	int32 bResult = 0;
diff --git a/Engine/Plugins/Runtime/ModelViewViewModel/Source/ModelViewViewModel/Private/View/MVVMView.cpp b/Engine/Plugins/Runtime/ModelViewViewModel/Source/ModelViewViewModel/Private/View/MVVMView.cpp
--- a/Engine/Plugins/Runtime/ModelViewViewModel/Source/ModelViewViewModel/Private/View/MVVMView.cpp
+++ b/Engine/Plugins/Runtime/ModelViewViewModel/Source/ModelViewViewModel/Private/View/MVVMView.cpp
@@ -68,7 +68,7 @@ void UMVVMView::Destruct()
 }
 
 
-void UMVVMView::SetViewModel(FName ViewModelName, UMVVMViewModelBase* NewValue)
+bool UMVVMView::SetViewModel(FName ViewModelName, UMVVMViewModelBase* NewValue)
 {
 	if (!ViewModelName.IsNone() && ClassExtension != nullptr)
 	{
@@ -77,7 +77,7 @@ void UMVVMView::SetViewModel(FName ViewModelName, UMVVMViewModelBase* NewValue)
 		{
 			UE::MVVM::FMessageLog Log(GetUserWidget());
 			Log.Error(LOCTEXT("SetViewModelInvalidName", "Invalid viewmodel name for SetViewModel."));
-			return;
+			return false;
 		}
 
 		{
@@ -91,7 +91,7 @@ void UMVVMView::SetViewModel(FName ViewModelName, UMVVMViewModelBase* NewValue)
 					{
 						UE::MVVM::FMessageLog Log(GetUserWidget());
 						Log.Error(LOCTEXT("SetViewModelInvalidValueType", "The new viewmodel is not of the expected type."));
-						return;
+						return false;
 					}
 					break;
 				}
@@ -101,7 +101,7 @@ void UMVVMView::SetViewModel(FName ViewModelName, UMVVMViewModelBase* NewValue)
 			{
 				UE::MVVM::FMessageLog Log(GetUserWidget());
 				Log.Error(LOCTEXT("SetViewModelViewModelNameNotFound", "The viewmodel name could not be found."));
-				return;
+				return false;
 			}
 		}
 
@@ -156,7 +156,9 @@ void UMVVMView::SetViewModel(FName ViewModelName, UMVVMViewModelBase* NewValue)
 				}
 			}
 		}
+		return true;
 	}
+	return false;
 }
 
 
diff --git a/Engine/Plugins/Runtime/ModelViewViewModel/Source/ModelViewViewModel/Public/MVVMViewModelBase.h b/Engine/Plugins/Runtime/ModelViewViewModel/Source/ModelViewViewModel/Public/MVVMViewModelBase.h
--- a/Engine/Plugins/Runtime/ModelViewViewModel/Source/ModelViewViewModel/Public/MVVMViewModelBase.h
+++ b/Engine/Plugins/Runtime/ModelViewViewModel/Source/ModelViewViewModel/Public/MVVMViewModelBase.h
@@ -54,6 +54,11 @@ public:
 	virtual const UE::FieldNotification::IClassDescriptor& GetFieldNotificationDescriptor() const override;
 	//~ End INotifyFieldValueChanged Interface
 
+	UFUNCTION(BlueprintCallable, Category = "FieldNotify", meta = (DisplayName = "Add Field Value Changed Delegate", ScriptName = "AddFieldValueChangedDelegate"))
+	void K2_AddFieldValueChangedDelegate(FFieldNotificationId FieldId, FFieldValueChangedDynamicDelegate Delegate);
+
+	UFUNCTION(BlueprintCallable, Category = "FieldNotify", meta = (DisplayName = "Remove Field Value Changed Delegate", ScriptName = "RemoveFieldValueChangedDelegate"))
+	void K2_RemoveFieldValueChangedDelegate(FFieldNotificationId FieldId, FFieldValueChangedDynamicDelegate Delegate);
 
 protected:
 	/** Execute the replication code. */
diff --git a/Engine/Plugins/Runtime/ModelViewViewModel/Source/ModelViewViewModel/Public/Types/MVVMFieldVariant.h b/Engine/Plugins/Runtime/ModelViewViewModel/Source/ModelViewViewModel/Public/Types/MVVMFieldVariant.h
--- a/Engine/Plugins/Runtime/ModelViewViewModel/Source/ModelViewViewModel/Public/Types/MVVMFieldVariant.h
+++ b/Engine/Plugins/Runtime/ModelViewViewModel/Source/ModelViewViewModel/Public/Types/MVVMFieldVariant.h
@@ -23,7 +23,17 @@ namespace UE::MVVM
 		friend TMVVMFieldVariant<!bConst>;
 		using FunctionType = std::conditional_t<bConst, const UFunction, UFunction>;
 		using PropertyType = std::conditional_t<bConst, const FProperty, FProperty>;
-		using VariantType = TVariant<FEmptyVariantState, PropertyType*, FunctionType*>;
+		union VariantType
+		{
+			FunctionType* Function = nullptr;
+			PropertyType* Property;
+		};
+		enum class EVariantType : uint8
+		{
+			Empty,
+			Function,
+			Property,
+		};
 
 	public:
 		TMVVMFieldVariant() = default;
@@ -32,98 +42,103 @@ namespace UE::MVVM
 		{
 			if (PropertyType* Property = InVariant.Get<PropertyType>())
 			{
-				Binding = VariantType(TInPlaceType<PropertyType*>(), Property);
+				Binding.Property = Property;
+				Type = EVariantType::Property;
 			}
 			else if (FunctionType* Function = InVariant.Get<FunctionType>())
 			{
-				Binding = VariantType(TInPlaceType<FunctionType*>(), Function);
+				Binding.Function = Function;
+				Type = EVariantType::Function;
 			}
 		}
 
 		explicit TMVVMFieldVariant(PropertyType* InValue)
 		{
-			Binding = VariantType(TInPlaceType<PropertyType*>(), InValue);
+			Binding.Property = InValue;
+			Type = EVariantType::Property;
 		}
 
 		explicit TMVVMFieldVariant(FunctionType* InValue)
 		{
-			Binding = VariantType(TInPlaceType<FunctionType*>(), InValue);
+			Binding.Function = InValue;
+			Type = EVariantType::Function;
 		}
 
 		UE_NODISCARD bool IsProperty() const
 		{
-			return Binding.template IsType<PropertyType*>();
+			return Type == EVariantType::Property;
 		}
 
 		UE_NODISCARD PropertyType* GetProperty() const
 		{
-			return Binding.template Get<PropertyType*>();
+			check(Type == EVariantType::Property);
+			return Binding.Property;
 		}
 
 		void SetProperty(PropertyType* InValue)
 		{
-			Binding.template Set<PropertyType*>(InValue);
+			Binding.Property = InValue;
+			Type = EVariantType::Property;
 		}
 
 		UE_NODISCARD bool IsFunction() const
 		{
-			return Binding.template IsType<FunctionType*>();
+			return Type == EVariantType::Function;
 		}
 
 		UE_NODISCARD FunctionType* GetFunction() const
 		{
-			return Binding.template Get<FunctionType*>();
+			check(Type == EVariantType::Function);
+			return Binding.Function;
 		}
 
 		void SetFunction(FunctionType* InValue)
 		{
-			Binding.template Set<FunctionType*>(InValue);
+			Binding.Function = InValue;
+			Type = EVariantType::Function;
 		}
 
 		UE_NODISCARD FName GetName() const
 		{
-			if (IsProperty())
+			if (Type == EVariantType::Property)
 			{
-				PropertyType* Property = GetProperty();
-				return Property ? Property->GetFName() : FName();
+				return Binding.Property ? Binding.Property->GetFName() : FName();
 			}
-			else if (IsFunction())
+			else if (Type == EVariantType::Function)
 			{
-				FunctionType* Function = GetFunction();
-				return Function ? Function->GetFName() : FName();
+				return Binding.Function ? Binding.Function->GetFName() : FName();
 			}
 			return FName();
 		}
 
 		UE_NODISCARD UStruct* GetOwner() const
 		{
-			if (IsProperty())
+			if (Type == EVariantType::Property)
 			{
-				PropertyType* Property = GetProperty();
-				return Property ? Property->GetOwnerStruct() : nullptr;
+				return Binding.Property ? Binding.Property->GetOwnerStruct() : nullptr;
 			}
-			else if (IsFunction())
+			else if (Type == EVariantType::Function)
 			{
-				FunctionType* Function = GetFunction();
-				return Function ? Function->GetOwnerClass() : nullptr;
+				return Binding.Function ? Binding.Function->GetOwnerClass() : nullptr;
 			}
 			return nullptr;
 		}
 
 		UE_NODISCARD bool IsEmpty() const
 		{
-			return Binding.template IsType<FEmptyVariantState>();
+			return Type == EVariantType::Empty || Binding.Function == nullptr;
 		}
 
 		void Reset()
 		{
-			Binding = VariantType();
+			Binding.Function = nullptr;
+			Type = EVariantType::Empty;
 		}
 
 		template<bool bOtherConst>
 		bool operator==(const TMVVMFieldVariant<bOtherConst>& B) const
 		{
-			if (Binding.GetIndex() != B.Binding.GetIndex())
+			if (Type != B.Type)
 			{
 				return false;
 			}
@@ -131,11 +146,7 @@ namespace UE::MVVM
 			{
 				return true;
 			}
-			if (IsProperty())
-			{
-				return GetProperty() == B.GetProperty();
-			}
-			return GetFunction() == B.GetFunction();
+			return Binding.Function == B.Binding.Function;
 		}
 
 		template<bool bOtherConst>
@@ -146,19 +157,19 @@ namespace UE::MVVM
 
 		friend int32 GetTypeHash(const TMVVMFieldVariant<bConst>& Variant)
 		{
-			if (Variant.IsProperty())
-			{
-				return GetTypeHash(Variant.GetProperty());
-			}
-			if (Variant.IsFunction())
+			switch (Type)
 			{
-				return GetTypeHash(Variant.GetFunction());
+				case EVariantType::Function:
+					return GetTypeHash(Binding.Function);
+				case EVariantType::Property:
+					return GetTypeHash(Binding.Property);
 			}
 			return 0;
 		}
 
 	private:
 		VariantType Binding;
+		EVariantType Type = EVariantType::Empty;
 	};
 
 	/** */
diff --git a/Engine/Plugins/Runtime/ModelViewViewModel/Source/ModelViewViewModel/Public/View/MVVMView.h b/Engine/Plugins/Runtime/ModelViewViewModel/Source/ModelViewViewModel/Public/View/MVVMView.h
--- a/Engine/Plugins/Runtime/ModelViewViewModel/Source/ModelViewViewModel/Public/View/MVVMView.h
+++ b/Engine/Plugins/Runtime/ModelViewViewModel/Source/ModelViewViewModel/Public/View/MVVMView.h
@@ -47,7 +47,7 @@ public:
 	//bool IsLibraryBindingEnabled(FGuid ViewModelId, FMVVMBindingName BindingName) const;
 
 	UFUNCTION(BlueprintCallable, Category = "MVVM")
-	void SetViewModel(FName ViewModelName, UMVVMViewModelBase* ViewModel);
+	bool SetViewModel(FName ViewModelName, UMVVMViewModelBase* ViewModel);
 
 private:
 	void HandledLibraryBindingValueChanged(UObject* InViewModel, UE::FieldNotification::FFieldId InFieldId, int32 InCompiledBindingIndex) const;
diff --git a/Engine/Plugins/Runtime/ModelViewViewModel/Source/ModelViewViewModelEditor/Private/MVVMEditorSubsystem.cpp b/Engine/Plugins/Runtime/ModelViewViewModel/Source/ModelViewViewModelEditor/Private/MVVMEditorSubsystem.cpp
--- a/Engine/Plugins/Runtime/ModelViewViewModel/Source/ModelViewViewModelEditor/Private/MVVMEditorSubsystem.cpp
+++ b/Engine/Plugins/Runtime/ModelViewViewModel/Source/ModelViewViewModelEditor/Private/MVVMEditorSubsystem.cpp
@@ -682,6 +682,7 @@ TArray<UFunction*> UMVVMEditorSubsystem::GetAvailableConversionFunctions(const U
 					// functions in the widget blueprint can do anything they want, other functions have to be static functions in a BlueprintFunctionLibrary
 					const UClass* FunctionClass = Function->GetOuterUClass();
 					if ((FunctionClass->ClassGeneratedBy == WidgetBlueprint) ||
+						(WidgetBlueprint->ParentClass && WidgetBlueprint->ParentClass->IsChildOf(FunctionClass) && Function->HasAllFunctionFlags(FUNC_Static | FUNC_BlueprintPure)) ||
 						(FunctionClass->IsChildOf<UBlueprintFunctionLibrary>() && Function->HasAllFunctionFlags(FUNC_Static | FUNC_BlueprintPure)))
 					{
 						if (IsValidConversionFunction(Function, Source, Destination))
diff --git a/Engine/Plugins/Runtime/ModelViewViewModel/Source/ModelViewViewModelEditor/Private/Widgets/SMVVMFieldEntry.cpp b/Engine/Plugins/Runtime/ModelViewViewModel/Source/ModelViewViewModelEditor/Private/Widgets/SMVVMFieldEntry.cpp
--- a/Engine/Plugins/Runtime/ModelViewViewModel/Source/ModelViewViewModelEditor/Private/Widgets/SMVVMFieldEntry.cpp
+++ b/Engine/Plugins/Runtime/ModelViewViewModel/Source/ModelViewViewModelEditor/Private/Widgets/SMVVMFieldEntry.cpp
@@ -14,27 +14,34 @@ namespace Private
 
 FText GetFieldDisplayName(const FMVVMConstFieldVariant& Field)
 {
-	if (Field.IsProperty())
+	if (!Field.IsEmpty())
 	{
-		return Field.GetProperty()->GetDisplayNameText();
-	}
-	if (Field.IsFunction())
-	{
-		return Field.GetFunction()->GetDisplayNameText();
+		if (Field.IsProperty())
+		{
+			return Field.GetProperty()->GetDisplayNameText();
+		}
+		if (Field.IsFunction())
+		{
+			return Field.GetFunction()->GetDisplayNameText();
+		}
 	}
 	return LOCTEXT("None", "<None>");
 }
 
 FText GetFieldToolTip(const FMVVMConstFieldVariant& Field)
 {
-	if (Field.IsFunction())
+	if (!Field.IsEmpty())
 	{
-		return Field.GetFunction()->GetToolTipText();
-	}
-	if (Field.IsProperty())
-	{
-		return FText::Join(FText::FromString(TEXT("\n")), Field.GetProperty()->GetToolTipText(), FText::FromString(Field.GetProperty()->GetCPPType()));
+		if (Field.IsFunction())
+		{
+			return Field.GetFunction()->GetToolTipText();
+		}
+		if (Field.IsProperty())
+		{
+			return FText::Join(FText::FromString(TEXT("\n")), Field.GetProperty()->GetToolTipText(), FText::FromString(Field.GetProperty()->GetCPPType()));
+		}
 	}
+
 	return FText::GetEmpty();
 }
 
