Fix compiler error. Introduced with CL 22073194.
#rnx

diff --git a/Engine/Plugins/Runtime/ModelViewViewModel/Source/ModelViewViewModel/Public/Types/MVVMFieldVariant.h b/Engine/Plugins/Runtime/ModelViewViewModel/Source/ModelViewViewModel/Public/Types/MVVMFieldVariant.h
--- a/Engine/Plugins/Runtime/ModelViewViewModel/Source/ModelViewViewModel/Public/Types/MVVMFieldVariant.h
+++ b/Engine/Plugins/Runtime/ModelViewViewModel/Source/ModelViewViewModel/Public/Types/MVVMFieldVariant.h
@@ -157,12 +157,12 @@ namespace UE::MVVM
 
 		friend int32 GetTypeHash(const TMVVMFieldVariant<bConst>& Variant)
 		{
-			switch (Type)
+			switch (Variant.Type)
 			{
 				case EVariantType::Function:
-					return GetTypeHash(Binding.Function);
+					return GetTypeHash(Variant.Binding.Function);
 				case EVariantType::Property:
-					return GetTypeHash(Binding.Property);
+					return GetTypeHash(Variant.Binding.Property);
 			}
 			return 0;
 		}
