UMG: Add PreConstruct and Construct call back to extension. Do not remove extension in the destruction call. They maybe be added out of the construct/destruct loop. Fix bad test in EvaluateObjectProperty.
#cherrypick 21739076

diff --git a/Engine/Plugins/Runtime/ModelViewViewModel/Source/ModelViewViewModel/Private/View/MVVMViewClass.cpp b/Engine/Plugins/Runtime/ModelViewViewModel/Source/ModelViewViewModel/Private/View/MVVMViewClass.cpp
--- a/Engine/Plugins/Runtime/ModelViewViewModel/Source/ModelViewViewModel/Private/View/MVVMViewClass.cpp
+++ b/Engine/Plugins/Runtime/ModelViewViewModel/Source/ModelViewViewModel/Private/View/MVVMViewClass.cpp
@@ -165,7 +165,7 @@ FString FMVVMViewClass_CompiledBinding::ToString() const
 // 
 ///////////////////////////////////////////////////////////////////////
 
-void UMVVMViewClass::Construct(UUserWidget* UserWidget)
+void UMVVMViewClass::Initialize(UUserWidget* UserWidget)
 {
 	ensure(UserWidget->GetExtension<UMVVMView>() == nullptr);
 	UMVVMView* View = UserWidget->AddExtension<UMVVMView>();
@@ -181,10 +181,4 @@ void UMVVMViewClass::Construct(UUserWidget* UserWidget)
 	}
 }
 
-
-void UMVVMViewClass::Destruct(UUserWidget* UserWidget)
-{
-	UserWidget->RemoveExtensions<UMVVMView>();
-}
-
 #undef LOCTEXT_NAMESPACE
diff --git a/Engine/Plugins/Runtime/ModelViewViewModel/Source/ModelViewViewModel/Public/View/MVVMViewClass.h b/Engine/Plugins/Runtime/ModelViewViewModel/Source/ModelViewViewModel/Public/View/MVVMViewClass.h
--- a/Engine/Plugins/Runtime/ModelViewViewModel/Source/ModelViewViewModel/Public/View/MVVMViewClass.h
+++ b/Engine/Plugins/Runtime/ModelViewViewModel/Source/ModelViewViewModel/Public/View/MVVMViewClass.h
@@ -194,8 +194,7 @@ class MODELVIEWVIEWMODEL_API UMVVMViewClass : public UWidgetBlueprintGeneratedCl
 
 public:
 	//~ Begin UWidgetBlueprintGeneratedClassExtension
-	virtual void Construct(UUserWidget* UserWidget) override;
-	virtual void Destruct(UUserWidget* UserWidget) override;
+	virtual void Initialize(UUserWidget* UserWidget) override;
 	//~ End UWidgetBlueprintGeneratedClassExtension
 
 public:
