MVVM: Fix ViewModel name in IMPLEMENT_MODULE. It prevented the module to be loaded in cooked build.
#rnx

diff --git a/Engine/Plugins/Runtime/ModelViewViewModel/Source/ModelViewViewModelBlueprint/Private/ModelViewViewModelBlueprintModule.cpp b/Engine/Plugins/Runtime/ModelViewViewModel/Source/ModelViewViewModelBlueprint/Private/ModelViewViewModelBlueprintModule.cpp
--- a/Engine/Plugins/Runtime/ModelViewViewModel/Source/ModelViewViewModelBlueprint/Private/ModelViewViewModelBlueprintModule.cpp
+++ b/Engine/Plugins/Runtime/ModelViewViewModel/Source/ModelViewViewModelBlueprint/Private/ModelViewViewModelBlueprintModule.cpp
@@ -33,4 +33,4 @@ private:
 };
 
 
-IMPLEMENT_MODULE(FModelViewViewModelBlueprintModule, ModelViewViewModelBlueprintModule);
+IMPLEMENT_MODULE(FModelViewViewModelBlueprintModule, ModelViewViewModelBlueprint);
