MVVM - Function parameters of simple conversion functions no longer clear their property.

#rnx

diff --git a/Engine/Plugins/Runtime/ModelViewViewModel/Source/ModelViewViewModelEditor/Private/Widgets/SMVVMFunctionParameter.cpp b/Engine/Plugins/Runtime/ModelViewViewModel/Source/ModelViewViewModelEditor/Private/Widgets/SMVVMFunctionParameter.cpp
--- a/Engine/Plugins/Runtime/ModelViewViewModel/Source/ModelViewViewModelEditor/Private/Widgets/SMVVMFunctionParameter.cpp
+++ b/Engine/Plugins/Runtime/ModelViewViewModel/Source/ModelViewViewModelEditor/Private/Widgets/SMVVMFunctionParameter.cpp
@@ -122,11 +122,6 @@ void SFunctionParameter::Construct(const FArguments& InArgs)
 				]
 			]
 		];
-
-	if (Pin == nullptr)
-	{
-		OnBindArgumentChecked(ECheckBoxState::Checked);
-	}
 }
 
 EVisibility SFunctionParameter::OnGetVisibility(bool bDefaultValue) const
