MVVM - Simple conversion functions do not allow you to enter a single hardcoded value into the function parameter anymore.

#rnx

diff --git a/Engine/Plugins/Runtime/ModelViewViewModel/Source/ModelViewViewModelEditor/Private/Widgets/SMVVMFunctionParameter.cpp b/Engine/Plugins/Runtime/ModelViewViewModel/Source/ModelViewViewModelEditor/Private/Widgets/SMVVMFunctionParameter.cpp
--- a/Engine/Plugins/Runtime/ModelViewViewModel/Source/ModelViewViewModelEditor/Private/Widgets/SMVVMFunctionParameter.cpp
+++ b/Engine/Plugins/Runtime/ModelViewViewModel/Source/ModelViewViewModelEditor/Private/Widgets/SMVVMFunctionParameter.cpp
@@ -32,6 +32,7 @@ void SFunctionParameter::Construct(const FArguments& InArgs)
 	check(!ParameterName.IsNone());
 	
 	bSourceToDestination = InArgs._SourceToDestination;
+	bAllowDefault = InArgs._AllowDefault;
 
 	GetBindingModeDelegate = InArgs._OnGetBindingMode;
 	check(GetBindingModeDelegate.IsBound());
@@ -73,12 +74,14 @@ void SFunctionParameter::Construct(const FArguments& InArgs)
 		bFromViewModel = UE::MVVM::IsBackwardBinding(Binding->BindingType);
 	}
 
+	TSharedPtr<SHorizontalBox> HBox;
+
 	ChildSlot
 	[
 		SNew(SBox)
 		.MinDesiredWidth(100)
 		[
-			SNew(SHorizontalBox)
+			SAssignNew(HBox, SHorizontalBox)
 			+ SHorizontalBox::Slot()
 			.Padding(0, 2)
 			.VAlign(VAlign_Center)
@@ -98,7 +101,12 @@ void SFunctionParameter::Construct(const FArguments& InArgs)
 					.AssignableTo(Property)
 				]
 			]
-			+ SHorizontalBox::Slot()
+		]
+	];
+
+	if (InArgs._AllowDefault)
+	{
+		HBox->AddSlot()
 			.HAlign(HAlign_Right)
 			.VAlign(VAlign_Center)
 			.Padding(8, 0, 0, 0)
@@ -118,14 +126,18 @@ void SFunctionParameter::Construct(const FArguments& InArgs)
 							ECheckBoxState CheckState = OnGetIsBindArgumentChecked();
 							return (CheckState == ECheckBoxState::Checked) ? FAppStyle::GetBrush("Icons.Link") : FAppStyle::GetBrush("Icons.Unlink");
 						})
-					]
 				]
-			]
-		];
+			];		
+	}
 }
 
 EVisibility SFunctionParameter::OnGetVisibility(bool bDefaultValue) const
 {
+	if (!bAllowDefault)
+	{
+		return bDefaultValue ? EVisibility::Collapsed : EVisibility::Visible;
+	}
+
 	FMVVMBlueprintPropertyPath Path = OnGetSelectedField();
 
 	// if we're not bound then show the default value widget, otherwise show the binding widget
diff --git a/Engine/Plugins/Runtime/ModelViewViewModel/Source/ModelViewViewModelEditor/Private/Widgets/SMVVMFunctionParameter.h b/Engine/Plugins/Runtime/ModelViewViewModel/Source/ModelViewViewModelEditor/Private/Widgets/SMVVMFunctionParameter.h
--- a/Engine/Plugins/Runtime/ModelViewViewModel/Source/ModelViewViewModelEditor/Private/Widgets/SMVVMFunctionParameter.h
+++ b/Engine/Plugins/Runtime/ModelViewViewModel/Source/ModelViewViewModelEditor/Private/Widgets/SMVVMFunctionParameter.h
@@ -33,7 +33,8 @@ public:
 		SLATE_ARGUMENT(UWidgetBlueprint*, WidgetBlueprint)
 		SLATE_ARGUMENT(FMVVMBlueprintViewBinding*, Binding)
 		SLATE_ARGUMENT(FName, ParameterName)
-		SLATE_ARGUMENT(bool, SourceToDestination)
+		SLATE_ARGUMENT_DEFAULT(bool, SourceToDestination) = true;
+		SLATE_ARGUMENT_DEFAULT(bool, AllowDefault) = true;
 	SLATE_END_ARGS()
 
 	void Construct(const FArguments& InArgs);
@@ -54,8 +55,10 @@ private:
 	/** This reference is just to keep the default value widget alive. */
 	TSharedPtr<SGraphPin> GraphPin;
 	TSharedPtr<UE::MVVM::SFieldSelector> FieldSelector;
-	bool bSourceToDestination = true;
 	FGetBindingMode GetBindingModeDelegate;
+
+	bool bSourceToDestination = true;
+	bool bAllowDefault = true;
 };
 
 }
diff --git a/Engine/Plugins/Runtime/ModelViewViewModel/Source/ModelViewViewModelEditor/Private/Widgets/SMVVMViewBindingListView.cpp b/Engine/Plugins/Runtime/ModelViewViewModel/Source/ModelViewViewModelEditor/Private/Widgets/SMVVMViewBindingListView.cpp
--- a/Engine/Plugins/Runtime/ModelViewViewModel/Source/ModelViewViewModelEditor/Private/Widgets/SMVVMViewBindingListView.cpp
+++ b/Engine/Plugins/Runtime/ModelViewViewModel/Source/ModelViewViewModelEditor/Private/Widgets/SMVVMViewBindingListView.cpp
@@ -824,6 +824,8 @@ class SFunctionParameterRow : public STableRow<TSharedPtr<FBindingEntry>>
 		const FSlateBrush* PrimaryBrush = nullptr;
 		const FSlateBrush* SecondaryBrush = nullptr;
 		FText DisplayName, ToolTip;
+
+		bool bSimpleConversionFunction = false;
 		
 		if (UEdGraphPin* Pin = EditorSubsystem->FindConversionFunctionArgumentPin(InWidgetBlueprint, *Binding, Entry->GetName(), bSourceToDestination))
 		{
@@ -836,6 +838,8 @@ class SFunctionParameterRow : public STableRow<TSharedPtr<FBindingEntry>>
 			// no wrapper graph, this is a simple conversion function of the form: int32 Convert(float x)
 			if (const FProperty* Argument = BindingHelper::GetFirstArgumentProperty(Function))
 			{
+				bSimpleConversionFunction = true;
+
 				PrimaryBrush = FBlueprintEditor::GetVarIconAndColorFromProperty(Argument, PrimaryColor, SecondaryBrush, SecondaryColor);
 				DisplayName = Argument->GetDisplayNameText();
 				ToolTip = Argument->GetToolTipText();
@@ -891,6 +895,7 @@ class SFunctionParameterRow : public STableRow<TSharedPtr<FBindingEntry>>
 						.Binding(Binding)
 						.ParameterName(Entry->GetName())
 						.SourceToDestination(bSourceToDestination)
+						.AllowDefault(!bSimpleConversionFunction)
 					]
 				]
 			]
