MVVM: Use WeakObjectPtr to save UObject reference. Display the widget current slot properties. Confirm that the selected property can be assign or read before returning it. Created project setting. Filter property that can be display in the FieldSelector.

diff --git a/Engine/Plugins/Runtime/ModelViewViewModel/Config/BaseModelViewViewModel.ini b/Engine/Plugins/Runtime/ModelViewViewModel/Config/BaseModelViewViewModel.ini
--- a/Engine/Plugins/Runtime/ModelViewViewModel/Config/BaseModelViewViewModel.ini
+++ b/Engine/Plugins/Runtime/ModelViewViewModel/Config/BaseModelViewViewModel.ini
@@ -1,3 +1,7 @@
 [CoreRedirects]
 +StructRedirects=(OldName="MVVMViewModelPropertyPath",NewName="/script/ModelViewViewModelBlueprint.MVVMBlueprintPropertyPath")
 +StructRedirects=(OldName="MVVMWidgetPropertyPath",NewName="/script/ModelViewViewModelBlueprint.MVVMBlueprintPropertyPath")
+
+
+[/Script/ModelViewViewModelBlueprint.MVVMDeveloperProjectSettings]
++FieldSelectorPermissions=(("/Script/UMG.Widget", (DisallowedFieldNames=("TooltipWidget","ForceVolatile","SetUserFocus","Navigation"),AdvancedFieldNames=())))
\ No newline at end of file
diff --git a/Engine/Plugins/Runtime/ModelViewViewModel/Source/ModelViewViewModelBlueprint/ModelViewViewModelBlueprint.Build.cs b/Engine/Plugins/Runtime/ModelViewViewModel/Source/ModelViewViewModelBlueprint/ModelViewViewModelBlueprint.Build.cs
--- a/Engine/Plugins/Runtime/ModelViewViewModel/Source/ModelViewViewModelBlueprint/ModelViewViewModelBlueprint.Build.cs
+++ b/Engine/Plugins/Runtime/ModelViewViewModel/Source/ModelViewViewModelBlueprint/ModelViewViewModelBlueprint.Build.cs
@@ -11,9 +11,10 @@ public class ModelViewViewModelBlueprint : ModuleRules
 		PublicDependencyModuleNames.AddRange(
 			new string[]
 			{
-				"Engine",
 				"Core",
 				"CoreUObject",
+				"DeveloperSettings",
+				"Engine",
 				"ModelViewViewModel",
 			});
 
diff --git a/Engine/Plugins/Runtime/ModelViewViewModel/Source/ModelViewViewModelBlueprint/Private/MVVMDeveloperProjectSettings.cpp b/Engine/Plugins/Runtime/ModelViewViewModel/Source/ModelViewViewModelBlueprint/Private/MVVMDeveloperProjectSettings.cpp
new file mode 100644
--- /dev/null
+++ b/Engine/Plugins/Runtime/ModelViewViewModel/Source/ModelViewViewModelBlueprint/Private/MVVMDeveloperProjectSettings.cpp
@@ -0,0 +1,45 @@
+// Copyright Epic Games, Inc. All Rights Reserved.
+
+#include "MVVMDeveloperProjectSettings.h"
+
+#define LOCTEXT_NAMESPACE "MVVMDeveloperProjectSettings"
+
+FName UMVVMDeveloperProjectSettings::GetCategoryName() const
+{
+	return TEXT("Plugins");
+}
+
+FText UMVVMDeveloperProjectSettings::GetSectionText() const
+{
+	return LOCTEXT("MVVMProjectSettings", "Model View Viewmodel");
+}
+
+bool UMVVMDeveloperProjectSettings::IsPropertyAllowed(FProperty* Property) const
+{
+	check(Property);
+	TStringBuilder<256> StringBuilder;
+	Property->GetOwnerClass()->GetPathName(nullptr, StringBuilder);
+	FSoftClassPath StructPath;
+	StructPath.SetPath(StringBuilder);
+	if (const FMVVMDeveloperProjectWidgetSettings* Settings = FieldSelectorPermissions.Find(StructPath))
+	{
+		return !Settings->DisallowedFieldNames.Find(Property->GetFName());
+	}
+	return true;
+}
+
+bool UMVVMDeveloperProjectSettings::IsFunctionAllowed(UFunction* Function) const
+{
+	check(Function);
+	TStringBuilder<256> StringBuilder;
+	Function->GetOwnerClass()->GetPathName(nullptr, StringBuilder);
+	FSoftClassPath StructPath;
+	StructPath.SetPath(StringBuilder);
+	if (const FMVVMDeveloperProjectWidgetSettings* Settings = FieldSelectorPermissions.Find(StructPath))
+	{
+		return !Settings->DisallowedFieldNames.Find(Function->GetFName());
+	}
+	return true;
+}
+
+#undef LOCTEXT_NAMESPACE
diff --git a/Engine/Plugins/Runtime/ModelViewViewModel/Source/ModelViewViewModelBlueprint/Public/MVVMDeveloperProjectSettings.h b/Engine/Plugins/Runtime/ModelViewViewModel/Source/ModelViewViewModelBlueprint/Public/MVVMDeveloperProjectSettings.h
new file mode 100644
--- /dev/null
+++ b/Engine/Plugins/Runtime/ModelViewViewModel/Source/ModelViewViewModelBlueprint/Public/MVVMDeveloperProjectSettings.h
@@ -0,0 +1,48 @@
+// Copyright Epic Games, Inc. All Rights Reserved.
+
+#pragma once
+
+#include "CoreMinimal.h"
+#include "Components/Widget.h"
+#include "Engine/DeveloperSettings.h"
+#include "UObject/SoftObjectPtr.h"
+#include "MVVMDeveloperProjectSettings.generated.h"
+
+/**
+ * 
+ */
+USTRUCT()
+struct FMVVMDeveloperProjectWidgetSettings
+{
+	GENERATED_BODY()
+
+	/** Properties or functions name that should not be use for binding (read or write). */
+	UPROPERTY(EditAnywhere, config, Category=MVVM)
+	TSet<FName> DisallowedFieldNames;
+	
+	/** Properties or functions name that are displayed in the advanced category. */
+	UPROPERTY(EditAnywhere, config, Category=MVVM)
+	TSet<FName> AdvancedFieldNames;
+};
+
+
+/**
+ * Implements the settings for the MVVM Editor
+ */
+UCLASS(config=ModelViewViewModel, defaultconfig)
+class MODELVIEWVIEWMODELBLUEPRINT_API UMVVMDeveloperProjectSettings : public UDeveloperSettings
+{
+	GENERATED_BODY()
+
+public:
+	virtual FName GetCategoryName() const override;
+	virtual FText GetSectionText() const override;
+
+	bool IsPropertyAllowed(FProperty* Property) const;
+	bool IsFunctionAllowed(UFunction* Function) const;
+
+private:
+	/** Permission list for filtering which properties are visible in UI. */
+	UPROPERTY(EditAnywhere, config, Category=MVVM)
+	TMap<FSoftClassPath, FMVVMDeveloperProjectWidgetSettings> FieldSelectorPermissions;
+};
diff --git a/Engine/Plugins/Runtime/ModelViewViewModel/Source/ModelViewViewModelEditor/Private/Types/MVVMBindingSource.h b/Engine/Plugins/Runtime/ModelViewViewModel/Source/ModelViewViewModelEditor/Private/Types/MVVMBindingSource.h
--- a/Engine/Plugins/Runtime/ModelViewViewModel/Source/ModelViewViewModelEditor/Private/Types/MVVMBindingSource.h
+++ b/Engine/Plugins/Runtime/ModelViewViewModel/Source/ModelViewViewModelEditor/Private/Types/MVVMBindingSource.h
@@ -13,7 +13,7 @@ namespace UE::MVVM
 	{
 		FGuid ViewModelId;
 		FName Name;
-		UClass* Class = nullptr;
+		TWeakObjectPtr<UClass> Class;
 		FText DisplayName;
 
 		bool operator==(const FBindingSource& Other) const
diff --git a/Engine/Plugins/Runtime/ModelViewViewModel/Source/ModelViewViewModelEditor/Private/Widgets/SMVVMFieldSelector.cpp b/Engine/Plugins/Runtime/ModelViewViewModel/Source/ModelViewViewModelEditor/Private/Widgets/SMVVMFieldSelector.cpp
--- a/Engine/Plugins/Runtime/ModelViewViewModel/Source/ModelViewViewModelEditor/Private/Widgets/SMVVMFieldSelector.cpp
+++ b/Engine/Plugins/Runtime/ModelViewViewModel/Source/ModelViewViewModelEditor/Private/Widgets/SMVVMFieldSelector.cpp
@@ -21,6 +21,7 @@
 #include "Widgets/Images/SImage.h"
 #include "Widgets/Input/SComboButton.h"
 #include "Widgets/Input/SSearchBox.h"
+#include "Widgets/SNullWidget.h"
 
 #define LOCTEXT_NAMESPACE "MVVMFieldSelector"
 
@@ -48,7 +49,7 @@ FBindingSource GetSourceFromPath(const UWidgetBlueprint* WidgetBlueprint, const
 void SFieldSelector::Construct(const FArguments& InArgs, const UWidgetBlueprint* InWidgetBlueprint, bool bInViewModelProperty)
 {
 	WidgetBlueprint = InWidgetBlueprint;
-	check(WidgetBlueprint);
+	check(InWidgetBlueprint);
 	
 	BindingMode = InArgs._BindingMode;
 	check(BindingMode.IsSet());
@@ -180,9 +181,9 @@ void SFieldSelector::Refresh()
 		SelectedEntryWidget->SetField(CachedSelectedField);
 	}
 
-	if (SelectedSourceWidget.IsValid())
+	if (SelectedSourceWidget.IsValid() && WidgetBlueprint.IsValid())
 	{
-		SelectedSourceWidget->RefreshSource(Private::GetSourceFromPath(WidgetBlueprint, CachedSelectedField));
+		SelectedSourceWidget->RefreshSource(Private::GetSourceFromPath(WidgetBlueprint.Get(), CachedSelectedField));
 	}
 }
 
@@ -307,9 +308,9 @@ void SFieldSelector::OnViewModelSelected(FBindingSource Source, ESelectInfo::Typ
 
 void SFieldSelector::OnWidgetSelected(FName WidgetName, ESelectInfo::Type)
 {
-	if (BindingList.IsValid())
+	if (BindingList.IsValid() && WidgetBlueprint.IsValid())
 	{
-		FBindingSource Source = FBindingSource::CreateForWidget(WidgetBlueprint, WidgetName);
+		FBindingSource Source = FBindingSource::CreateForWidget(WidgetBlueprint.Get(), WidgetName);
 
 		BindingList->Clear();
 		BindingList->AddSource(Source);
@@ -333,7 +334,7 @@ TSharedRef<ITableRow> SFieldSelector::GenerateViewModelRow(FBindingSource ViewMo
 			.AutoWidth()
 			[
 				SNew(SImage)
-				.Image(FSlateIconFinder::FindIconBrushForClass(ViewModel.Class))
+				.Image(FSlateIconFinder::FindIconBrushForClass(ViewModel.Class.Get()))
 				.ColorAndOpacity(FSlateColor::UseForeground())
 			]
 			+ SHorizontalBox::Slot()
@@ -472,6 +473,11 @@ TSharedRef<SWidget> SFieldSelector::CreateSourcePanel()
 	ViewModelList.Reset();
 	WidgetList.Reset();
 
+	if (!WidgetBlueprint.IsValid())
+	{
+		return SNullWidget::NullWidget;
+	}
+
 	// show source picker
 	TSharedPtr<SWidget> SourcePicker;
 	if (bViewModelProperty)
@@ -479,7 +485,7 @@ TSharedRef<SWidget> SFieldSelector::CreateSourcePanel()
 		FBindingSource SelectedSource;
 
 		UMVVMEditorSubsystem* EditorSubsystem = GEditor->GetEditorSubsystem<UMVVMEditorSubsystem>();
-		if (const UMVVMBlueprintView* View = EditorSubsystem->GetView(WidgetBlueprint))
+		if (const UMVVMBlueprintView* View = EditorSubsystem->GetView(WidgetBlueprint.Get()))
 		{
 			const TArrayView<const FMVVMBlueprintViewModelContext> ViewModels = View->GetViewModels();
 			ViewModelSources.Reserve(ViewModels.Num());
@@ -516,7 +522,7 @@ TSharedRef<SWidget> SFieldSelector::CreateSourcePanel()
 	}
 	else
 	{
-		WidgetList = SNew(SReadOnlyHierarchyView, WidgetBlueprint)
+		WidgetList = SNew(SReadOnlyHierarchyView, WidgetBlueprint.Get())
 			.OnSelectionChanged(this, &SFieldSelector::OnWidgetSelected)
 			.SelectionMode(ESelectionMode::Multi)
 			.ShowSearch(false);
@@ -527,7 +533,7 @@ TSharedRef<SWidget> SFieldSelector::CreateSourcePanel()
 	if (ShowConversionFunctions.Get(false))
 	{
 		UMVVMEditorSubsystem* Subsystem = GEditor->GetEditorSubsystem<UMVVMEditorSubsystem>();
-		TArray<UFunction*> AllConversionFunctions = Subsystem->GetAvailableConversionFunctions(WidgetBlueprint, FMVVMBlueprintPropertyPath(), FMVVMBlueprintPropertyPath());
+		TArray<UFunction*> AllConversionFunctions = Subsystem->GetAvailableConversionFunctions(WidgetBlueprint.Get(), FMVVMBlueprintPropertyPath(), FMVVMBlueprintPropertyPath());
 
 		// remove all incompatible conversion functions
 		for (int32 Idx = 0; Idx < AllConversionFunctions.Num(); ++Idx)
@@ -623,6 +629,11 @@ TSharedRef<SWidget> SFieldSelector::CreateSourcePanel()
 
 TSharedRef<SWidget> SFieldSelector::OnGetMenuContent()
 {
+	if (!WidgetBlueprint.IsValid())
+	{
+		return SNullWidget::NullWidget;
+	}
+
 	TSharedRef<SVerticalBox> VBox = SNew(SVerticalBox)
 		+ SVerticalBox::Slot()
 		.HAlign(HAlign_Fill)
@@ -644,7 +655,7 @@ TSharedRef<SWidget> SFieldSelector::OnGetMenuContent()
 		AssignableToProperty = AssignableTo.Get(nullptr);
 	}
 
-	BindingList = SNew(SSourceBindingList, WidgetBlueprint)
+	BindingList = SNew(SSourceBindingList, WidgetBlueprint.Get())
 		.ShowSearchBox(false)
 		.OnDoubleClicked(this, &SFieldSelector::SetPropertySelection)
 		.FieldVisibilityFlags(GetFieldVisibilityFlags())
@@ -652,7 +663,7 @@ TSharedRef<SWidget> SFieldSelector::OnGetMenuContent()
 
 	if (FixedSource.IsSet())
 	{
-		// single fixed source, don't show the separate source panel
+		// Single fixed source, don't show the separate source panel.
 		BindingList->AddSource(FixedSource.GetValue());
 
 		VBox->AddSlot()
diff --git a/Engine/Plugins/Runtime/ModelViewViewModel/Source/ModelViewViewModelEditor/Private/Widgets/SMVVMFieldSelector.h b/Engine/Plugins/Runtime/ModelViewViewModel/Source/ModelViewViewModelEditor/Private/Widgets/SMVVMFieldSelector.h
--- a/Engine/Plugins/Runtime/ModelViewViewModel/Source/ModelViewViewModelEditor/Private/Widgets/SMVVMFieldSelector.h
+++ b/Engine/Plugins/Runtime/ModelViewViewModel/Source/ModelViewViewModelEditor/Private/Widgets/SMVVMFieldSelector.h
@@ -146,7 +146,7 @@ private:
 
 	TOptional<FBindingSource> FixedSource;
 
-	const UWidgetBlueprint* WidgetBlueprint = nullptr;
+	TWeakObjectPtr<const UWidgetBlueprint> WidgetBlueprint;
 	bool bViewModelProperty = false;
 
 	TArray<FBindingSource> ViewModelSources;
diff --git a/Engine/Plugins/Runtime/ModelViewViewModel/Source/ModelViewViewModelEditor/Private/Widgets/SMVVMSourceEntry.cpp b/Engine/Plugins/Runtime/ModelViewViewModel/Source/ModelViewViewModelEditor/Private/Widgets/SMVVMSourceEntry.cpp
--- a/Engine/Plugins/Runtime/ModelViewViewModel/Source/ModelViewViewModelEditor/Private/Widgets/SMVVMSourceEntry.cpp
+++ b/Engine/Plugins/Runtime/ModelViewViewModel/Source/ModelViewViewModelEditor/Private/Widgets/SMVVMSourceEntry.cpp
@@ -22,7 +22,7 @@ const FSlateBrush* GetSourceIcon(const FBindingSource& Source)
 		return nullptr;
 	}
 
-	return FSlateIconFinder::FindIconBrushForClass(Source.Class);
+	return FSlateIconFinder::FindIconBrushForClass(Source.Class.Get());
 }
 
 FLinearColor GetSourceColor(const FBindingSource& Source)
diff --git a/Engine/Plugins/Runtime/ModelViewViewModel/Source/ModelViewViewModelEditor/Private/Widgets/SMVVMViewModelBindingListWidget.cpp b/Engine/Plugins/Runtime/ModelViewViewModel/Source/ModelViewViewModelEditor/Private/Widgets/SMVVMViewModelBindingListWidget.cpp
--- a/Engine/Plugins/Runtime/ModelViewViewModel/Source/ModelViewViewModelEditor/Private/Widgets/SMVVMViewModelBindingListWidget.cpp
+++ b/Engine/Plugins/Runtime/ModelViewViewModel/Source/ModelViewViewModelEditor/Private/Widgets/SMVVMViewModelBindingListWidget.cpp
@@ -4,11 +4,14 @@
 
 #include "Algo/Transform.h"
 #include "Bindings/MVVMBindingHelper.h"
+#include "Blueprint/WidgetTree.h"
 #include "Components/Widget.h"
 #include "Editor.h"
 #include "Engine/Engine.h"
 #include "FieldNotification/IFieldValueChanged.h"
+#include "Misc/Optional.h"
 #include "MVVMBlueprintViewModelContext.h"
+#include "MVVMDeveloperProjectSettings.h"
 #include "MVVMEditorSubsystem.h"
 #include "MVVMSubsystem.h"
 #include "MVVMViewModelBase.h"
@@ -24,55 +27,99 @@ using UE::PropertyViewer::SPropertyViewer;
 namespace UE::MVVM
 {
 
-FFieldIterator_Bindable::FFieldIterator_Bindable(const UWidgetBlueprint* InWidgetBlueprint, EFieldVisibility InVisibilityFlags, const FProperty* InAssignableTo) :
-	WidgetBlueprint(InWidgetBlueprint),
-	FieldVisibilityFlags(InVisibilityFlags),
-	AssignableTo(InAssignableTo)
-{
-}
-
-TArray<FFieldVariant> FFieldIterator_Bindable::GetFields(const UStruct* Struct) const
+namespace Private
 {
-	TArray<FFieldVariant> Result;
+	TOptional<FFieldVariant> PassFilter(const FMVVMAvailableBinding& Binding, const FMVVMFieldVariant& FieldVariant, EFieldVisibility FieldVisibilityFlags, const FProperty* AssignableTo, bool bDoObjectProperty);
+	TOptional<FFieldVariant> PassFilter(const FMVVMAvailableBinding& Binding, const UStruct* Struct, EFieldVisibility FieldVisibilityFlags, const FProperty* AssignableTo, bool bDoObjectProperty);
 
-	auto AddResult = [this, &Result, Struct](const TArray<FMVVMAvailableBinding>& AvailableBindingsList)
+	TOptional<FFieldVariant> PassFilter(const FMVVMAvailableBinding& Binding, const UStruct* Struct, EFieldVisibility FieldVisibilityFlags, const FProperty* AssignableTo, bool bDoObjectProperty)
 	{
-		Result.Reserve(AvailableBindingsList.Num());
-		for (const FMVVMAvailableBinding& Value : AvailableBindingsList)
+		if (EnumHasAllFlags(FieldVisibilityFlags, EFieldVisibility::Readable) && !Binding.IsReadable())
 		{
-			if (EnumHasAllFlags(FieldVisibilityFlags, EFieldVisibility::Readable) && !Value.IsReadable())
-			{
-				continue;
-			} 
+			return TOptional<FFieldVariant>();
+		}
 
-			if (EnumHasAllFlags(FieldVisibilityFlags, EFieldVisibility::Writable) && !Value.IsWritable())
-			{
-				continue;
-			}
-			
-			if (EnumHasAllFlags(FieldVisibilityFlags, EFieldVisibility::Notify) && !Value.HasNotify())
-			{
-				continue;
-			}
+		if (EnumHasAllFlags(FieldVisibilityFlags, EFieldVisibility::Notify) && !Binding.HasNotify())
+		{
+			return TOptional<FFieldVariant>();
+		}
+
+		FMVVMFieldVariant FieldVariant = BindingHelper::FindFieldByName(Struct, Binding.GetBindingName());
+		return PassFilter(Binding, FieldVariant, FieldVisibilityFlags, AssignableTo, bDoObjectProperty);
+	}
 
-			FMVVMFieldVariant FieldVariant = BindingHelper::FindFieldByName(Struct, Value.GetBindingName());
+	TOptional<FFieldVariant> PassFilter(const FMVVMAvailableBinding& Binding, const FMVVMFieldVariant& FieldVariant, EFieldVisibility FieldVisibilityFlags, const FProperty* AssignableTo, bool bDoObjectProperty)
+	{
+		if (ensure(!FieldVariant.IsEmpty()))
+		{
 			if (FieldVariant.IsFunction())
 			{
+				if (EnumHasAllFlags(FieldVisibilityFlags, EFieldVisibility::Writable) && !Binding.IsWritable())
+				{
+					return TOptional<FFieldVariant>();
+				}
+
+				if (!GetDefault<UMVVMDeveloperProjectSettings>()->IsFunctionAllowed(FieldVariant.GetFunction()))
+				{
+					return TOptional<FFieldVariant>();
+				}
+
 				if (AssignableTo != nullptr && !BindingHelper::ArePropertiesCompatible(BindingHelper::GetReturnProperty(FieldVariant.GetFunction()), AssignableTo))
 				{
-					continue;
+					return TOptional<FFieldVariant>();
 				}
 
-				Result.Add(FFieldVariant(FieldVariant.GetFunction()));
+				return FFieldVariant(FieldVariant.GetFunction());
 			}
 			else if (FieldVariant.IsProperty())
 			{
+				if (!GetDefault<UMVVMDeveloperProjectSettings>()->IsPropertyAllowed(FieldVariant.GetProperty()))
+				{
+					return TOptional<FFieldVariant>();
+				}
+
+				// If the path ends with the object property, then it needs to follow the writable rule
+				if (bDoObjectProperty || !CastField<FObjectPropertyBase>(FieldVariant.GetProperty()))
+				{
+					if (EnumHasAllFlags(FieldVisibilityFlags, EFieldVisibility::Writable) && !Binding.IsWritable())
+					{
+						return TOptional<FFieldVariant>();
+					}
+				}
+
 				if (AssignableTo != nullptr && !BindingHelper::ArePropertiesCompatible(FieldVariant.GetProperty(), AssignableTo))
 				{
-					continue;
+					return TOptional<FFieldVariant>();
 				}
 
-				Result.Add(FFieldVariant(FieldVariant.GetProperty()));
+				return FFieldVariant(FieldVariant.GetProperty());
+			}
+		}
+
+		return TOptional<FFieldVariant>();
+	}
+}
+
+FFieldIterator_Bindable::FFieldIterator_Bindable(const UWidgetBlueprint* InWidgetBlueprint, EFieldVisibility InVisibilityFlags, const FProperty* InAssignableTo) :
+	WidgetBlueprint(InWidgetBlueprint),
+	FieldVisibilityFlags(InVisibilityFlags),
+	AssignableTo(InAssignableTo)
+{
+}
+
+TArray<FFieldVariant> FFieldIterator_Bindable::GetFields(const UStruct* Struct) const
+{
+	TArray<FFieldVariant> Result;
+
+	auto AddResult = [this, &Result, Struct](const TArray<FMVVMAvailableBinding>& AvailableBindingsList)
+	{
+		Result.Reserve(AvailableBindingsList.Num());
+		for (const FMVVMAvailableBinding& Value : AvailableBindingsList)
+		{
+			TOptional<FFieldVariant> PassResult = Private::PassFilter(Value, Struct, FieldVisibilityFlags, AssignableTo, false);
+			if (PassResult.IsSet())
+			{
+				Result.Add(MoveTemp(PassResult.GetValue()));
 			}
 		}
 	};
@@ -125,14 +172,20 @@ TArray<FFieldVariant> FFieldIterator_Bindable::GetFields(const UStruct* Struct)
 }
 
 /** */
-void SSourceBindingList::Construct(const FArguments& InArgs, const UWidgetBlueprint* WidgetBlueprint)
+void SSourceBindingList::Construct(const FArguments& InArgs, const UWidgetBlueprint* InWidgetBlueprint)
 {
-	FieldIterator = MakeUnique<FFieldIterator_Bindable>(WidgetBlueprint, InArgs._FieldVisibilityFlags, InArgs._AssignableTo);
+	WidgetBlueprint = InWidgetBlueprint;
+	FieldIterator = MakeUnique<FFieldIterator_Bindable>(InWidgetBlueprint, InArgs._FieldVisibilityFlags, InArgs._AssignableTo);
+	FieldExpander = MakeUnique<UE::PropertyViewer::FFieldExpander_Default>();
+	FieldExpander->SetExpandObject(UE::PropertyViewer::FFieldExpander_Default::EObjectExpandFlag::UseInstanceClass);
+	FieldExpander->SetExpandScriptStruct(true);
+	FieldExpander->SetExpandFunction(false);
 
 	OnDoubleClicked = InArgs._OnDoubleClicked;
 
 	PropertyViewer = SNew(SPropertyViewer)
 		.FieldIterator(FieldIterator.Get())
+		.FieldExpander(FieldExpander.Get())
 		.PropertyVisibility(SPropertyViewer::EPropertyVisibility::Hidden)
 		.bShowFieldIcon(true)
 		.bSanitizeName(true)
@@ -219,22 +272,32 @@ void SSourceBindingList::AddSource(const FBindingSource& InSource)
 
 void SSourceBindingList::AddSources(TArrayView<const FBindingSource> InSources)
 {
-	for (const FBindingSource& Source : InSources)
+	if (ensure(PropertyViewer))
 	{
-		if (Source.Class != nullptr)
-		{
-			Sources.Add(TPair<FBindingSource, SPropertyViewer::FHandle>(Source, SPropertyViewer::FHandle()));
-		}
-	}
-
-	if (PropertyViewer)
-	{
-		for (TPair<FBindingSource, SPropertyViewer::FHandle>& Source : Sources)
+		const UWidgetBlueprint* WidgetBlueprintPtr = WidgetBlueprint.Get();
+		for (const FBindingSource& Source : InSources)
 		{
-			const UClass* Class = Source.Key.Class;
-			if (Class && Class->ImplementsInterface(UNotifyFieldValueChanged::StaticClass()))
+			if (const UClass* SourceClass = Source.Class.Get())
 			{
-				Source.Value = PropertyViewer->AddContainer(Class);
+				SPropertyViewer::FHandle Handle;
+				if (SourceClass && Source.Class->ImplementsInterface(UNotifyFieldValueChanged::StaticClass()))
+				{
+					UWidget* Widget = nullptr;
+					if (WidgetBlueprintPtr && SourceClass->IsChildOf(UWidget::StaticClass()))
+					{
+						Widget = WidgetBlueprintPtr->WidgetTree->FindWidget(Source.Name);
+					}
+
+					if (Widget)
+					{
+						Handle = PropertyViewer->AddInstance(Widget);
+					}
+					else
+					{
+						Handle = PropertyViewer->AddContainer(SourceClass);
+					}
+				}
+				Sources.Emplace(Source, Handle);
 			}
 		}
 	}
@@ -242,14 +305,62 @@ void SSourceBindingList::AddSources(TArrayView<const FBindingSource> InSources)
 
 FMVVMBlueprintPropertyPath SSourceBindingList::CreateBlueprintPropertyPath(SPropertyViewer::FHandle Handle, TArrayView<const FFieldVariant> FieldPath) const
 {
+	if (FieldPath.Num() < 0)
+	{
+		return FMVVMBlueprintPropertyPath();
+	}
+
+	const UWidgetBlueprint* WidgetBlueprintPtr = WidgetBlueprint.Get();
+	ensure(WidgetBlueprintPtr);
+	if (!WidgetBlueprintPtr)
+	{
+		return FMVVMBlueprintPropertyPath();
+	}
+
 	const TPair<FBindingSource, SPropertyViewer::FHandle>* Source = Sources.FindByPredicate(
 		[Handle](const TPair<FBindingSource, SPropertyViewer::FHandle>& Source)
 		{
 			return Source.Value == Handle;
 		});
-		
+
+	if (!Source)
+	{
+		return FMVVMBlueprintPropertyPath();
+	}
+
+	// Backward, test if the object can be access.
+	//The last property can be a struct variable, inside a struct, inside..., inside an object. 
+	bool bPassFilter = false;
+	TSubclassOf<UObject> AccessorClass = WidgetBlueprintPtr ? WidgetBlueprintPtr->GeneratedClass : nullptr;
+	for (int32 Index = FieldPath.Num() - 1; Index >= 0; --Index)
+	{
+		const FFieldVariant& FieldVariant = FieldPath[Index];
+		const UStruct* OwnerStruct = nullptr;
+		FName FieldName;
+		if (const FProperty* Property = FieldVariant.Get<FProperty>())
+		{
+			OwnerStruct = Property->GetOwnerStruct();
+			FieldName = Property->GetFName();
+		}
+		else if (const UFunction* Function = FieldVariant.Get<UFunction>())
+		{
+			OwnerStruct = Function->GetOwnerClass();
+			FieldName = Function->GetFName();
+		}
+
+		if (const UClass* OwnerClass = Cast<const UClass>(OwnerStruct))
+		{
+			FMVVMAvailableBinding Binding = GEngine->GetEngineSubsystem<UMVVMSubsystem>()->GetAvailableBinding(OwnerClass, FMVVMBindingName(FieldName), AccessorClass);
+			if (Binding.IsValid())
+			{
+				bPassFilter = Private::PassFilter(Binding, OwnerClass, FieldIterator->GetFieldVisibilityFlags(), FieldIterator->GetAssignableTo(), true).IsSet();
+			}
+			break;
+		}
+	}
+
 	FMVVMBlueprintPropertyPath PropertyPath;
-	if (Source != nullptr)
+	if (bPassFilter)
 	{
 		if (Source->Key.ViewModelId.IsValid())
 		{
@@ -261,11 +372,11 @@ FMVVMBlueprintPropertyPath SSourceBindingList::CreateBlueprintPropertyPath(SProp
 		}
 
 		PropertyPath.ResetBasePropertyPath();
-	}
 
-	for (const FFieldVariant& Field : FieldPath)
-	{
-		PropertyPath.AppendBasePropertyPath(FMVVMConstFieldVariant(Field));
+		for (const FFieldVariant& Field : FieldPath)
+		{
+			PropertyPath.AppendBasePropertyPath(FMVVMConstFieldVariant(Field));
+		}
 	}
 
 	return PropertyPath;
diff --git a/Engine/Plugins/Runtime/ModelViewViewModel/Source/ModelViewViewModelEditor/Private/Widgets/SMVVMViewModelBindingListWidget.h b/Engine/Plugins/Runtime/ModelViewViewModel/Source/ModelViewViewModelEditor/Private/Widgets/SMVVMViewModelBindingListWidget.h
--- a/Engine/Plugins/Runtime/ModelViewViewModel/Source/ModelViewViewModelEditor/Private/Widgets/SMVVMViewModelBindingListWidget.h
+++ b/Engine/Plugins/Runtime/ModelViewViewModel/Source/ModelViewViewModelEditor/Private/Widgets/SMVVMViewModelBindingListWidget.h
@@ -3,16 +3,19 @@
 #pragma once
 
 #include "CoreMinimal.h"
+#include "Framework/PropertyViewer/IFieldExpander.h"
 #include "Framework/PropertyViewer/IFieldIterator.h"
 #include "Misc/EnumClassFlags.h"
 #include "MVVMPropertyPath.h"
 #include "Templates/SubclassOf.h"
+#include "Templates/Tuple.h"
 #include "Types/MVVMBindingSource.h"
 #include "UObject/WeakObjectPtr.h"
 #include "Widgets/SCompoundWidget.h"
 #include "Widgets/PropertyViewer/SPropertyViewer.h"
 
 class UWidgetBlueprint;
+class UWidget;
 
 struct FMVVMBlueprintViewModelContext;
 
@@ -37,6 +40,14 @@ namespace UE::MVVM
 	public:
 		FFieldIterator_Bindable(const UWidgetBlueprint* WidgetBlueprint, EFieldVisibility InVisibilityFlags, const FProperty* InAssignableTo = nullptr);
 		virtual TArray<FFieldVariant> GetFields(const UStruct*) const override;
+		EFieldVisibility GetFieldVisibilityFlags() const
+		{
+			return FieldVisibilityFlags;
+		}
+		const FProperty* GetAssignableTo() const
+		{
+			return AssignableTo;
+		}
 
 	private:
 		TWeakObjectPtr<const UWidgetBlueprint> WidgetBlueprint; 
@@ -95,9 +106,12 @@ namespace UE::MVVM
 	private:
 		FOnDoubleClicked OnDoubleClicked;
 		TUniquePtr<FFieldIterator_Bindable> FieldIterator;
+		TUniquePtr<UE::PropertyViewer::FFieldExpander_Default> FieldExpander;
+		using FSourceType = TTuple<FBindingSource, SPropertyViewer::FHandle>;
 		TArray<TPair<FBindingSource, SPropertyViewer::FHandle>> Sources;
 		FMVVMBlueprintPropertyPath SelectedPath;
 		TSharedPtr<SPropertyViewer> PropertyViewer;
+		TWeakObjectPtr<const UWidgetBlueprint> WidgetBlueprint;
 	};
 
 } //namespace UE::MVVM
\ No newline at end of file
