MVVM - Added SSourceBindingList::SetSelectedProperty().

Renamed SSourceBindingList::Clear() to ClearSources().

#rnx

diff --git a/Engine/Plugins/Runtime/ModelViewViewModel/Source/ModelViewViewModelEditor/Private/Widgets/SMVVMFieldSelector.cpp b/Engine/Plugins/Runtime/ModelViewViewModel/Source/ModelViewViewModelEditor/Private/Widgets/SMVVMFieldSelector.cpp
--- a/Engine/Plugins/Runtime/ModelViewViewModel/Source/ModelViewViewModelEditor/Private/Widgets/SMVVMFieldSelector.cpp
+++ b/Engine/Plugins/Runtime/ModelViewViewModel/Source/ModelViewViewModelEditor/Private/Widgets/SMVVMFieldSelector.cpp
@@ -296,7 +296,7 @@ void SFieldSelector::OnViewModelSelected(FBindingSource Source, ESelectInfo::Typ
 	{
 		TArray<FBindingSource> Selection = ViewModelList->GetSelectedItems();
 
-		BindingList->Clear();
+		BindingList->ClearSources();
 		BindingList->AddSources(Selection);
 	}
 
@@ -312,7 +312,7 @@ void SFieldSelector::OnWidgetSelected(FName WidgetName, ESelectInfo::Type)
 	{
 		FBindingSource Source = FBindingSource::CreateForWidget(WidgetBlueprint.Get(), WidgetName);
 
-		BindingList->Clear();
+		BindingList->ClearSources();
 		BindingList->AddSource(Source);
 	}
 
@@ -991,7 +991,7 @@ void SFieldSelector::OnConversionFunctionCategorySelected(TSharedPtr<FConversion
 
 	if (BindingList.IsValid())
 	{
-		BindingList->Clear();
+		BindingList->ClearSources();
 	}
 
 	if (ViewModelList.IsValid())
diff --git a/Engine/Plugins/Runtime/ModelViewViewModel/Source/ModelViewViewModelEditor/Private/Widgets/SMVVMSelectViewModel.cpp b/Engine/Plugins/Runtime/ModelViewViewModel/Source/ModelViewViewModelEditor/Private/Widgets/SMVVMSelectViewModel.cpp
--- a/Engine/Plugins/Runtime/ModelViewViewModel/Source/ModelViewViewModelEditor/Private/Widgets/SMVVMSelectViewModel.cpp
+++ b/Engine/Plugins/Runtime/ModelViewViewModel/Source/ModelViewViewModelEditor/Private/Widgets/SMVVMSelectViewModel.cpp
@@ -135,7 +135,7 @@ void SMVVMSelectViewModel::Construct(const FArguments& InArgs, const UWidgetBlue
 
 void SMVVMSelectViewModel::HandleClassPicked(UClass* ClassPicked)
 {
-	BindingListWidget->Clear();
+	BindingListWidget->ClearSources();
 	SelectedClass.Reset();
 
 	if (Private::IsValidViewModel(ClassPicked))
diff --git a/Engine/Plugins/Runtime/ModelViewViewModel/Source/ModelViewViewModelEditor/Private/Widgets/SMVVMViewModelBindingListWidget.cpp b/Engine/Plugins/Runtime/ModelViewViewModel/Source/ModelViewViewModelEditor/Private/Widgets/SMVVMViewModelBindingListWidget.cpp
--- a/Engine/Plugins/Runtime/ModelViewViewModel/Source/ModelViewViewModelEditor/Private/Widgets/SMVVMViewModelBindingListWidget.cpp
+++ b/Engine/Plugins/Runtime/ModelViewViewModel/Source/ModelViewViewModelEditor/Private/Widgets/SMVVMViewModelBindingListWidget.cpp
@@ -64,12 +64,14 @@ namespace Private
 					return TOptional<FFieldVariant>();
 				}
 
-				if (AssignableTo != nullptr && !BindingHelper::ArePropertiesCompatible(BindingHelper::GetReturnProperty(FieldVariant.GetFunction()), AssignableTo))
+				const UFunction* Function = FieldVariant.GetFunction();
+				if (Function != nullptr && AssignableTo != nullptr && 
+					!BindingHelper::ArePropertiesCompatible(BindingHelper::GetReturnProperty(Function), AssignableTo))
 				{
 					return TOptional<FFieldVariant>();
 				}
 
-				return FFieldVariant(FieldVariant.GetFunction());
+				return FFieldVariant(Function);
 			}
 			else if (FieldVariant.IsProperty())
 			{
@@ -78,6 +80,13 @@ namespace Private
 					return TOptional<FFieldVariant>();
 				}
 
+				const FProperty* Property = FieldVariant.GetProperty();
+				if (Property != nullptr && AssignableTo != nullptr && 
+					!BindingHelper::ArePropertiesCompatible(Property, AssignableTo))
+				{
+					return TOptional<FFieldVariant>();
+				}
+
 				// If the path ends with the object property, then it needs to follow the writable rule
 				if (bDoObjectProperty || !CastField<FObjectPropertyBase>(FieldVariant.GetProperty()))
 				{
@@ -87,12 +96,7 @@ namespace Private
 					}
 				}
 
-				if (AssignableTo != nullptr && !BindingHelper::ArePropertiesCompatible(FieldVariant.GetProperty(), AssignableTo))
-				{
-					return TOptional<FFieldVariant>();
-				}
-
-				return FFieldVariant(FieldVariant.GetProperty());
+				return FFieldVariant(Property);
 			}
 		}
 
@@ -194,19 +198,16 @@ void SSourceBindingList::Construct(const FArguments& InArgs, const UWidgetBluepr
 		.OnSelectionChanged(this, &SSourceBindingList::HandleSelectionChanged)
 		.OnDoubleClicked(this, &SSourceBindingList::HandleDoubleClicked);
 
-	if (InArgs._InitialSource.IsValid())
-	{
-		AddSources(MakeArrayView(&InArgs._InitialSource, 1));
-	}
-
 	ChildSlot
 	[
 		PropertyViewer.ToSharedRef()
 	];
 }
 
-void SSourceBindingList::Clear()
+void SSourceBindingList::ClearSources()
 {
+	SelectedPath = FMVVMBlueprintPropertyPath();
+
 	Sources.Reset();
 	if (PropertyViewer)
 	{
@@ -409,6 +410,47 @@ FMVVMBlueprintPropertyPath SSourceBindingList::GetSelectedProperty() const
 	return SelectedPath;
 }
 
+void SSourceBindingList::SetSelectedProperty(const FMVVMBlueprintPropertyPath& Property)
+{
+	if (!PropertyViewer.IsValid())
+	{
+		return;
+	}
+
+	SPropertyViewer::FHandle SelectedHandle;
+	for (TPair<FBindingSource, SPropertyViewer::FHandle>& Source : Sources)
+	{
+		if (Property.IsFromViewModel() && Source.Key.ViewModelId == Property.GetViewModelId() ||
+			Property.IsFromWidget() && Source.Key.Name == Property.GetWidgetName())
+		{
+			SelectedHandle = Source.Value;
+			break;
+		}
+	}
+
+	TArray<FFieldVariant, TMemStackAllocator<>> FieldPath;
+	if (SelectedHandle.IsValid())
+	{
+		TArray<FMVVMConstFieldVariant> FieldVariants = Property.GetFields();
+		FieldPath.Reserve(FieldVariants.Num());
+
+		for (const FMVVMConstFieldVariant& Variant : FieldVariants)
+		{
+			FFieldVariant& Field = FieldPath.AddDefaulted_GetRef();
+			if (Variant.IsFunction())
+			{
+				Field = FFieldVariant(Variant.GetFunction());
+			}
+			else if (Variant.IsProperty())
+			{
+				Field = FFieldVariant(Variant.GetProperty());
+			}
+		}
+	}
+	
+	PropertyViewer->SetSelection(SelectedHandle, FieldPath);
+}
+
 } // namespace UE::MVVM
 
 #undef LOCTEXT_NAMESPACE
diff --git a/Engine/Plugins/Runtime/ModelViewViewModel/Source/ModelViewViewModelEditor/Private/Widgets/SMVVMViewModelBindingListWidget.h b/Engine/Plugins/Runtime/ModelViewViewModel/Source/ModelViewViewModelEditor/Private/Widgets/SMVVMViewModelBindingListWidget.h
--- a/Engine/Plugins/Runtime/ModelViewViewModel/Source/ModelViewViewModelEditor/Private/Widgets/SMVVMViewModelBindingListWidget.h
+++ b/Engine/Plugins/Runtime/ModelViewViewModel/Source/ModelViewViewModelEditor/Private/Widgets/SMVVMViewModelBindingListWidget.h
@@ -64,7 +64,6 @@ namespace UE::MVVM
 		DECLARE_DELEGATE_OneParam(FOnDoubleClicked, const FMVVMBlueprintPropertyPath&);
 	
 		SLATE_BEGIN_ARGS(SSourceBindingList) {}
-			SLATE_ARGUMENT(FBindingSource, InitialSource)
 			SLATE_ARGUMENT_DEFAULT(bool, ShowSearchBox) = false;
 			SLATE_ARGUMENT_DEFAULT(EFieldVisibility, FieldVisibilityFlags) = EFieldVisibility::All;
 			SLATE_ARGUMENT_DEFAULT(bool, EnableSelection) = true;
@@ -81,7 +80,7 @@ namespace UE::MVVM
 
 		void Construct(const FArguments& InArgs, const UWidgetBlueprint* WidgetBlueprint);
 
-		void Clear();
+		void ClearSources();
 
 		void AddSource(UClass* Class, FName Name, FGuid Guid);
 		void AddSource(const FBindingSource& InSource);
@@ -91,7 +90,8 @@ namespace UE::MVVM
 		void AddWidgets(TArrayView<const UWidget*> Widgets);
 		void AddViewModels(TArrayView<const FMVVMBlueprintViewModelContext> ViewModels);
 
-		FMVVMBlueprintPropertyPath GetSelectedProperty() const;
+		FMVVMBlueprintPropertyPath GetSelectedProperty() const; 
+		void SetSelectedProperty(const FMVVMBlueprintPropertyPath& Property);
 
 		void SetRawFilterText(const FText& InFilterText);
 
