CIS fix. Added parentheses to disambiguate "&& inside ||" as reported by Clang.

#rnx

diff --git a/Engine/Plugins/Runtime/ModelViewViewModel/Source/ModelViewViewModelEditor/Private/Widgets/SMVVMViewModelBindingListWidget.cpp b/Engine/Plugins/Runtime/ModelViewViewModel/Source/ModelViewViewModelEditor/Private/Widgets/SMVVMViewModelBindingListWidget.cpp
--- a/Engine/Plugins/Runtime/ModelViewViewModel/Source/ModelViewViewModelEditor/Private/Widgets/SMVVMViewModelBindingListWidget.cpp
+++ b/Engine/Plugins/Runtime/ModelViewViewModel/Source/ModelViewViewModelEditor/Private/Widgets/SMVVMViewModelBindingListWidget.cpp
@@ -420,8 +420,8 @@ void SSourceBindingList::SetSelectedProperty(const FMVVMBlueprintPropertyPath& P
 	SPropertyViewer::FHandle SelectedHandle;
 	for (TPair<FBindingSource, SPropertyViewer::FHandle>& Source : Sources)
 	{
-		if (Property.IsFromViewModel() && Source.Key.ViewModelId == Property.GetViewModelId() ||
-			Property.IsFromWidget() && Source.Key.Name == Property.GetWidgetName())
+		if ((Property.IsFromViewModel() && Source.Key.ViewModelId == Property.GetViewModelId()) ||
+			(Property.IsFromWidget() && Source.Key.Name == Property.GetWidgetName()))
 		{
 			SelectedHandle = Source.Value;
 			break;
