Added missing FMemMark.

#rnx

diff --git a/Engine/Plugins/Runtime/ModelViewViewModel/Source/ModelViewViewModelEditor/Private/Widgets/SMVVMViewModelBindingListWidget.cpp b/Engine/Plugins/Runtime/ModelViewViewModel/Source/ModelViewViewModelEditor/Private/Widgets/SMVVMViewModelBindingListWidget.cpp
--- a/Engine/Plugins/Runtime/ModelViewViewModel/Source/ModelViewViewModelEditor/Private/Widgets/SMVVMViewModelBindingListWidget.cpp
+++ b/Engine/Plugins/Runtime/ModelViewViewModel/Source/ModelViewViewModelEditor/Private/Widgets/SMVVMViewModelBindingListWidget.cpp
@@ -9,6 +9,7 @@
 #include "Editor.h"
 #include "Engine/Engine.h"
 #include "FieldNotification/IFieldValueChanged.h"
+#include "Misc/MemStack.h"
 #include "Misc/Optional.h"
 #include "MVVMBlueprintViewModelContext.h"
 #include "MVVMDeveloperProjectSettings.h"
@@ -428,6 +429,7 @@ void SSourceBindingList::SetSelectedProperty(const FMVVMBlueprintPropertyPath& P
 		}
 	}
 
+	FMemMark Mark(FMemStack::Get());
 	TArray<FFieldVariant, TMemStackAllocator<>> FieldPath;
 	if (SelectedHandle.IsValid())
 	{
