MVVM - The visibility of the details panel in the Bindings panel is now togglable.

#rnx

diff --git a/Engine/Plugins/Runtime/ModelViewViewModel/Source/ModelViewViewModelEditor/Private/Widgets/SMVVMViewBindingPanel.cpp b/Engine/Plugins/Runtime/ModelViewViewModel/Source/ModelViewViewModelEditor/Private/Widgets/SMVVMViewBindingPanel.cpp
--- a/Engine/Plugins/Runtime/ModelViewViewModel/Source/ModelViewViewModelEditor/Private/Widgets/SMVVMViewBindingPanel.cpp
+++ b/Engine/Plugins/Runtime/ModelViewViewModel/Source/ModelViewViewModelEditor/Private/Widgets/SMVVMViewBindingPanel.cpp
@@ -186,8 +186,7 @@ TSharedRef<SWidget> SBindingsPanel::CreateDrawerDockButton()
 {
 	if (bIsDrawerTab)
 	{
-		return
-			SNew(SButton)
+		return SNew(SButton)
 			.ButtonStyle(FAppStyle::Get(), "SimpleButton")
 			.ToolTipText(LOCTEXT("BindingDockInLayout_Tooltip", "Docks the binding drawer in tab."))
 			.ContentPadding(FMargin(1.f, 0.f))
@@ -195,22 +194,22 @@ TSharedRef<SWidget> SBindingsPanel::CreateDrawerDockButton()
 			[
 				SNew(SHorizontalBox)
 				+ SHorizontalBox::Slot()
-			.AutoWidth()
-			.VAlign(VAlign_Center)
-			.Padding(4.0, 0.0f)
-			[
-				SNew(SImage)
-				.ColorAndOpacity(FSlateColor::UseForeground())
-			.Image(FAppStyle::Get().GetBrush("Icons.Layout"))
-			]
-		+ SHorizontalBox::Slot()
-			.VAlign(VAlign_Center)
-			.Padding(4.0, 0.0f)
-			[
-				SNew(STextBlock)
-				.Text(LOCTEXT("DockInLayout", "Dock in Layout"))
-			.ColorAndOpacity(FSlateColor::UseForeground())
-			]
+				.AutoWidth()
+				.VAlign(VAlign_Center)
+				.Padding(4.0, 0.0f)
+				[
+					SNew(SImage)
+					.ColorAndOpacity(FSlateColor::UseForeground())
+					.Image(FAppStyle::Get().GetBrush("Icons.Layout"))
+				]
+				+ SHorizontalBox::Slot()
+				.VAlign(VAlign_Center)
+				.Padding(4.0, 0.0f)
+				[
+					SNew(STextBlock)
+					.Text(LOCTEXT("DockInLayout", "Dock in Layout"))
+					.ColorAndOpacity(FSlateColor::UseForeground())
+				]
 			];
 	}
 
@@ -302,17 +301,18 @@ TSharedRef<SWidget> SBindingsPanel::GenerateEditViewWidget()
 		ToolbarBuilderGlobal.AddWidget(CreateDrawerDockButton());
 
 		ToolbarBuilderGlobal.BeginSection("Options");
-		ToolbarBuilderGlobal.AddComboButton(
+
+		ToolbarBuilderGlobal.AddToolBarButton(
 			FUIAction(
-				FExecuteAction(),
+				FExecuteAction::CreateSP(this, &SBindingsPanel::ToggleDetailsVisibility),
 				FCanExecuteAction(),
-				FGetActionCheckState()
+				FGetActionCheckState::CreateSP(this, &SBindingsPanel::GetDetailsVisibleCheckState)
 			),
-			FOnGetContent::CreateSP(this, &SBindingsPanel::GenerateSettingsMenu),
-			LOCTEXT("Settings", "Settings"),
-			LOCTEXT("SettingsTooltip", "ModelView Settings"),
-			FSlateIcon(FAppStyle::GetAppStyleSetName(), "LevelEditor.GameSettings"),
-			false
+			"ToggleDetails",
+			LOCTEXT("Details", "Details"),
+			LOCTEXT("DetailsToolTip", "Open Details View"),
+			FSlateIcon(FAppStyle::GetAppStyleSetName(), "WorldBrowser.DetailsButtonBrush"),
+			EUserInterfaceActionType::ToggleButton
 		);
 		ToolbarBuilderGlobal.EndSection();
 	}
@@ -413,7 +413,6 @@ TSharedRef<SWidget> SBindingsPanel::GenerateEditViewWidget()
 				.Value(0.25f)
 				[
 					SAssignNew(DetailContainer, SBorder)
-					.Visibility(this, &SBindingsPanel::GetVisibility, true)
 					[
 						DetailsView.ToSharedRef()
 					]
@@ -422,6 +421,24 @@ TSharedRef<SWidget> SBindingsPanel::GenerateEditViewWidget()
 		];
 }
 
+ECheckBoxState SBindingsPanel::GetDetailsVisibleCheckState() const
+{
+	if (DetailContainer.IsValid())
+	{
+		return DetailContainer->GetVisibility() == EVisibility::Visible ? ECheckBoxState::Checked : ECheckBoxState::Unchecked;
+	}
+	return ECheckBoxState::Unchecked;
+}
+
+void SBindingsPanel::ToggleDetailsVisibility() 
+{
+	if (DetailContainer.IsValid())
+	{
+		EVisibility NewVisibility = GetDetailsVisibleCheckState() == ECheckBoxState::Checked ? EVisibility::Collapsed : EVisibility::Visible;
+		return DetailContainer->SetVisibility(NewVisibility);
+	}
+}
+
 EVisibility SBindingsPanel::GetVisibility(bool bVisibleWithBindings) const
 {
 	if (UMVVMWidgetBlueprintExtension_View* MVVMExtensionPtr = MVVMExtension.Get())
diff --git a/Engine/Plugins/Runtime/ModelViewViewModel/Source/ModelViewViewModelEditor/Private/Widgets/SMVVMViewBindingPanel.h b/Engine/Plugins/Runtime/ModelViewViewModel/Source/ModelViewViewModelEditor/Private/Widgets/SMVVMViewBindingPanel.h
--- a/Engine/Plugins/Runtime/ModelViewViewModel/Source/ModelViewViewModelEditor/Private/Widgets/SMVVMViewBindingPanel.h
+++ b/Engine/Plugins/Runtime/ModelViewViewModel/Source/ModelViewViewModelEditor/Private/Widgets/SMVVMViewBindingPanel.h
@@ -18,6 +18,8 @@ class SBorder;
 class UBlueprintExtension;
 class UMVVMWidgetBlueprintExtension_View;
 
+enum class ECheckBoxState : uint8;
+
 namespace UE::MVVM
 {
 class SBindingsList;
@@ -60,6 +62,9 @@ private:
 
 	EVisibility GetVisibility(bool bVisibleWithBindings) const;
 
+	ECheckBoxState GetDetailsVisibleCheckState() const;
+	void ToggleDetailsVisibility();
+
 private:
 	TWeakPtr<FWidgetBlueprintEditor> WeakBlueprintEditor;
 	TSharedPtr<SBindingsList> BindingsList;
