MVVM - SFieldSelector now displays conversion functions for widget properties again.

Conversion function entries are now sorted.

Conversion function categories display how many functions are inside of them.

The current selection of the field is now correctly displayed in the menu.

#rnx

diff --git a/Engine/Plugins/Runtime/ModelViewViewModel/Source/ModelViewViewModelEditor/Private/Widgets/SMVVMFieldSelector.cpp b/Engine/Plugins/Runtime/ModelViewViewModel/Source/ModelViewViewModelEditor/Private/Widgets/SMVVMFieldSelector.cpp
--- a/Engine/Plugins/Runtime/ModelViewViewModel/Source/ModelViewViewModelEditor/Private/Widgets/SMVVMFieldSelector.cpp
+++ b/Engine/Plugins/Runtime/ModelViewViewModel/Source/ModelViewViewModelEditor/Private/Widgets/SMVVMFieldSelector.cpp
@@ -21,6 +21,7 @@
 #include "Widgets/Images/SImage.h"
 #include "Widgets/Input/SComboButton.h"
 #include "Widgets/Input/SSearchBox.h"
+#include "Widgets/Layout/SScrollBox.h"
 #include "Widgets/SNullWidget.h"
 
 #define LOCTEXT_NAMESPACE "MVVMFieldSelector"
@@ -67,18 +68,9 @@ void SFieldSelector::Construct(const FArguments& InArgs, const UWidgetBlueprint*
 	CachedSelectedField = SelectedField.Get();
 
 	ShowConversionFunctions = InArgs._ShowConversionFunctions;
-
-	const UFunction* CurrentFunction = nullptr;
-	if (ShowConversionFunctions.Get(false))
-	{
-		SelectedConversionFunction = InArgs._SelectedConversionFunction;
-		check(SelectedConversionFunction.IsSet());
-
-		CachedSelectedConversionFunction = SelectedConversionFunction.Get(nullptr);
-
-		OnConversionFunctionSelectionChangedDelegate = InArgs._OnConversionFunctionSelectionChanged;
-		check(OnConversionFunctionSelectionChangedDelegate.IsBound());
-	}
+	SelectedConversionFunction = InArgs._SelectedConversionFunction;
+	CachedSelectedConversionFunction = SelectedConversionFunction.Get(nullptr);
+	OnConversionFunctionSelectionChangedDelegate = InArgs._OnConversionFunctionSelectionChanged;
 
 	TSharedPtr<SHorizontalBox> SourceEntryBox;
 
@@ -132,6 +124,11 @@ void SFieldSelector::Construct(const FArguments& InArgs, const UWidgetBlueprint*
 		]
 	];
 	
+	if (InArgs._Source.IsValid())
+	{
+		FixedSource = InArgs._Source;
+	}
+
 	if (InArgs._ShowSource)
 	{
 		SourceEntryBox->AddSlot()
@@ -153,10 +150,6 @@ void SFieldSelector::Construct(const FArguments& InArgs, const UWidgetBlueprint*
 				.Image(FAppStyle::Get().GetBrush("Icons.ChevronRight"))
 			];					
 	}
-	else
-	{
-		FixedSource = InArgs._Source;
-	}
 
 	SourceEntryBox->AddSlot()
 		.Padding(0, 0, 8, 0)
@@ -239,8 +232,17 @@ bool SFieldSelector::IsClearEnabled() const
 
 FReply SFieldSelector::OnClearClicked()
 {
-	SetPropertySelection(FMVVMBlueprintPropertyPath());
-	SetConversionFunctionSelection(nullptr);
+	if (FixedSource.IsSet())
+	{
+		FMVVMBlueprintPropertyPath Path = CachedSelectedField;
+		Path.ResetBasePropertyPath();
+		SetPropertySelection(Path);
+	}
+	else
+	{
+		SetPropertySelection(FMVVMBlueprintPropertyPath());
+		SetConversionFunctionSelection(nullptr);
+	}
 	return FReply::Handled();
 }
 
@@ -292,34 +294,68 @@ void SFieldSelector::OnSearchBoxTextChanged(const FText& NewText)
 
 void SFieldSelector::OnViewModelSelected(FBindingSource Source, ESelectInfo::Type)
 {
-	if (BindingList.IsValid())
+	if (!bIsMenuInitialized)
 	{
-		TArray<FBindingSource> Selection = ViewModelList->GetSelectedItems();
+		return;
+	}
 
+	if (BindingList.IsValid())
+	{
 		BindingList->ClearSources();
-		BindingList->AddSources(Selection);
+	}
+
+	TArray<FBindingSource> Selection = ViewModelList->GetSelectedItems();
+	if (Selection.Num() == 0)
+	{
+		return;
 	}
 
 	if (ConversionFunctionCategoryTree.IsValid())
 	{
 		ConversionFunctionCategoryTree->ClearSelection();
 	}
+
+	if (BindingList.IsValid())
+	{
+		BindingList->AddSources(Selection);
+	}
 }
 
 void SFieldSelector::OnWidgetSelected(FName WidgetName, ESelectInfo::Type)
 {
-	if (BindingList.IsValid() && WidgetBlueprint.IsValid())
+	if (!bIsMenuInitialized)
 	{
-		FBindingSource Source = FBindingSource::CreateForWidget(WidgetBlueprint.Get(), WidgetName);
+		return;
+	}
 
+	if (BindingList.IsValid())
+	{
 		BindingList->ClearSources();
-		BindingList->AddSource(Source);
+	}
+
+	TArray<FName> Selection = WidgetList->GetSelectedWidgets();
+	if (Selection.Num() == 0)
+	{
+		return;
 	}
 
 	if (ConversionFunctionCategoryTree.IsValid())
 	{
 		ConversionFunctionCategoryTree->ClearSelection();
 	}
+
+	if (BindingList.IsValid())
+	{
+		const UWidgetBlueprint* WidgetBP = WidgetBlueprint.Get();
+
+		TArray<FBindingSource> Sources;
+		Algo::Transform(Selection, Sources, [WidgetBP](const FName& WidgetName)
+		{
+			return FBindingSource::CreateForWidget(WidgetBP, WidgetName);
+		});
+
+		BindingList->AddSources(Sources);
+	}
 }
 
 TSharedRef<ITableRow> SFieldSelector::GenerateViewModelRow(FBindingSource ViewModel, const TSharedRef<STableViewBase>& OwnerTable) const
@@ -350,9 +386,9 @@ TSharedRef<ITableRow> SFieldSelector::GenerateViewModelRow(FBindingSource ViewMo
 		];
 }
 
-bool SFieldSelector::FilterConversionFunctionCategoryChildren(const TArray<FString>& FilterStrings, const TArray<TSharedPtr<FConversionFunctionItem>>& SourceArray, TArray<TSharedPtr<FConversionFunctionItem>>& OutDestArray)
+int32 SFieldSelector::FilterConversionFunctionCategoryChildren(const TArray<FString>& FilterStrings, const TArray<TSharedPtr<FConversionFunctionItem>>& SourceArray, TArray<TSharedPtr<FConversionFunctionItem>>& OutDestArray)
 {
-	bool bAnyAdded = false;
+	int32 NumFunctions = 0;
 
 	for (const TSharedPtr<FConversionFunctionItem>& SourceItem : SourceArray)
 	{
@@ -361,7 +397,7 @@ bool SFieldSelector::FilterConversionFunctionCategoryChildren(const TArray<FStri
 		// check if our name matches the filters
 		bool bMatchesFilters = true;
 
-		const FString ItemName = SourceItem->Function != nullptr ? SourceItem->Function->GetName() : SourceItem->CategoryName;
+		const FString ItemName = SourceItem->Function != nullptr ? SourceItem->Function->GetName() : FString();
 		
 		for (const FString& Filter : FilterStrings)
 		{
@@ -372,36 +408,37 @@ bool SFieldSelector::FilterConversionFunctionCategoryChildren(const TArray<FStri
 			}
 		}
 
-		bool bChildrenMatch = false;
-
+		int32 NumChildren = 0;
 		if (bMatchesFilters)
 		{
-			// if we match the filter, then we always want to add all our children
-			FilteredChildren = SourceItem->Children;
+			ensureAlways(SourceItem->Function != nullptr);
+			NumChildren = 1;
 		}
 		else
 		{
-			bChildrenMatch = FilterConversionFunctionCategoryChildren(FilterStrings, SourceItem->Children, FilteredChildren);
+			// if we don't match, then we still want to check all our children
+			NumChildren = FilterConversionFunctionCategoryChildren(FilterStrings, SourceItem->Children, FilteredChildren);
 		}
 
 		// then add this item to the destination array
-		if (bMatchesFilters || bChildrenMatch)
+		if (NumChildren > 0)
 		{
 			TSharedPtr<FConversionFunctionItem>& NewItem = OutDestArray.Add_GetRef(MakeShared<FConversionFunctionItem>());
 			NewItem->CategoryName = SourceItem->CategoryName;
-			NewItem->Function = SourceItem->Function;
+			NewItem->Function = SourceItem->Function; 
 			NewItem->Children = FilteredChildren;
+			NewItem->NumFunctions = NumChildren;
 
-			bAnyAdded = true;
+			NumFunctions += NewItem->NumFunctions;
 		}
 	}
 
-	return bAnyAdded;
+	return NumFunctions;
 }
 
 void SFieldSelector::FilterConversionFunctionCategories()
 {
-	FilteredConversionFunctionCategories.Reset();
+	FilteredConversionFunctionRoot.Reset();
 
 	TArray<FString> FilterStrings;
 	if (SearchBox.IsValid())
@@ -412,11 +449,176 @@ void SFieldSelector::FilterConversionFunctionCategories()
 	if (FilterStrings.IsEmpty())
 	{
 		// don't bother filtering if we don't have any search terms
-		FilteredConversionFunctionCategories = ConversionFunctionCategories;
+		FilteredConversionFunctionRoot = ConversionFunctionRoot;
 		return;
 	}
+
+	TSharedPtr<FConversionFunctionItem> RootItem = FilteredConversionFunctionRoot.Add_GetRef(MakeShared<FConversionFunctionItem>());
+	RootItem->CategoryName = TEXT("Conversion Functions");
 	
-	FilterConversionFunctionCategoryChildren(FilterStrings, ConversionFunctionCategories, FilteredConversionFunctionCategories);
+	RootItem->NumFunctions = FilterConversionFunctionCategoryChildren(FilterStrings, ConversionFunctionRoot[0]->Children, FilteredConversionFunctionRoot[0]->Children);
+	
+	ExpandAll(FilteredConversionFunctionRoot, true);
+}
+
+int32 SFieldSelector::SortConversionFunctionItemsRecursive(TArray<TSharedPtr<FConversionFunctionItem>>& Items)
+{
+	int32 NumFound = 0;
+
+	Items.Sort([](const TSharedPtr<FConversionFunctionItem>& A, const TSharedPtr<FConversionFunctionItem>& B)
+	{
+		if (!A->CategoryName.IsEmpty() && !B->CategoryName.IsEmpty())
+		{
+			return A->CategoryName < B->CategoryName;
+		}
+		if (!A->CategoryName.IsEmpty() && B->CategoryName.IsEmpty())
+		{
+			return true;
+		}
+		if (A->CategoryName.IsEmpty() && !B->CategoryName.IsEmpty())
+		{
+			return false;
+		}
+		if (A->Function != nullptr && B->Function != nullptr)
+		{
+			return A->Function->GetDisplayNameText().CompareTo(B->Function->GetDisplayNameText()) <= 0;
+		}
+		return true;
+	});
+
+	for (const TSharedPtr<FConversionFunctionItem>& Item : Items)
+	{
+		NumFound += SortConversionFunctionItemsRecursive(Item->Children);
+
+		if (Item->Function != nullptr)
+		{
+			NumFound += 1;
+		}
+	}
+
+	return NumFound;
+}
+
+void SFieldSelector::GenerateConversionFunctionItems()
+{
+	UMVVMEditorSubsystem* Subsystem = GEditor->GetEditorSubsystem<UMVVMEditorSubsystem>();
+	TArray<UFunction*> AllConversionFunctions = Subsystem->GetAvailableConversionFunctions(WidgetBlueprint.Get(), FMVVMBlueprintPropertyPath(), FMVVMBlueprintPropertyPath());
+
+	// remove all incompatible conversion functions
+	for (int32 Idx = 0; Idx < AllConversionFunctions.Num(); ++Idx)
+	{
+		const UFunction* Function = AllConversionFunctions[Idx];
+		const FProperty* ReturnProperty = BindingHelper::GetReturnProperty(Function);
+		const FProperty* AssignToProperty = AssignableTo.Get(nullptr);
+
+		if (AssignToProperty != nullptr && !BindingHelper::ArePropertiesCompatible(ReturnProperty, AssignToProperty))
+		{
+			AllConversionFunctions.RemoveAt(Idx);
+			--Idx;
+		}
+	}
+
+	TSharedPtr<FConversionFunctionItem> CurrentSelectedItem;
+
+	TArray<FString> SubCategories;
+
+	ConversionFunctionRoot.Reset();
+	TSharedPtr<FConversionFunctionItem>& RootItem = ConversionFunctionRoot.Add_GetRef(MakeShared<FConversionFunctionItem>());
+	RootItem->CategoryName = TEXT("Conversion Functions");
+
+	for (const UFunction* Function : AllConversionFunctions)
+	{
+		FText CategoryName = Function->GetMetaDataText("Category");
+		if (CategoryName.IsEmpty())
+		{
+			TSharedPtr<FConversionFunctionItem>& Item = RootItem->Children.Add_GetRef(MakeShared<FConversionFunctionItem>());
+			Item->Function = Function;
+			Item->NumFunctions = 1;
+			RootItem->NumFunctions += 1;
+			continue;
+		}
+
+		SubCategories.Reset();
+		CategoryName.ToString().ParseIntoArray(SubCategories, TEXT("|"));
+
+		TSharedPtr<FConversionFunctionItem> ParentItem = RootItem;
+
+		// create items for the entire category path
+		// eg. "Math|Boolean|AND" 
+		// Math 
+		//   > Boolean
+		//     > AND
+		for (const FString& SubCategory : SubCategories)
+		{
+			const FString Trimmed = SubCategory.TrimStartAndEnd();
+
+			ParentItem->NumFunctions += 1;
+			ParentItem = FindOrCreateItemForCategory(ParentItem->Children, Trimmed);
+		}
+
+		TSharedPtr<FConversionFunctionItem>& Item = ParentItem->Children.Add_GetRef(MakeShared<FConversionFunctionItem>());
+		Item->Function = Function;
+		Item->NumFunctions = 1;
+		ParentItem->NumFunctions += 1;
+	}
+
+	int32 NumItems = SortConversionFunctionItemsRecursive(ConversionFunctionRoot);
+	ensure(NumItems == RootItem->NumFunctions);
+}
+
+void SFieldSelector::ExpandAllToItem(const UFunction* Function)
+{
+	TArray<TSharedPtr<FConversionFunctionItem>> Path;
+
+	FText CategoryName = Function->GetMetaDataText("Category");
+	if (CategoryName.IsEmpty())
+	{
+		Path.Add(FilteredConversionFunctionRoot[0]);
+	}
+	else
+	{
+		TArray<FString> SubCategories;
+		CategoryName.ToString().ParseIntoArray(SubCategories, TEXT("|"));
+
+		TSharedPtr<FConversionFunctionItem> CurrentParent = FilteredConversionFunctionRoot[0];
+
+		for (const FString& SubCategory : SubCategories)
+		{
+			const FString Trimmed = SubCategory.TrimStartAndEnd();
+
+			TSharedPtr<FConversionFunctionItem>* FoundItem = 
+				CurrentParent->Children.FindByPredicate([Trimmed, Function](const TSharedPtr<FConversionFunctionItem>& Item)
+				{
+					return Item->CategoryName == Trimmed || Item->Function == Function;
+				});
+
+			if (FoundItem != nullptr)
+			{
+				Path.Add(*FoundItem);
+				CurrentParent = *FoundItem;
+			}
+		}
+	}
+
+	if (Path.Num() > 0)
+	{
+		ConversionFunctionCategoryTree->SetItemExpansion(FilteredConversionFunctionRoot[0], true);
+
+		ExpandAll(Path, false);
+	}
+}
+
+void SFieldSelector::ExpandAll(const TArray<TSharedPtr<FConversionFunctionItem>>& Items, bool bRecursive)
+{
+	for (const TSharedPtr<FConversionFunctionItem>& Item : Items)
+	{
+		ConversionFunctionCategoryTree->SetItemExpansion(Item, true);
+
+		if (bRecursive)
+		{
+			ExpandAll(Item->Children, bRecursive);
+		}
+	}
 }
 
 void SFieldSelector::FilterViewModels()
@@ -467,12 +669,6 @@ void SFieldSelector::FilterViewModels()
 
 TSharedRef<SWidget> SFieldSelector::CreateSourcePanel()
 {
-	ConversionFunctionCategoryTree.Reset();
-	ConversionFunctionCategories.Reset();
-	ViewModelSources.Reset();
-	ViewModelList.Reset();
-	WidgetList.Reset();
-
 	if (!WidgetBlueprint.IsValid())
 	{
 		return SNullWidget::NullWidget;
@@ -482,27 +678,14 @@ TSharedRef<SWidget> SFieldSelector::CreateSourcePanel()
 	TSharedPtr<SWidget> SourcePicker;
 	if (bViewModelProperty)
 	{
-		FBindingSource SelectedSource;
-
-		UMVVMEditorSubsystem* EditorSubsystem = GEditor->GetEditorSubsystem<UMVVMEditorSubsystem>();
-		if (const UMVVMBlueprintView* View = EditorSubsystem->GetView(WidgetBlueprint.Get()))
+		if (FixedSource.IsSet())
 		{
-			const TArrayView<const FMVVMBlueprintViewModelContext> ViewModels = View->GetViewModels();
-			ViewModelSources.Reserve(ViewModels.Num());
-
-			for (const FMVVMBlueprintViewModelContext& ViewModel : ViewModels)
-			{
-				FBindingSource& Source = ViewModelSources.AddDefaulted_GetRef();
-				Source.Class = ViewModel.GetViewModelClass();
-				Source.Name = ViewModel.GetViewModelName();
-				Source.DisplayName = ViewModel.GetDisplayName();
-				Source.ViewModelId = ViewModel.GetViewModelId();
-
-				if (Source.ViewModelId == CachedSelectedField.GetViewModelId())
-				{
-					SelectedSource = Source;
-				}
-			}
+			ViewModelSources.Add(FixedSource.GetValue());
+		}
+		else
+		{
+			UMVVMEditorSubsystem* EditorSubsystem = GEditor->GetEditorSubsystem<UMVVMEditorSubsystem>();
+			ViewModelSources = EditorSubsystem->GetAllViewModels(WidgetBlueprint.Get());
 		}
 
 		FilterViewModels();
@@ -513,6 +696,15 @@ TSharedRef<SWidget> SFieldSelector::CreateSourcePanel()
 			.SelectionMode(ESelectionMode::Multi)
 			.OnSelectionChanged(this, &SFieldSelector::OnViewModelSelected);
 
+		FBindingSource SelectedSource;
+		for (const FBindingSource& Source : FilteredViewModelSources)
+		{
+			if (Source.ViewModelId == CachedSelectedField.GetViewModelId())
+			{
+				SelectedSource = Source;
+			}
+		}
+
 		if (SelectedSource.IsValid())
 		{
 			ViewModelList->SetItemSelection(SelectedSource, true);
@@ -522,104 +714,49 @@ TSharedRef<SWidget> SFieldSelector::CreateSourcePanel()
 	}
 	else
 	{
+		TArray<FName> ShowOnly;
+		if (FixedSource.IsSet())
+		{
+			ShowOnly.Add(FixedSource.GetValue().Name);
+		}
+
 		WidgetList = SNew(SReadOnlyHierarchyView, WidgetBlueprint.Get())
 			.OnSelectionChanged(this, &SFieldSelector::OnWidgetSelected)
 			.SelectionMode(ESelectionMode::Multi)
-			.ShowSearch(false);
+			.ShowSearch(false)
+			.ShowOnly(ShowOnly);
+
+		if (CachedSelectedField.IsFromWidget())
+		{
+			WidgetList->SetSelectedWidget(CachedSelectedField.GetWidgetName());
+		}
 
 		SourcePicker = WidgetList;
 	}
 
 	if (ShowConversionFunctions.Get(false))
 	{
-		UMVVMEditorSubsystem* Subsystem = GEditor->GetEditorSubsystem<UMVVMEditorSubsystem>();
-		TArray<UFunction*> AllConversionFunctions = Subsystem->GetAvailableConversionFunctions(WidgetBlueprint.Get(), FMVVMBlueprintPropertyPath(), FMVVMBlueprintPropertyPath());
-
-		// remove all incompatible conversion functions
-		for (int32 Idx = 0; Idx < AllConversionFunctions.Num(); ++Idx)
-		{
-			const UFunction* Function = AllConversionFunctions[Idx];
-			const FProperty* ReturnProperty = BindingHelper::GetReturnProperty(Function);
-			const FProperty* AssignToProperty = AssignableTo.Get(nullptr);
-
-			if (AssignToProperty != nullptr && !BindingHelper::ArePropertiesCompatible(ReturnProperty, AssignToProperty))
-			{
-				AllConversionFunctions.RemoveAt(Idx);
-				--Idx;
-			}
-		}
-
-		TSharedPtr<FConversionFunctionItem> CurrentSelectedItem;
-
-		TArray<FString> SubCategories;
-		TArray<TSharedPtr<FConversionFunctionItem>> PathToSelection;
-
-		for (const UFunction* Function : AllConversionFunctions)
-		{
-			FText CategoryName = Function->GetMetaDataText("Category");
-
-			SubCategories.Reset();
-			CategoryName.ToString().ParseIntoArray(SubCategories, TEXT("|"));
-
-			TArray<TSharedPtr<FConversionFunctionItem>>* CurrentItems = &ConversionFunctionCategories;
-			TSharedPtr<FConversionFunctionItem> ParentItem;
-
-			// create items for the entire category path
-			// eg. "Math|Boolean|AND" 
-			// Math 
-			//   > Boolean
-			//     > AND
-			for (const FString& SubCategory : SubCategories)
-			{
-				const FString Trimmed = SubCategory.TrimStartAndEnd();
-
-				ParentItem = FindOrCreateItemForCategory(*CurrentItems, Trimmed);
-				CurrentItems = &ParentItem->Children;
-
-				if (Function == CachedSelectedConversionFunction)
-				{
-					PathToSelection.Add(ParentItem);
-				}
-			}
-
-			TSharedPtr<FConversionFunctionItem> Item = CurrentItems->Add_GetRef(MakeShared<FConversionFunctionItem>());
-			Item->Function = Function;
-		}
+		GenerateConversionFunctionItems();
 
 		FilterConversionFunctionCategories();
-
-		ConversionFunctionCategoryRoot.Reset();
-		TSharedPtr<FConversionFunctionItem> RootItem = ConversionFunctionCategoryRoot.Add_GetRef(MakeShared<FConversionFunctionItem>());
-		RootItem->CategoryName = TEXT("Conversion Functions");
-
+		
 		TSharedRef<SWidget> StackedSourcePicker =
-			SNew(SVerticalBox)
-			+ SVerticalBox::Slot()
-			.AutoHeight()
-			[
-				SourcePicker.ToSharedRef()
-			]
-			+ SVerticalBox::Slot()
-			.AutoHeight()
-			[
-				SAssignNew(ConversionFunctionCategoryTree, STreeView<TSharedPtr<FConversionFunctionItem>>)
-				.SelectionMode(ESelectionMode::Multi)
-				.TreeItemsSource(&ConversionFunctionCategoryRoot)
-				.OnGenerateRow(this, &SFieldSelector::GenerateConversionFunctionCategoryRow)
-				.OnSelectionChanged(this, &SFieldSelector::OnConversionFunctionCategorySelected)
-				.OnGetChildren(this, &SFieldSelector::GetConversionFunctionCategoryChildren)
-			];
-
-		if (PathToSelection.Num() > 0)
-		{
-			ConversionFunctionCategoryTree->SetItemExpansion(RootItem, true);
-
-			for (const TSharedPtr<FConversionFunctionItem>& ParentItem : PathToSelection)
-			{
-				ConversionFunctionCategoryTree->SetItemExpansion(ParentItem, true);
-			}
-			ConversionFunctionCategoryTree->SetItemSelection(PathToSelection.Last(), true);
-		}
+				SNew(SVerticalBox)
+				+ SVerticalBox::Slot()
+				.AutoHeight()
+				[
+					SourcePicker.ToSharedRef()
+				]
+				+ SVerticalBox::Slot()
+				.AutoHeight()
+				[
+					SAssignNew(ConversionFunctionCategoryTree, STreeView<TSharedPtr<FConversionFunctionItem>>)
+					.SelectionMode(ESelectionMode::Multi)
+					.TreeItemsSource(&FilteredConversionFunctionRoot)
+					.OnGenerateRow(this, &SFieldSelector::GenerateConversionFunctionCategoryRow)
+					.OnSelectionChanged(this, &SFieldSelector::OnConversionFunctionCategorySelected)
+					.OnGetChildren(this, &SFieldSelector::GetConversionFunctionCategoryChildren)
+				];
 
 		SourcePicker = StackedSourcePicker;
 	}
@@ -634,6 +771,17 @@ TSharedRef<SWidget> SFieldSelector::OnGetMenuContent()
 		return SNullWidget::NullWidget;
 	}
 
+	bIsMenuInitialized = false;
+
+	ConversionFunctionCategoryTree.Reset();
+	FilteredConversionFunctionRoot.Reset();
+	ConversionFunctions.Reset();
+	FilteredConversionFunctions.Reset();
+	ConversionFunctionRoot.Reset();
+	ViewModelSources.Reset();
+	ViewModelList.Reset();
+	WidgetList.Reset();
+
 	TSharedRef<SVerticalBox> VBox = SNew(SVerticalBox)
 		+ SVerticalBox::Slot()
 		.HAlign(HAlign_Fill)
@@ -665,33 +813,38 @@ TSharedRef<SWidget> SFieldSelector::OnGetMenuContent()
 	{
 		// Single fixed source, don't show the separate source panel.
 		BindingList->AddSource(FixedSource.GetValue());
+	}
 
-		VBox->AddSlot()
-			[
-				BindingList.ToSharedRef()
-			];
-	} 
-	else
-	{
-		TSharedPtr<SVerticalBox> InnerVBox;
-
-		VBox->AddSlot()
-			.FillHeight(1.0f)
+	BindingList->SetSelectedProperty(CachedSelectedField);
+	
+	TSharedPtr<SVerticalBox> BindingListVBox;
+		
+	VBox->AddSlot()
+		.FillHeight(1.0f)
+		[
+			SNew(SBorder)
+			.BorderImage(FAppStyle::Get().GetBrush("Brushes.Recessed"))
 			[
-				SNew(SBorder)
-				.BorderImage(FAppStyle::Get().GetBrush("Brushes.Recessed"))
+				SNew(SSplitter)
+				.PhysicalSplitterHandleSize(4.0f)
+				+ SSplitter::Slot()
+				.Value(0.5f)
+				.MinSize(100)
 				[
-					SNew(SSplitter)
-					.PhysicalSplitterHandleSize(4.0f)
-					+ SSplitter::Slot()
-					.Value(0.5f)
+					SNew(SScrollBox)
+					+ SScrollBox::Slot()
 					[
 						CreateSourcePanel()
 					]
-					+ SSplitter::Slot()
-					.Value(0.5f)
+				]
+				+ SSplitter::Slot()
+				.Value(0.5f)
+				.MinSize(100)
+				[
+					SNew(SScrollBox)
+					+ SScrollBox::Slot()
 					[
-						SAssignNew(InnerVBox, SVerticalBox)
+						SAssignNew(BindingListVBox, SVerticalBox)
 						+ SVerticalBox::Slot()
 						.AutoHeight()
 						[
@@ -699,24 +852,25 @@ TSharedRef<SWidget> SFieldSelector::OnGetMenuContent()
 						]
 					]
 				]
+			]
+		];
+
+	if (bShowConversionFunctions)
+	{
+		BindingListVBox->AddSlot()
+			.AutoHeight()
+			[
+				SAssignNew(ConversionFunctionList, SListView<const UFunction*>)
+				.SelectionMode(ESelectionMode::Single)
+				.ListItemsSource(&FilteredConversionFunctions)
+				.OnMouseButtonDoubleClick(this, &SFieldSelector::SetConversionFunctionSelection)
+				.OnGenerateRow(this, &SFieldSelector::GenerateConversionFunctionRow)
 			];
 
-		if (bShowConversionFunctions)
+		if (CachedSelectedConversionFunction != nullptr)
 		{
-			InnerVBox->AddSlot()
-				.AutoHeight()
-				[
-					 SAssignNew(ConversionFunctionList, SListView<const UFunction*>)
-					.SelectionMode(ESelectionMode::Single)
-					.ListItemsSource(&FilteredConversionFunctions)
-					.OnMouseButtonDoubleClick(this, &SFieldSelector::SetConversionFunctionSelection)
-					.OnGenerateRow(this, &SFieldSelector::GenerateConversionFunctionRow)
-				];
-
-			if (CachedSelectedConversionFunction != nullptr)
-			{
-				ConversionFunctionList->SetItemSelection(CachedSelectedConversionFunction, true);
-			}
+			ExpandAllToItem(CachedSelectedConversionFunction);
+			ConversionFunctionList->SetItemSelection(CachedSelectedConversionFunction, true);
 		}
 	}
 
@@ -775,6 +929,8 @@ TSharedRef<SWidget> SFieldSelector::OnGetMenuContent()
 		];
 
 	ComboButton->SetMenuContentWidgetToFocus(SearchBox);
+
+	bIsMenuInitialized = true;
 	return MenuWidget;
 }
 
@@ -844,50 +1000,36 @@ TSharedPtr<SFieldSelector::FConversionFunctionItem> SFieldSelector::FindOrCreate
 
 TSharedRef<ITableRow> SFieldSelector::GenerateConversionFunctionCategoryRow(TSharedPtr<FConversionFunctionItem> Item, const TSharedRef<STableViewBase>& OwnerTable)
 {
+	FText DisplayName = FText::FormatOrdered(FText::FromString("{0} ({1})"), FText::FromString(Item->CategoryName), FText::FromString(LexToString(Item->NumFunctions)));
+
 	return SNew(STableRow<TSharedPtr<FConversionFunctionItem>>, OwnerTable)
+	[
+		SNew(SHorizontalBox)
+		+ SHorizontalBox::Slot()
+		.HAlign(HAlign_Left)
+		.VAlign(VAlign_Center)
+		.Padding(0, 2.0f, 4.0f, 2.0f)
+		.AutoWidth()
 		[
-			SNew(SHorizontalBox)
-			+ SHorizontalBox::Slot()
-			.HAlign(HAlign_Left)
-			.VAlign(VAlign_Center)
-			.Padding(0, 2.0f, 4.0f, 2.0f)
-			.AutoWidth()
-			[
-				SNew(SImage)
-				.DesiredSizeOverride(FVector2D(16.0f, 16.0f))
-				.Image(FAppStyle::Get().GetBrush("GraphEditor.Function_16x"))
-			]
-			+ SHorizontalBox::Slot()
-			.HAlign(HAlign_Fill)
-			.VAlign(VAlign_Center)
-			[
-				SNew(STextBlock)
-				.Font(Item == ConversionFunctionCategoryRoot[0] ? FAppStyle::Get().GetFontStyle("NormalText") : FAppStyle::Get().GetFontStyle("BoldFont"))
-				.Text(FText::FromString(Item->CategoryName))
-				.HighlightText_Lambda([this]() { return SearchBox.IsValid() ? SearchBox->GetText() : FText::GetEmpty(); })
-			]
-		];
+			SNew(SImage)
+			.DesiredSizeOverride(FVector2D(16.0f, 16.0f))
+			.Image(FAppStyle::Get().GetBrush("GraphEditor.Function_16x"))
+		]
+		+ SHorizontalBox::Slot()
+		.HAlign(HAlign_Fill)
+		.VAlign(VAlign_Center)
+		[
+			SNew(STextBlock)
+			.Font(Item == FilteredConversionFunctionRoot[0] ? FAppStyle::Get().GetFontStyle("NormalText") : FAppStyle::Get().GetFontStyle("BoldFont"))
+			.Text(DisplayName)
+			.ToolTipText(FText::FromString(Item->CategoryName))
+			.HighlightText_Lambda([this]() { return SearchBox.IsValid() ? SearchBox->GetText() : FText::GetEmpty(); })
+		]
+	];
 }
 
 void SFieldSelector::GetConversionFunctionCategoryChildren(TSharedPtr<FConversionFunctionItem> Item, TArray<TSharedPtr<FConversionFunctionItem>>& OutItems) const
 {
-	if (ConversionFunctionCategoryRoot.Num() > 0)
-	{
-		if (Item == ConversionFunctionCategoryRoot[0])
-		{
-			Algo::TransformIf(FilteredConversionFunctionCategories, OutItems,
-				[](const TSharedPtr<FConversionFunctionItem>& Item)
-				{
-					return !Item->CategoryName.IsEmpty();
-				},
-				[](const TSharedPtr<FConversionFunctionItem>& Item)
-				{
-					return Item;
-				});
-			return;
-		}
-	}
-
 	Algo::TransformIf(Item->Children, OutItems,
 		[](const TSharedPtr<FConversionFunctionItem>& Item)
 		{
@@ -920,6 +1062,7 @@ TSharedRef<ITableRow> SFieldSelector::GenerateConversionFunctionRow(const UFunct
 			[
 				SNew(STextBlock)
 				.Text(Function->GetDisplayNameText())
+				.ToolTipText(Function->GetToolTipText())
 				.HighlightText_Lambda([this]() { return SearchBox.IsValid() ? SearchBox->GetText() : FText::GetEmpty(); })
 			]
 		];
@@ -986,46 +1129,40 @@ void SFieldSelector::FilterConversionFunctions()
 
 void SFieldSelector::OnConversionFunctionCategorySelected(TSharedPtr<FConversionFunctionItem> SelectedItem, ESelectInfo::Type)
 {
-	ConversionFunctions.Reset();
-	FilteredConversionFunctions.Reset();
-
-	if (BindingList.IsValid())
+	if (!bIsMenuInitialized)
 	{
-		BindingList->ClearSources();
-	}
-
-	if (ViewModelList.IsValid())
-	{
-		ViewModelList->ClearSelection();
+		return;
 	}
 
-	if (WidgetList.IsValid())
-	{
-		WidgetList->ClearSelection();
-	}
+	ConversionFunctions.Reset();
+	FilteredConversionFunctions.Reset();
 
 	TArray<TSharedPtr<FConversionFunctionItem>> SelectedItems = ConversionFunctionCategoryTree->GetSelectedItems();
-	if (SelectedItems.Num() == 0)
+	if (SelectedItems.Num() > 0)
 	{
-		return;
-	}
+		if (BindingList.IsValid())
+		{
+			BindingList->ClearSources();
+		}
 
-	if (SelectedItems.Contains(ConversionFunctionCategoryRoot[0]))
-	{ 
-		for (const TSharedPtr<FConversionFunctionItem>& Item : FilteredConversionFunctionCategories)
+		if (ViewModelList.IsValid())
+		{
+			ViewModelList->ClearSelection();
+		}
+
+		if (WidgetList.IsValid())
+		{
+			WidgetList->ClearSelection();
+		}
+
+		for (const TSharedPtr<FConversionFunctionItem>& Item : SelectedItems)
 		{
 			AddConversionFunctionChildrenRecursive(Item, ConversionFunctions);
 		}
-		return;
-	}
 
-	for (const TSharedPtr<FConversionFunctionItem>& Item : SelectedItems)
-	{
-		AddConversionFunctionChildrenRecursive(Item, ConversionFunctions);
+		FilterConversionFunctions();
 	}
 
-	FilterConversionFunctions();
-
 	if (ConversionFunctionList.IsValid())
 	{
 		ConversionFunctionList->RequestListRefresh();
diff --git a/Engine/Plugins/Runtime/ModelViewViewModel/Source/ModelViewViewModelEditor/Private/Widgets/SMVVMFieldSelector.h b/Engine/Plugins/Runtime/ModelViewViewModel/Source/ModelViewViewModelEditor/Private/Widgets/SMVVMFieldSelector.h
--- a/Engine/Plugins/Runtime/ModelViewViewModel/Source/ModelViewViewModelEditor/Private/Widgets/SMVVMFieldSelector.h
+++ b/Engine/Plugins/Runtime/ModelViewViewModel/Source/ModelViewViewModelEditor/Private/Widgets/SMVVMFieldSelector.h
@@ -108,13 +108,15 @@ private:
 		FString CategoryName;
 		const UFunction* Function = nullptr;
 		TArray<TSharedPtr<FConversionFunctionItem>> Children;
+		int32 NumFunctions = 0;
 	};
 
 	TSharedPtr<FConversionFunctionItem> FindOrCreateItemForCategory(TArray<TSharedPtr<FConversionFunctionItem>>& Items, const FString& Name);
 	TSharedRef<ITableRow> GenerateConversionFunctionCategoryRow(TSharedPtr<FConversionFunctionItem> Item, const TSharedRef<STableViewBase>& OwnerTable);
+	void GenerateConversionFunctionItems();
 	void FilterConversionFunctionCategories();
 	/** Recursively filter the items in SourceArray and place them into DestArray. Returns true if any items were added. */
-	bool FilterConversionFunctionCategoryChildren(const TArray<FString>& FilterStrings, const TArray<TSharedPtr<FConversionFunctionItem>>& SourceArray, TArray<TSharedPtr<FConversionFunctionItem>>& OutDestArray);
+	int32 FilterConversionFunctionCategoryChildren(const TArray<FString>& FilterStrings, const TArray<TSharedPtr<FConversionFunctionItem>>& SourceArray, TArray<TSharedPtr<FConversionFunctionItem>>& OutDestArray);
 
 	void FilterConversionFunctions();
 
@@ -123,7 +125,11 @@ private:
 	
 	void AddConversionFunctionChildrenRecursive(const TSharedPtr<FConversionFunctionItem>& Parent, TArray<const UFunction*>& OutFunctions);
 
+	int32 SortConversionFunctionItemsRecursive(TArray<TSharedPtr<FConversionFunctionItem>>& Items);
+
 	void SetConversionFunctionSelection(const UFunction* SelectedFunction);
+	void ExpandAllToItem(const UFunction* Function);
+	void ExpandAll(const TArray<TSharedPtr<FConversionFunctionItem>>& Items, bool bRecursive);
 
 	TSharedRef<ITableRow> GenerateConversionFunctionRow(const UFunction* Function, const TSharedRef<STableViewBase>& OwnerTable);
 
@@ -164,15 +170,15 @@ private:
 	FOnConversionFunctionSelectionChanged OnConversionFunctionSelectionChangedDelegate;
 
 	TSharedPtr<STreeView<TSharedPtr<FConversionFunctionItem>>> ConversionFunctionCategoryTree;
-	TArray<TSharedPtr<FConversionFunctionItem>> ConversionFunctionCategories;
-	TArray<TSharedPtr<FConversionFunctionItem>> FilteredConversionFunctionCategories;
-	TArray<TSharedPtr<FConversionFunctionItem>> ConversionFunctionCategoryRoot;
+	TArray<TSharedPtr<FConversionFunctionItem>> FilteredConversionFunctionRoot;
+	TArray<TSharedPtr<FConversionFunctionItem>> ConversionFunctionRoot;
 
 	TSharedPtr<SListView<const UFunction*>> ConversionFunctionList;
 
 	TArray<const UFunction*> ConversionFunctions;
 	TArray<const UFunction*> FilteredConversionFunctions;
 
+	bool bIsMenuInitialized = false;
 }; 
 
 } // namespace UE::MVVM
\ No newline at end of file
diff --git a/Engine/Plugins/Runtime/ModelViewViewModel/Source/ModelViewViewModelEditor/Private/Widgets/SMVVMSourceSelector.cpp b/Engine/Plugins/Runtime/ModelViewViewModel/Source/ModelViewViewModelEditor/Private/Widgets/SMVVMSourceSelector.cpp
--- a/Engine/Plugins/Runtime/ModelViewViewModel/Source/ModelViewViewModelEditor/Private/Widgets/SMVVMSourceSelector.cpp
+++ b/Engine/Plugins/Runtime/ModelViewViewModel/Source/ModelViewViewModelEditor/Private/Widgets/SMVVMSourceSelector.cpp
@@ -47,6 +47,7 @@ void SSourceSelector::Construct(const FArguments& Args, const UWidgetBlueprint*
 			[
 				SNew(SHorizontalBox)
 				+ SHorizontalBox::Slot()
+				.Padding(0, 2)
 				.HAlign(HAlign_Fill)
 				[
 					SAssignNew(SelectedSourceWidget, SSourceEntry)
