MVVM - Details panel is now initially collapsed.

#rnx

diff --git a/Engine/Plugins/Runtime/ModelViewViewModel/Source/ModelViewViewModelEditor/Private/Widgets/SMVVMViewBindingPanel.cpp b/Engine/Plugins/Runtime/ModelViewViewModel/Source/ModelViewViewModelEditor/Private/Widgets/SMVVMViewBindingPanel.cpp
--- a/Engine/Plugins/Runtime/ModelViewViewModel/Source/ModelViewViewModelEditor/Private/Widgets/SMVVMViewBindingPanel.cpp
+++ b/Engine/Plugins/Runtime/ModelViewViewModel/Source/ModelViewViewModelEditor/Private/Widgets/SMVVMViewBindingPanel.cpp
@@ -413,6 +413,7 @@ TSharedRef<SWidget> SBindingsPanel::GenerateEditViewWidget()
 				.Value(0.25f)
 				[
 					SAssignNew(DetailContainer, SBorder)
+					.Visibility(EVisibility::Collapsed)
 					[
 						DetailsView.ToSharedRef()
 					]
