MVVM - It is now possible to delete bindings by pressing Delete in the Bindings view.

Added missing Binding Mode tooltip.

#rnx

diff --git a/Engine/Plugins/Runtime/ModelViewViewModel/Source/ModelViewViewModelEditor/Private/Widgets/SMVVMViewBindingListView.cpp b/Engine/Plugins/Runtime/ModelViewViewModel/Source/ModelViewViewModelEditor/Private/Widgets/SMVVMViewBindingListView.cpp
--- a/Engine/Plugins/Runtime/ModelViewViewModel/Source/ModelViewViewModelEditor/Private/Widgets/SMVVMViewBindingListView.cpp
+++ b/Engine/Plugins/Runtime/ModelViewViewModel/Source/ModelViewViewModelEditor/Private/Widgets/SMVVMViewBindingListView.cpp
@@ -13,6 +13,7 @@
 #include "Framework/MultiBox/MultiBoxBuilder.h"
 #include "IPropertyAccessEditor.h"
 #include "K2Node_CallFunction.h"
+#include "Misc/MessageDialog.h"
 #include "MVVMBlueprintView.h"
 #include "MVVMBlueprintViewBinding.h"
 #include "MVVMEditorSubsystem.h"
@@ -313,7 +314,7 @@ public:
 					SNew(SBox)
 					.MinDesiredWidth(150)
 					[
-						SNew(SFieldSelector, InWidgetBlueprint, false)
+						SAssignNew(WidgetFieldSelector, SFieldSelector, InWidgetBlueprint, false)
 						.BindingMode(this, &SBindingRow::GetCurrentBindingMode)
 						.SelectedField(this, &SBindingRow::GetSelectedWidgetProperty)
 						.OnFieldSelectionChanged(this, &SBindingRow::OnWidgetPropertySelected)
@@ -327,7 +328,7 @@ public:
 				]
 
 				+ SHorizontalBox::Slot()
-				.Padding(2, 1)
+				.Padding(2, 0)
 				.VAlign(VAlign_Fill)
 				.HAlign(HAlign_Left) 
 				.AutoWidth()
@@ -335,8 +336,9 @@ public:
 					SNew(SComboBox<FName>)
 					.OptionsSource(Private::GetBindingModeNames())
 					.InitiallySelectedItem(StaticEnum<EMVVMBindingMode>()->GetNameByValue((int64) ViewBinding->BindingType))
-					.OnSelectionChanged(this, &SBindingRow::OnModeSelectionChanged)
-					.OnGenerateWidget(this, &SBindingRow::GenerateModeWidget)
+					.OnSelectionChanged(this, &SBindingRow::OnBindingModeSelectionChanged)
+					.OnGenerateWidget(this, &SBindingRow::GenerateBindingModeWidget)
+					.ToolTipText(this, &SBindingRow::GetCurrentBindingModeLabel)
 					.Content()
 					[
 						SNew(SBox)
@@ -346,7 +348,7 @@ public:
 						.HeightOverride(16)
 						[
 							SNew(SImage)
-							.Image(this, &SBindingRow::GetCurrentModeBrush)
+							.Image(this, &SBindingRow::GetCurrentBindingModeBrush)
 						]
 					]
 				]
@@ -459,25 +461,32 @@ private:
 		ErrorDialog.Reset();
 		ErrorItems.Reset();
 
-		if (const FMVVMBlueprintViewBinding* ViewModelBinding = GetThisViewBinding())
+		const UWidgetBlueprint* WidgetBlueprint = WidgetBlueprintWeak.Get();
+		UMVVMEditorSubsystem* EditorSubsystem = GEditor->GetEditorSubsystem<UMVVMEditorSubsystem>();
+		if (const UMVVMBlueprintView* View = EditorSubsystem->GetView(WidgetBlueprint))
 		{
-			for (const FText& ErrorText : ViewModelBinding->Errors)
+			if (const FMVVMBlueprintViewBinding* ViewBinding = Entry->GetBinding(View))
 			{
-				ErrorItems.Add(MakeShared<FText>(ErrorText));
-			}
+				for (const FText& ErrorText : ViewBinding->Errors)
+				{
+					ErrorItems.Add(MakeShared<FText>(ErrorText));
+				}
 
-			ErrorDialog = SNew(SCustomDialog)
-				.Buttons({
-					SCustomDialog::FButton(LOCTEXT("OK", "OK"))
-				})
-				.Content()
-				[
-					SNew(SListView<TSharedPtr<FText>>)
-					.ListItemsSource(&ErrorItems)
-					.OnGenerateRow(this, &SBindingRow::OnGenerateErrorRow)
-				];
+				const FText BindingDisplayName = FText::FromString(ViewBinding->GetDisplayNameString(View));
+				ErrorDialog = SNew(SCustomDialog)
+					.Title(FText::Format(LOCTEXT("Compilation Errors", "Compilation Errors for {0}"), BindingDisplayName))
+					.Buttons({
+						SCustomDialog::FButton(LOCTEXT("OK", "OK"))
+					})
+					.Content()
+					[
+						SNew(SListView<TSharedPtr<FText>>)
+						.ListItemsSource(&ErrorItems)
+						.OnGenerateRow(this, &SBindingRow::OnGenerateErrorRow)
+					];
 
-			ErrorDialog->Show();
+				ErrorDialog->Show();
+			}
 		}
 
 		return FReply::Handled();
@@ -548,6 +557,7 @@ private:
 			.Content()
 			[
 				SNew(SEditableTextBox)
+				.BackgroundColor(FStyleColors::Background)
 				.IsReadOnly(true)
 				.Text(*Text.Get())
 			];
@@ -613,6 +623,11 @@ private:
 			{
 				UMVVMEditorSubsystem* Subsystem = GEditor->GetEditorSubsystem<UMVVMEditorSubsystem>();
 				Subsystem->SetViewModelPropertyForBinding(WidgetBlueprintWeak.Get(), *ViewBinding, SelectedField);
+
+				if (WidgetFieldSelector.IsValid())
+				{
+					WidgetFieldSelector->Refresh();
+				}
 			}
 		}
 	}
@@ -685,7 +700,7 @@ private:
 		}
 	}
 
-	const FSlateBrush* GetModeBrush(EMVVMBindingMode BindingMode) const
+	const FSlateBrush* GetBindingModeBrush(EMVVMBindingMode BindingMode) const
 	{
 		switch (BindingMode)
 		{
@@ -704,16 +719,25 @@ private:
 		}
 	}
 
-	const FSlateBrush* GetCurrentModeBrush() const
+	const FSlateBrush* GetCurrentBindingModeBrush() const
 	{
 		if (FMVVMBlueprintViewBinding* ViewBinding = GetThisViewBinding())
 		{
-			return GetModeBrush(ViewBinding->BindingType);
+			return GetBindingModeBrush(ViewBinding->BindingType);
 		}
 		return nullptr;
 	}
 
-	const FText& GetModeLabel(EMVVMBindingMode BindingMode) const
+	FText GetCurrentBindingModeLabel() const
+	{
+		if (FMVVMBlueprintViewBinding* ViewBinding = GetThisViewBinding())
+		{
+			return GetBindingModeLabel(ViewBinding->BindingType);
+		}
+		return FText::GetEmpty();
+	}
+
+	FText GetBindingModeLabel(EMVVMBindingMode BindingMode) const
 	{
 		static FText OneTimeToDestinationLabel = LOCTEXT("OneTimeToDestinationLabel", "One Time To Widget");
 		static FText OneWayToDestinationLabel = LOCTEXT("OneWayToDestinationLabel", "One Way To Widget");
@@ -738,7 +762,7 @@ private:
 		}
 	}
 
-	TSharedRef<SWidget> GenerateModeWidget(FName ValueName) const
+	TSharedRef<SWidget> GenerateBindingModeWidget(FName ValueName) const
 	{
 		const UEnum* ModeEnum = StaticEnum<EMVVMBindingMode>();
 		int32 Index = ModeEnum->GetIndexByName(ValueName);
@@ -754,7 +778,7 @@ private:
 				.HeightOverride(16)
 				[
 					SNew(SImage)
-					.Image(GetModeBrush(MVVMBindingMode))
+					.Image(GetBindingModeBrush(MVVMBindingMode))
 				]
 			]
 			+SHorizontalBox::Slot()
@@ -762,12 +786,12 @@ private:
 			.Padding(2.0, 0, 0, 0)
 			[
 				SNew(STextBlock)
-				.Text(GetModeLabel(MVVMBindingMode))
+				.Text(GetBindingModeLabel(MVVMBindingMode))
 				.ToolTipText(ModeEnum->GetToolTipTextByIndex(Index))
 			];
 	}
 
-	void OnModeSelectionChanged(FName ValueName, ESelectInfo::Type)
+	void OnBindingModeSelectionChanged(FName ValueName, ESelectInfo::Type)
 	{
 		if (FMVVMBlueprintViewBinding* ViewBinding = GetThisViewBinding())
 		{
@@ -781,17 +805,31 @@ private:
 			{
 				ViewModelFieldSelector->Refresh();
 			}
+
+			if (WidgetFieldSelector.IsValid())
+			{
+				WidgetFieldSelector->Refresh();
+			}
 		}
 	}
 
 	void HandleBlueprintChanged(UBlueprint* Blueprint)
 	{
-		ViewModelFieldSelector->Refresh();
+		if (ViewModelFieldSelector.IsValid())
+		{
+			ViewModelFieldSelector->Refresh();
+		}
+
+		if (WidgetFieldSelector.IsValid())
+		{
+			WidgetFieldSelector->Refresh();
+		}
 	}
 
 private:
 	TSharedPtr<FBindingEntry> Entry;
 	TWeakObjectPtr<UWidgetBlueprint> WidgetBlueprintWeak;
+	TSharedPtr<SFieldSelector> WidgetFieldSelector;
 	TSharedPtr<SFieldSelector> ViewModelFieldSelector;
 	TSharedPtr<SWidget> ContextMenuOptionHelper;
 	TSharedPtr<SCustomDialog> ErrorDialog;
@@ -1098,15 +1136,66 @@ TSharedRef<ITableRow> SBindingsList::GenerateEntryRow(TSharedPtr<FBindingEntry>
 
 namespace Private
 {
-	void RemoveAllBindings(UMVVMBlueprintView* BlueprintView, const TConstArrayView<TSharedPtr<FBindingEntry>>& Entries)
+	void GatherChildBindings(UMVVMBlueprintView* BlueprintView, const TConstArrayView<TSharedPtr<FBindingEntry>>& Entries, TArray<const FMVVMBlueprintViewBinding*>& OutBindings)
 	{
 		for (const TSharedPtr<FBindingEntry>& Entry : Entries)
 		{
-			RemoveAllBindings(BlueprintView, Entry->GetChildren());
+			const FMVVMBlueprintViewBinding* Binding = Entry->GetBinding(BlueprintView);
+			if (Binding != nullptr)
+			{
+				OutBindings.AddUnique(Binding);
+			}
+
+			GatherChildBindings(BlueprintView, Entry->GetChildren(), OutBindings);
+		}
+	}
+}
+
+void SBindingsList::OnDeleteSelected()
+{
+	TArray<TSharedPtr<FBindingEntry>> Selection = TreeView->GetSelectedItems();
+	if (Selection.Num() == 0)
+	{
+		return;
+	}
+
+	if (UMVVMWidgetBlueprintExtension_View* MVVMExtensionPtr = MVVMExtension.Get())
+	{
+		if (UMVVMBlueprintView* BlueprintView = MVVMExtensionPtr->GetBlueprintView())
+		{
+			TArray<const FMVVMBlueprintViewBinding*> BindingsToRemove;
+			Private::GatherChildBindings(BlueprintView, Selection, BindingsToRemove);
+
+			if (BindingsToRemove.Num() == 0)
+			{
+				return;
+			}
 
-			if (FMVVMBlueprintViewBinding* Binding = Entry->GetBinding(BlueprintView))
+			TArray<FText> BindingDisplayNames;
+			for (const FMVVMBlueprintViewBinding* Binding : BindingsToRemove)
 			{
-				BlueprintView->RemoveBinding(Binding);
+				BindingDisplayNames.Add(FText::FromString(Binding->GetDisplayNameString(BlueprintView)));
+			}
+
+			const FText Message = FText::Format(BindingDisplayNames.Num() == 1 ?
+				LOCTEXT("ConfirmDeleteSingle", "Are you sure that you want to delete this binding?\n\n{1}") :
+				LOCTEXT("ConfirmDeleteMultiple", "Are you sure that you want to delete these {0} bindings?\n\n{1}"), 
+			BindingDisplayNames.Num(), 
+			FText::Join(FText::FromString("\n"), BindingDisplayNames));
+
+			const FText Title = LOCTEXT("DeleteBindings", "Delete Bindings?");
+
+			if (FMessageDialog::Open(EAppMsgType::YesNo, EAppReturnType::Yes, Message, &Title) == EAppReturnType::Yes)
+			{
+				if (TSharedPtr<SBindingsPanel> BindingPanelPtr = BindingPanel.Pin())
+				{
+					BindingPanelPtr->OnBindingListSelectionChanged(TConstArrayView<FMVVMBlueprintViewBinding*>());
+				}
+
+				for (const FMVVMBlueprintViewBinding* Binding : BindingsToRemove)
+				{
+					BlueprintView->RemoveBinding(Binding);
+				}
 			}
 		}
 	}
@@ -1121,21 +1210,7 @@ TSharedPtr<SWidget> SBindingsList::OnSourceConstructContextMenu()
 	if (Selection.Num() > 0)
 	{
 		FUIAction RemoveAction;
-		RemoveAction.ExecuteAction = FExecuteAction::CreateLambda([this, Selection]()
-			{
-				if (UMVVMWidgetBlueprintExtension_View* MVVMExtensionPtr = MVVMExtension.Get())
-				{
-					if (TSharedPtr<SBindingsPanel> BindingPanelPtr = BindingPanel.Pin())
-					{
-						BindingPanelPtr->OnBindingListSelectionChanged(TConstArrayView<FMVVMBlueprintViewBinding*>());
-					}
-
-					if (UMVVMBlueprintView* BlueprintView = MVVMExtensionPtr->GetBlueprintView())
-					{
-						Private::RemoveAllBindings(BlueprintView, Selection);
-					}
-				}
-			});
+		RemoveAction.ExecuteAction = FExecuteAction::CreateSP(this, &SBindingsList::OnDeleteSelected);
 		MenuBuilder.AddMenuEntry(LOCTEXT("RemoveBinding", "Remove Binding"), 
 			LOCTEXT("RemoveBindingTooltip", "Remove this binding."), 
 			FSlateIcon(FAppStyle::GetAppStyleSetName(), "Icons.Delete"), 
@@ -1145,6 +1220,17 @@ TSharedPtr<SWidget> SBindingsList::OnSourceConstructContextMenu()
 	return MenuBuilder.MakeWidget();
 }
 
+FReply SBindingsList::OnKeyDown(const FGeometry& MyGeometry, const FKeyEvent& InKeyEvent)
+{
+	if (InKeyEvent.GetKey() == EKeys::Delete)
+	{
+		OnDeleteSelected();
+		return FReply::Handled();
+	}
+
+	return FReply::Unhandled();
+}
+
 void SBindingsList::OnSourceListSelectionChanged(TSharedPtr<FBindingEntry> Entry, ESelectInfo::Type SelectionType) const
 {
 	if (bSelectionChangedGuard)
diff --git a/Engine/Plugins/Runtime/ModelViewViewModel/Source/ModelViewViewModelEditor/Private/Widgets/SMVVMViewBindingListView.h b/Engine/Plugins/Runtime/ModelViewViewModel/Source/ModelViewViewModelEditor/Private/Widgets/SMVVMViewBindingListView.h
--- a/Engine/Plugins/Runtime/ModelViewViewModel/Source/ModelViewViewModelEditor/Private/Widgets/SMVVMViewBindingListView.h
+++ b/Engine/Plugins/Runtime/ModelViewViewModel/Source/ModelViewViewModelEditor/Private/Widgets/SMVVMViewBindingListView.h
@@ -37,6 +37,9 @@ private:
 	void OnSourceListSelectionChanged(TSharedPtr<FBindingEntry> Entry, ESelectInfo::Type SelectionType) const;
 	void GetChildrenOfEntry(TSharedPtr<FBindingEntry> Entry, TArray<TSharedPtr<FBindingEntry>>& OutChildren) const;
 
+	virtual FReply OnKeyDown(const FGeometry& MyGeometry, const FKeyEvent& InKeyEvent) override;
+	void OnDeleteSelected();
+
 private:
 	TWeakPtr<SBindingsPanel> BindingPanel;
 	TSharedPtr<STreeView<TSharedPtr<FBindingEntry>>> TreeView;
