MVVM - Bindings now have a GUID for identification purposes instead of using the index.

Display name is now more readable when conversion functions are involved.

Moved some conversion function code from MVVMEditorSubsystem to ConversionFunctionHelper so that it can be reused in the compiler.

#rnx

diff --git a/Engine/Plugins/Runtime/ModelViewViewModel/Source/ModelViewViewModel/Private/Bindings/MVVMBindingHelper.cpp b/Engine/Plugins/Runtime/ModelViewViewModel/Source/ModelViewViewModel/Private/Bindings/MVVMBindingHelper.cpp
--- a/Engine/Plugins/Runtime/ModelViewViewModel/Source/ModelViewViewModel/Private/Bindings/MVVMBindingHelper.cpp
+++ b/Engine/Plugins/Runtime/ModelViewViewModel/Source/ModelViewViewModel/Private/Bindings/MVVMBindingHelper.cpp
@@ -579,6 +579,22 @@ namespace UE::MVVM::BindingHelper
 		return nullptr;
 	}
 
+	TArray<const FProperty*> GetAllArgumentProperties(const UFunction* InFunction)
+	{
+		check(InFunction);
+
+		TArray<const FProperty*> Arguments;
+		for (TFieldIterator<FProperty> It(InFunction); It && (It->PropertyFlags & CPF_Parm); ++It)
+		{
+			if (IsValidArgumentProperty(*It))
+			{
+				Arguments.Add(*It);
+			}
+		}
+
+		return Arguments;
+	}
+
 	void ExecuteBinding_NoCheck(const FFieldContext& Source, const FFieldContext& Destination)
 	{
 		check(!Source.GetObjectVariant().IsNull() && !Destination.GetObjectVariant().IsNull());
@@ -820,3 +836,5 @@ namespace UE::MVVM::BindingHelper
 		}
 	}
 }
+
+#undef LOCTEXT_NAMESPACE
diff --git a/Engine/Plugins/Runtime/ModelViewViewModel/Source/ModelViewViewModel/Public/Bindings/MVVMBindingHelper.h b/Engine/Plugins/Runtime/ModelViewViewModel/Source/ModelViewViewModel/Public/Bindings/MVVMBindingHelper.h
--- a/Engine/Plugins/Runtime/ModelViewViewModel/Source/ModelViewViewModel/Public/Bindings/MVVMBindingHelper.h
+++ b/Engine/Plugins/Runtime/ModelViewViewModel/Source/ModelViewViewModel/Public/Bindings/MVVMBindingHelper.h
@@ -111,6 +111,12 @@ namespace UE::MVVM::BindingHelper
 	 */
 	UE_NODISCARD MODELVIEWVIEWMODEL_API const FProperty* GetFirstArgumentProperty(const UFunction* InFunction);
 
+
+	/**
+	 * Returns all argument properties of this function.
+	 */
+	UE_NODISCARD MODELVIEWVIEWMODEL_API TArray<const FProperty*> GetAllArgumentProperties(const UFunction* InFunction);
+
 	/**
 	 * Returns true if the given property is a valid function argument property.
 	 */
diff --git a/Engine/Plugins/Runtime/ModelViewViewModel/Source/ModelViewViewModelBlueprint/Private/Bindings/MVVMConversionFunctionHelper.cpp b/Engine/Plugins/Runtime/ModelViewViewModel/Source/ModelViewViewModelBlueprint/Private/Bindings/MVVMConversionFunctionHelper.cpp
--- a/Engine/Plugins/Runtime/ModelViewViewModel/Source/ModelViewViewModelBlueprint/Private/Bindings/MVVMConversionFunctionHelper.cpp
+++ b/Engine/Plugins/Runtime/ModelViewViewModel/Source/ModelViewViewModelBlueprint/Private/Bindings/MVVMConversionFunctionHelper.cpp
@@ -3,136 +3,218 @@
 #include "Bindings/MVVMConversionFunctionHelper.h"
 
 #include "BlueprintTypePromotion.h"
-#include "EngineLogs.h"
+#include "Components/Widget.h"
+#include "Containers/Deque.h"
 #include "EdGraph/EdGraph.h"
+#include "EngineLogs.h"
 #include "K2Node_PromotableOperator.h"
 #include "K2Node_VariableGet.h"
 #include "Kismet2/BlueprintEditorUtils.h"
-#include "Misc/StringBuilder.h"
-
-#include "Components/Widget.h"
 #include "MVVMBlueprintView.h"
 #include "MVVMBlueprintViewModelContext.h"
-#include "MVVMWidgetBlueprintExtension_View.h"
 #include "MVVMViewModelBase.h"
+#include "MVVMWidgetBlueprintExtension_View.h"
+#include "Misc/StringBuilder.h"
+#include "UObject/MetaData.h"
 #include "View/MVVMView.h"
 #include "WidgetBlueprint.h"
 
-
 namespace UE::MVVM::ConversionFunctionHelper
 {
+static const FName ConversionFunctionMetadataKey = "ConversionFunction";
 
 namespace Private
 {
-void BuildPropertyPath(TArray<FMVVMBlueprintPropertyPath>& Result, FMVVMBlueprintPropertyPath BasePath, UClass* Class, const UK2Node_VariableGet* CurrentNode)
-{
-	bool bAddPath = false;
-	const UEdGraphPin* OutputPin = CurrentNode->FindPinByPredicate([](UEdGraphPin* Pin) { return Pin->Direction == EGPD_Output; });
-	if (OutputPin)
+
+	void BuildPropertyPath(TArray<FMVVMBlueprintPropertyPath>& Result, FMVVMBlueprintPropertyPath BasePath, UClass* Class, const UK2Node_VariableGet* CurrentNode)
 	{
-		for (const UEdGraphPin* LinkedTo : OutputPin->LinkedTo)
+		bool bAddPath = false;
+		const UEdGraphPin* OutputPin = CurrentNode->FindPinByPredicate([](UEdGraphPin* Pin) { return Pin->Direction == EGPD_Output; });
+		if (OutputPin)
 		{
-			if (const UK2Node_VariableGet* LinkToVariable = Cast<UK2Node_VariableGet>(LinkedTo->GetOwningNode()))
-			{
-				BasePath.SetBasePropertyPath(UE::MVVM::FMVVMConstFieldVariant(LinkToVariable->VariableReference.ResolveMember<FProperty>(Class)));
-				BuildPropertyPath(Result, BasePath, Class, LinkToVariable);
-			}
-			else
+			for (const UEdGraphPin* LinkedTo : OutputPin->LinkedTo)
 			{
-				bAddPath = true;
+				if (const UK2Node_VariableGet* LinkToVariable = Cast<UK2Node_VariableGet>(LinkedTo->GetOwningNode()))
+				{
+					BasePath.SetBasePropertyPath(UE::MVVM::FMVVMConstFieldVariant(LinkToVariable->VariableReference.ResolveMember<FProperty>(Class)));
+					BuildPropertyPath(Result, BasePath, Class, LinkToVariable);
+				}
+				else
+				{
+					bAddPath = true;
+				}
 			}
 		}
+
+		if (bAddPath)
+		{
+			Result.Add(BasePath);
+		}
 	}
 
-	if (bAddPath)
+	UEdGraph* FindExistingConversionFunctionWrapper(const UWidgetBlueprint* WidgetBlueprint, FName WrapperName)
 	{
-		Result.Add(BasePath);
+		const TObjectPtr<UEdGraph>* Result = WidgetBlueprint->FunctionGraphs.FindByPredicate([WrapperName](const UEdGraph* GraphPtr) { return GraphPtr->GetFName() == WrapperName; });
+		return Result ? Result->Get() : nullptr;
 	}
+} // namespace Private
+
+FName CreateWrapperName(const FMVVMBlueprintViewBinding& Binding, bool bSourceToDestination)
+{
+	TStringBuilder<256> StringBuilder;
+	StringBuilder << TEXT("__");
+	StringBuilder << Binding.GetFName();
+	StringBuilder << (bSourceToDestination ? TEXT("_SourceToDest") : TEXT("_DestToSource"));
+
+	return FName(StringBuilder.ToString());
+}
+
+
+UEdGraph* GetGraph(const UWidgetBlueprint* WidgetBlueprint, const FMVVMBlueprintViewBinding& Binding, bool bSourceToDestination)
+{
+	check(WidgetBlueprint);
+	const FName WrapperName = bSourceToDestination ? Binding.Conversion.SourceToDestinationWrapper : Binding.Conversion.DestinationToSourceWrapper;
+	return Private::FindExistingConversionFunctionWrapper(WidgetBlueprint, WrapperName);
+}
+
+UK2Node_CallFunction* GetFunctionNode(UEdGraph* Graph)
+{
+	if (Graph == nullptr)
+	{
+		return nullptr;
+	}
+
+	TArray<UK2Node_CallFunction*> FunctionNodes;
+	Graph->GetNodesOfClass<UK2Node_CallFunction>(FunctionNodes);
+
+	if (FunctionNodes.Num() == 1)
+	{
+		return FunctionNodes[0];
+	}
+
+	for (UK2Node_CallFunction* FunctionNode : FunctionNodes)
+	{
+		// check if we've set any metadata on the nodes to figure out which one it is
+		if (FunctionNode->GetPackage()->GetMetaData()->HasValue(FunctionNode, ConversionFunctionMetadataKey))
+		{
+			return FunctionNode;
+		}
+	}
+
+	return nullptr;
 }
-} //namespace
 
-TArray<FMVVMBlueprintPropertyPath> FindAllPropertyPathInGraph(const UEdGraph* Graph, const UMVVMBlueprintView* BlueprintView, UClass* Class)
+void MarkAsConversionFunction(const UK2Node_CallFunction* FunctionNode, const FMVVMBlueprintViewBinding& Binding)
 {
-	TArray<FMVVMBlueprintPropertyPath> Result;
+	check(FunctionNode != nullptr);
 
-	TArray<UK2Node_VariableGet*> GetNodes;
-	Graph->GetNodesOfClass<UK2Node_VariableGet>(GetNodes);
-	for (const UK2Node_VariableGet* Node : GetNodes)
+	FunctionNode->GetPackage()->GetMetaData()->SetValue(FunctionNode, ConversionFunctionMetadataKey, *Binding.BindingId.ToString());
+}
+
+namespace Private
+{
+	FMVVMBlueprintPropertyPath GetPropertyPathForPin(const UWidgetBlueprint* WidgetBlueprint, const UEdGraphPin* StartPin)
 	{
-		if (Node->VariableReference.IsSelfContext())
+		UMVVMWidgetBlueprintExtension_View* ExtensionView = UMVVMWidgetBlueprintExtension_View::GetExtension<UMVVMWidgetBlueprintExtension_View>(WidgetBlueprint);
+		UMVVMBlueprintView* BlueprintView = ExtensionView->GetBlueprintView();
+
+		TArray<FMemberReference> PathParts;
+		TDeque<UEdGraphNode*> NodesToSearch;
+		NodesToSearch.PushLast(StartPin->LinkedTo[0]->GetOwningNode());
+
+		while (NodesToSearch.Num() > 0)
 		{
-			if (FProperty* Property = Node->VariableReference.ResolveMember<FProperty>(Class))
+			UEdGraphNode* Node = NodesToSearch[0];
+			NodesToSearch.PopFirst();
+
+			if (UK2Node_VariableGet* Getter = Cast<UK2Node_VariableGet>(Node))
 			{
-				const UClass* OwningClass = Class;
-				if (const FObjectProperty* ObjectProperty = CastField<FObjectProperty>(Property))
-				{
-					OwningClass = ObjectProperty->PropertyClass;
-				}
-				check(OwningClass);
+				PathParts.Add(Getter->VariableReference);
+			}
 
-				FMVVMBlueprintPropertyPath PropertyPath;
-				if (OwningClass->IsChildOf<UWidget>() || OwningClass->IsChildOf<UWidgetBlueprint>())
-				{
-					PropertyPath.SetWidgetName(Property->GetFName());
-				}
-				else if (OwningClass->ImplementsInterface(UNotifyFieldValueChanged::StaticClass()))
+			for (UEdGraphPin* Pin : Node->Pins)
+			{
+				if (Pin->Direction == EGPD_Input && Pin->LinkedTo.Num() > 0)
 				{
-					if (const FMVVMBlueprintViewModelContext* ViewModel = BlueprintView->FindViewModel(Property->GetFName()))
-					{
-						PropertyPath.SetViewModelId(ViewModel->GetViewModelId());
-					}
+					NodesToSearch.PushLast(Pin->LinkedTo[0]->GetOwningNode());
 				}
+			}
+		}
 
-				if (PropertyPath.IsEmpty())
+		if (PathParts.Num() == 0)
+		{
+			return FMVVMBlueprintPropertyPath();
+		}
+
+		Algo::Reverse(PathParts);
+
+		FMVVMBlueprintPropertyPath ResultPath;
+
+		FMemberReference RootRef = PathParts[0];
+		PathParts.RemoveAt(0);
+
+		if (const FObjectProperty* Property = CastField<FObjectProperty>(RootRef.ResolveMember<FProperty>(WidgetBlueprint->SkeletonGeneratedClass)))
+		{
+			if (Property->PropertyClass->ImplementsInterface(UNotifyFieldValueChanged::StaticClass()))
+			{
+				if (const FMVVMBlueprintViewModelContext* ViewModel = BlueprintView->FindViewModel(Property->GetFName()))
 				{
-					continue;
+					ResultPath.SetViewModelId(ViewModel->GetViewModelId());
 				}
-
-				Private::BuildPropertyPath(Result, PropertyPath, Class, Node);
+			}
+			else if (Property->PropertyClass->IsChildOf<UWidget>() || Property->PropertyClass->IsChildOf<UWidgetBlueprint>())
+			{
+				ResultPath.SetWidgetName(Property->GetFName());
 			}
 		}
-	}
-
-	return Result;
-}
 
+		if (PathParts.Num() == 0)
+		{
+			return ResultPath;
+		}
 
-FName GetWrapperName(const UMVVMBlueprintView* View, const FMVVMBlueprintViewBinding& Binding, bool bSourceToDestination)
-{
-	TStringBuilder<256> StringBuilder;
-	StringBuilder << TEXT("__");
-	StringBuilder << Binding.GetFName(View);
-	StringBuilder << (bSourceToDestination ? TEXT("_SourceToDest") : TEXT("_DestToSource"));
+		for (const FMemberReference& MemberReference : PathParts)
+		{
+			if (UFunction* Function = MemberReference.ResolveMember<UFunction>(WidgetBlueprint->SkeletonGeneratedClass))
+			{
+				ResultPath.AppendBasePropertyPath(UE::MVVM::FMVVMConstFieldVariant(Function));
+			}
+			else if (FProperty* Property = MemberReference.ResolveMember<FProperty>(WidgetBlueprint->SkeletonGeneratedClass))
+			{
+				ResultPath.AppendBasePropertyPath(UE::MVVM::FMVVMConstFieldVariant(Property));
+			}
+		}
 
-	return FName(StringBuilder.ToString());
+		return ResultPath;
+	}
 }
 
-
-namespace Private
+FMVVMBlueprintPropertyPath GetPropertyPathForArgument(const UWidgetBlueprint* WidgetBlueprint, const UK2Node_CallFunction* FunctionNode, FName ArgumentName)
 {
-UEdGraph* FindExistingConversionFunctionWrapper(const UWidgetBlueprint* WidgetBlueprint, FName WrapperName)
-{
-	const TObjectPtr<UEdGraph>* Result = WidgetBlueprint->FunctionGraphs.FindByPredicate([WrapperName](const UEdGraph* GraphPtr) { return GraphPtr->GetFName() == WrapperName; });
-	return Result ? Result->Get() : nullptr;
+	const UEdGraphPin* ArgumentPin = FunctionNode->FindPin(ArgumentName, EGPD_Input);
+	if (ArgumentPin == nullptr || ArgumentPin->LinkedTo.Num() == 0)
+	{
+		return FMVVMBlueprintPropertyPath();
+	}
+
+	return Private::GetPropertyPathForPin(WidgetBlueprint, ArgumentPin);
 }
-} // namespace
 
-UEdGraph* GetGraph(const UWidgetBlueprint* WidgetBlueprint, const FMVVMBlueprintViewBinding& Binding, bool bSourceToDestination)
+TArray<FMVVMBlueprintPropertyPath> GetAllArgumentPathsForConversionNode(const UWidgetBlueprint* WidgetBlueprint, const UK2Node_CallFunction* FunctionNode)
 {
-	check(WidgetBlueprint);
+	check(FunctionNode);
 
-	const UMVVMBlueprintView* View = nullptr;
-	if (UMVVMWidgetBlueprintExtension_View* ExtensionView = UMVVMWidgetBlueprintExtension_View::GetExtension<UMVVMWidgetBlueprintExtension_View>(WidgetBlueprint))
+	TArray<FMVVMBlueprintPropertyPath> Paths;
+	for (const UEdGraphPin* Pin : FunctionNode->GetAllPins())
 	{
-		View = ExtensionView->GetBlueprintView();
+		FMVVMBlueprintPropertyPath Path = Private::GetPropertyPathForPin(WidgetBlueprint, Pin);
+		if (!Path.IsEmpty())
+		{
+			Paths.Add(Path);
+		}
 	}
 
-	if (View)
-	{
-		const FName WrapperName = GetWrapperName(View, Binding, bSourceToDestination);
-		return Private::FindExistingConversionFunctionWrapper(WidgetBlueprint, WrapperName);
-	}
-	return nullptr;
+	return Paths;
 }
 
-} //namespace
\ No newline at end of file
+} //namespace UE::MVVM
\ No newline at end of file
diff --git a/Engine/Plugins/Runtime/ModelViewViewModel/Source/ModelViewViewModelBlueprint/Private/MVVMBlueprintView.cpp b/Engine/Plugins/Runtime/ModelViewViewModel/Source/ModelViewViewModelBlueprint/Private/MVVMBlueprintView.cpp
--- a/Engine/Plugins/Runtime/ModelViewViewModel/Source/ModelViewViewModelBlueprint/Private/MVVMBlueprintView.cpp
+++ b/Engine/Plugins/Runtime/ModelViewViewModel/Source/ModelViewViewModelBlueprint/Private/MVVMBlueprintView.cpp
@@ -143,6 +143,7 @@ FMVVMBlueprintViewBinding& UMVVMBlueprintView::AddBinding(const UWidget* Widget,
 	FMVVMBlueprintViewBinding& NewBinding = Bindings.AddDefaulted_GetRef();
 	NewBinding.WidgetPath.SetWidgetName(Widget->GetFName());
 	NewBinding.WidgetPath.SetBasePropertyPath(UE::MVVM::FMVVMConstFieldVariant(Property));
+	NewBinding.BindingId = FGuid::NewGuid();
 
 	OnBindingsUpdated.Broadcast();
 	return NewBinding;
@@ -174,6 +175,19 @@ const FMVVMBlueprintViewBinding* UMVVMBlueprintView::GetBindingAt(int32 Index) c
 	return nullptr;
 }
 
+void UMVVMBlueprintView::PostLoad()
+{
+	Super::PostLoad();
+
+	for (FMVVMBlueprintViewBinding& Binding : Bindings)
+	{
+		if (!Binding.BindingId.IsValid())
+		{
+			Binding.BindingId = FGuid::NewGuid();
+		}
+	}
+}
+
 #if WITH_EDITOR
 void UMVVMBlueprintView::PostEditChangeProperty(FPropertyChangedEvent& PropertyChangedEvent)
 {
diff --git a/Engine/Plugins/Runtime/ModelViewViewModel/Source/ModelViewViewModelBlueprint/Private/MVVMBlueprintViewBinding.cpp b/Engine/Plugins/Runtime/ModelViewViewModel/Source/ModelViewViewModelBlueprint/Private/MVVMBlueprintViewBinding.cpp
--- a/Engine/Plugins/Runtime/ModelViewViewModel/Source/ModelViewViewModelBlueprint/Private/MVVMBlueprintViewBinding.cpp
+++ b/Engine/Plugins/Runtime/ModelViewViewModel/Source/ModelViewViewModelBlueprint/Private/MVVMBlueprintViewBinding.cpp
@@ -1,112 +1,257 @@
 // Copyright Epic Games, Inc. All Rights Reserved.
 
 #include "MVVMBlueprintViewBinding.h"
+
+#include "Bindings/MVVMConversionFunctionHelper.h"
+#include "EdGraph/EdGraphPin.h"
+#include "K2Node_CallFunction.h"
 #include "MVVMBlueprintView.h"
 #include "MVVMWidgetBlueprintExtension_View.h"
+#include "WidgetBlueprint.h"
 #include "WidgetBlueprintExtension.h"
 
 #include UE_INLINE_GENERATED_CPP_BY_NAME(MVVMBlueprintViewBinding)
 
-FName FMVVMBlueprintViewBinding::GetFName(const UMVVMBlueprintView* View) const
+FName FMVVMBlueprintViewBinding::GetFName() const
 {
-	TStringBuilder<256> BindingName;
+	return *BindingId.ToString();
+}
 
-	const FMVVMBlueprintViewModelContext* ViewModel = nullptr;
-	if (View != nullptr)
+namespace UE::MVVM::Private
+{
+	void AppendViewModelPathString(const UMVVMBlueprintView* BlueprintView, const FMVVMBlueprintPropertyPath& ViewModelPath, FStringBuilderBase& PathBuilder)
 	{
-		ViewModel = View->FindViewModel(ViewModelPath.GetViewModelId());
-	}
+		if (ViewModelPath.IsEmpty())
+		{
+			PathBuilder << TEXT("<none>");
+			return;
+		}
 
-	if (ViewModel != nullptr)
-	{
-		BindingName << ViewModel->GetViewModelName();
-	}
-	else
-	{
-		BindingName << TEXT("Invalid");
-	}
+		const FMVVMBlueprintViewModelContext* ViewModel = BlueprintView != nullptr ? BlueprintView->FindViewModel(ViewModelPath.GetViewModelId()) : nullptr;
+		if (ViewModel)
+		{
+			PathBuilder << ViewModel->GetDisplayName().ToString();
+		}
+		else
+		{
+			PathBuilder << TEXT("<none>");
+		}
 
-	BindingName << TEXT("_");
-	BindingName << ViewModelPath.GetBasePropertyPath();
+		PathBuilder << TEXT(".");
 
-	if (BindingType == EMVVMBindingMode::TwoWay)
-	{
-		BindingName << TEXT("_Both_");
-	}
-	else if (UE::MVVM::IsForwardBinding(BindingType))
-	{
-		BindingName << TEXT("_To_");
+		const FString PropertyPath = ViewModelPath.GetBasePropertyPath();
+		if (!PropertyPath.IsEmpty())
+		{
+			PathBuilder << PropertyPath;
+		}
+		else
+		{
+			PathBuilder << TEXT("<none>");
+		}
 	}
-	else if (UE::MVVM::IsBackwardBinding(BindingType))
+	
+	FString GetBindingViewModelName(const UWidgetBlueprint* WidgetBlueprint, const FMVVMBlueprintViewBinding& Binding)
 	{
-		BindingName << TEXT("_From_");
+		UMVVMWidgetBlueprintExtension_View* ExtensionView = UMVVMWidgetBlueprintExtension_View::GetExtension<UMVVMWidgetBlueprintExtension_View>(WidgetBlueprint);
+		UMVVMBlueprintView* BlueprintView = ExtensionView->GetBlueprintView();
+
+		TStringBuilder<256> NameBuilder;
+
+		if (UEdGraph* Graph = ConversionFunctionHelper::GetGraph(WidgetBlueprint, Binding, true))
+		{
+			if (const UK2Node_CallFunction* CallFunctionNode = ConversionFunctionHelper::GetFunctionNode(Graph))
+			{
+				NameBuilder << CallFunctionNode->GetFunctionName();
+				NameBuilder << TEXT("(");
+
+				bool bFirst = true;
+
+				for (const UEdGraphPin* Pin : CallFunctionNode->Pins)
+				{
+					if (Pin->Direction == EGPD_Input)
+					{
+						if (!bFirst)
+						{
+							NameBuilder << TEXT(", ");
+						}
+
+						FMVVMBlueprintPropertyPath ArgumentPath = ConversionFunctionHelper::GetPropertyPathForArgument(WidgetBlueprint, CallFunctionNode, Pin->GetFName());
+						if (!ArgumentPath.IsEmpty())
+						{
+							AppendViewModelPathString(BlueprintView, ArgumentPath, NameBuilder);
+						}
+						else
+						{
+							NameBuilder << Pin->GetDefaultAsString();
+						}
+					}
+				}
+
+				NameBuilder << TEXT(")");
+			}
+			else
+			{
+				NameBuilder << TEXT("<error>");
+			}
+		}
+		else if (!Binding.Conversion.SourceToDestinationFunction.GetMemberName().IsNone())
+		{
+			if (const UFunction* SourceToDestFunction = Binding.Conversion.SourceToDestinationFunction.ResolveMember<UFunction>(WidgetBlueprint->SkeletonGeneratedClass))
+			{
+				NameBuilder << SourceToDestFunction->GetName();
+				NameBuilder << TEXT("(");
+			
+				AppendViewModelPathString(BlueprintView, Binding.ViewModelPath, NameBuilder);
+
+				NameBuilder << TEXT(")");
+			}
+			else
+			{
+				NameBuilder << Binding.Conversion.SourceToDestinationFunction.GetMemberName();
+				NameBuilder << TEXT("()");
+			}
+		}
+		else if (!Binding.ViewModelPath.IsEmpty())
+		{
+			AppendViewModelPathString(BlueprintView, Binding.ViewModelPath, NameBuilder);
+		}
+
+		FString Name = NameBuilder.ToString();
+		if (Name.IsEmpty())
+		{
+			Name = TEXT("<none>");
+		}
+
+		return Name;
 	}
 
-	if (View == nullptr || WidgetPath.GetWidgetName().IsNone())
+	void AppendWidgetPathString(const UWidgetBlueprint* WidgetBlueprint, const FMVVMBlueprintPropertyPath& WidgetPath, FStringBuilderBase& PathBuilder)
 	{
-		BindingName << TEXT("Invalid_");
+		if (WidgetBlueprint == nullptr || WidgetPath.GetWidgetName().IsNone())
+		{
+			PathBuilder << TEXT("<none>");
+		}
+		else if (WidgetBlueprint->GetFName() == WidgetPath.GetWidgetName())
+		{
+			PathBuilder << TEXT("self");
+		}
+		else
+		{
+			PathBuilder << WidgetPath.GetWidgetName();
+		}
+
+		PathBuilder << TEXT(".");
+
+		const FString PropertyPath = WidgetPath.GetBasePropertyPath();
+		if (!PropertyPath.IsEmpty())
+		{
+			PathBuilder << PropertyPath;
+		}
+		else
+		{
+			PathBuilder << TEXT("<none>");
+		}
 	}
-	else if (View->GetOuterUMVVMWidgetBlueprintExtension_View()->GetWidgetBlueprint()->GetFName() != WidgetPath.GetWidgetName())
+
+	FString GetBindingWidgetName(const UWidgetBlueprint* WidgetBlueprint, const FMVVMBlueprintViewBinding& Binding)
 	{
-		BindingName << WidgetPath.GetWidgetName();
-		BindingName << TEXT("_");
-	}
+		TStringBuilder<256> NameBuilder;
+
+		if (UEdGraph* Graph = ConversionFunctionHelper::GetGraph(WidgetBlueprint, Binding, false))
+		{
+			if (const UK2Node_CallFunction* CallFunctionNode = ConversionFunctionHelper::GetFunctionNode(Graph))
+			{
+				NameBuilder << CallFunctionNode->GetFunctionName();
+				NameBuilder << TEXT("(");
+
+				bool bFirst = true;
+
+				for (const UEdGraphPin* Pin : CallFunctionNode->Pins)
+				{
+					if (Pin->Direction == EGPD_Input)
+					{
+						if (!bFirst)
+						{
+							NameBuilder << TEXT(", ");
+						}
+
+						FMVVMBlueprintPropertyPath ArgumentPath = ConversionFunctionHelper::GetPropertyPathForArgument(WidgetBlueprint, CallFunctionNode, Pin->GetFName());
+						if (!ArgumentPath.IsEmpty())
+						{
+							AppendWidgetPathString(WidgetBlueprint, ArgumentPath, NameBuilder);
+						}
+						else
+						{
+							NameBuilder << Pin->GetDefaultAsString();
+						}
+					}
+				}
 
-	BindingName << WidgetPath.GetBasePropertyPath();
+				NameBuilder << TEXT(")");
+			}
+			else
+			{
+				NameBuilder << TEXT("<error>");
+			}
+		}
+		else if (!Binding.Conversion.DestinationToSourceFunction.GetMemberName().IsNone())
+		{
+			if (const UFunction* DestToSourceFunction = Binding.Conversion.DestinationToSourceFunction.ResolveMember<UFunction>(WidgetBlueprint->SkeletonGeneratedClass))
+			{
+				NameBuilder << DestToSourceFunction->GetName();
+				NameBuilder << TEXT("(");
+			
+				AppendWidgetPathString(WidgetBlueprint, Binding.WidgetPath, NameBuilder);
 
-	return BindingName.ToString();
+				NameBuilder << TEXT(")");
+			}
+			else
+			{
+				NameBuilder << Binding.Conversion.SourceToDestinationFunction.GetMemberName();
+				NameBuilder << TEXT("()");
+			}
+		}
+		else if (!Binding.WidgetPath.IsEmpty())
+		{
+			AppendWidgetPathString(WidgetBlueprint, Binding.WidgetPath, NameBuilder);
+		}
+
+		FString Name = NameBuilder.ToString();
+		if (Name.IsEmpty())
+		{
+			Name = TEXT("<none>");
+		}
+
+		return Name;
+	}
 }
 
-FString FMVVMBlueprintViewBinding::GetDisplayNameString(const UMVVMBlueprintView* View) const
+FString FMVVMBlueprintViewBinding::GetDisplayNameString(const UWidgetBlueprint* WidgetBlueprint) const
 {
-	const FMVVMBlueprintViewModelContext* ViewModel = nullptr;
-	if (View != nullptr)
-	{
-		ViewModel = View->FindViewModel(ViewModelPath.GetViewModelId());
-	}
+	check(WidgetBlueprint);
 
-	TStringBuilder<512> BindingName;
-	if (ViewModel != nullptr)
-	{
-		BindingName << ViewModel->GetViewModelName();
-	}
-	else
-	{
-		BindingName << TEXT("<none>");
-	}
+	TStringBuilder<256> NameBuilder;
 
-	BindingName << TEXT(".");
-	BindingName << ViewModelPath.GetBasePropertyPath();
+	NameBuilder << UE::MVVM::Private::GetBindingWidgetName(WidgetBlueprint, *this);
 
 	if (BindingType == EMVVMBindingMode::TwoWay)
 	{
-		BindingName << TEXT(" <-> ");
+		NameBuilder << TEXT(" <-> ");
 	}
 	else if (UE::MVVM::IsForwardBinding(BindingType))
 	{
-		BindingName << TEXT(" -> ");
+		NameBuilder << TEXT(" <- ");
 	}
 	else if (UE::MVVM::IsBackwardBinding(BindingType))
 	{
-		BindingName << TEXT(" <- ");
+		NameBuilder << TEXT(" -> ");
 	}
 	else
 	{
-		BindingName << TEXT(" ??? "); // shouldn't happen
+		NameBuilder << TEXT(" ??? "); // shouldn't happen
 	}
 
-	if (View == nullptr || WidgetPath.GetWidgetName().IsNone())
-	{
-		BindingName << TEXT("<none>.");
-	}
-	else if (View->GetOuterUMVVMWidgetBlueprintExtension_View()->GetWidgetBlueprint()->GetFName() != WidgetPath.GetWidgetName())
-	{
-		BindingName << WidgetPath.GetWidgetName();
-		BindingName << TEXT(".");
-	}
-	
-	BindingName << WidgetPath.GetBasePropertyPath();
-	
-	return BindingName.ToString();
-}
+	NameBuilder << UE::MVVM::Private::GetBindingViewModelName(WidgetBlueprint, *this);
 
+	return NameBuilder.ToString();
+}
diff --git a/Engine/Plugins/Runtime/ModelViewViewModel/Source/ModelViewViewModelBlueprint/Private/MVVMViewBlueprintCompiler.cpp b/Engine/Plugins/Runtime/ModelViewViewModel/Source/ModelViewViewModelBlueprint/Private/MVVMViewBlueprintCompiler.cpp
--- a/Engine/Plugins/Runtime/ModelViewViewModel/Source/ModelViewViewModelBlueprint/Private/MVVMViewBlueprintCompiler.cpp
+++ b/Engine/Plugins/Runtime/ModelViewViewModel/Source/ModelViewViewModelBlueprint/Private/MVVMViewBlueprintCompiler.cpp
@@ -18,9 +18,9 @@
 namespace UE::MVVM::Private
 {
 
-void FMVVMViewBlueprintCompiler::AddErrorForBinding(FMVVMBlueprintViewBinding& Binding, const UMVVMBlueprintView* View, const FString& Message) const
+void FMVVMViewBlueprintCompiler::AddErrorForBinding(FMVVMBlueprintViewBinding& Binding, const UWidgetBlueprint* WidgetBlueprint, const FString& Message) const
 {
-	const FString BindingName = Binding.GetDisplayNameString(View);
+	const FString BindingName = Binding.GetDisplayNameString(WidgetBlueprint);
 	WidgetBlueprintCompilerContext.MessageLog.Error(*(BindingName + TEXT(": ") + Message));
 	Binding.Errors.Add(FText::FromString(BindingName + TEXT(": ") + Message));
 }
@@ -485,6 +485,8 @@ bool FMVVMViewBlueprintCompiler::Compile(UWidgetBlueprintGeneratedClass* Class,
 
 bool FMVVMViewBlueprintCompiler::PreCompileBindingSources(UWidgetBlueprintGeneratedClass* Class, UMVVMBlueprintView* BlueprintView)
 {
+	const UWidgetBlueprint* WidgetBlueprint = WidgetBlueprintCompilerContext.WidgetBlueprint();
+
 	const int32 NumBindings = BlueprintView->GetNumBindings();
 	for (int32 Index = 0; Index < NumBindings; ++Index)
 	{
@@ -507,7 +509,7 @@ bool FMVVMViewBlueprintCompiler::PreCompileBindingSources(UWidgetBlueprintGenera
 			const TValueOrError<FBindingSourceContext, FString> CreatedBindingSourceContext = CreateBindingSourceContext(BlueprintView, Class, Binding.ViewModelPath);
 			if (CreatedBindingSourceContext.HasError())
 			{
-				AddErrorForBinding(Binding, BlueprintView, FString::Printf(TEXT("The property path '%s' is invalid. %s")
+				AddErrorForBinding(Binding, WidgetBlueprint, FString::Printf(TEXT("The property path '%s' is invalid. %s")
 					, *PropertyPathToString(BlueprintView, Binding.ViewModelPath)
 					, *CreatedBindingSourceContext.GetError()));
 				bIsBindingsValid = false;
@@ -517,19 +519,19 @@ bool FMVVMViewBlueprintCompiler::PreCompileBindingSources(UWidgetBlueprintGenera
 			FBindingSourceContext BindingSourceContext = CreatedBindingSourceContext.GetValue();
 			if (!IsPropertyPathValid(BindingSourceContext.PropertyPath))
 			{
-				AddErrorForBinding(Binding, BlueprintView, FString::Printf(TEXT("The property path '%s' is invalid."), *PropertyPathToString(BlueprintView, Binding.ViewModelPath)));
+				AddErrorForBinding(Binding, WidgetBlueprint, FString::Printf(TEXT("The property path '%s' is invalid."), *PropertyPathToString(BlueprintView, Binding.ViewModelPath)));
 				bIsBindingsValid = false;
 				continue;
 			}
 			if (BindingSourceContext.SourceClass == nullptr)
 			{
-				AddErrorForBinding(Binding, BlueprintView, TEXT("Internal error. The binding could not find its source."));
+				AddErrorForBinding(Binding, WidgetBlueprint, TEXT("Internal error. The binding could not find its source."));
 				bIsBindingsValid = false;
 				continue;
 			}
 			if (!BindingSourceContext.bIsRootWidget && BindingSourceContext.CompilerSourceContextIndex == INDEX_NONE)
 			{
-				AddErrorForBinding(Binding, BlueprintView, TEXT("Internal error. The binding could not find its source."));
+				AddErrorForBinding(Binding, WidgetBlueprint, TEXT("Internal error. The binding could not find its source."));
 				bIsBindingsValid = false;
 				continue;
 			}
@@ -544,7 +546,7 @@ bool FMVVMViewBlueprintCompiler::PreCompileBindingSources(UWidgetBlueprintGenera
 			const TValueOrError<FBindingSourceContext, FString> CreatedBindingSourceContext = CreateBindingSourceContext(BlueprintView, Class, Binding.WidgetPath);
 			if (CreatedBindingSourceContext.HasError())
 			{
-				AddErrorForBinding(Binding, BlueprintView, FString::Printf(TEXT("The property path '%s' is invalid. %s")
+				AddErrorForBinding(Binding, WidgetBlueprint, FString::Printf(TEXT("The property path '%s' is invalid. %s")
 					, *PropertyPathToString(BlueprintView, Binding.ViewModelPath)
 					, *CreatedBindingSourceContext.GetError()));
 				bIsBindingsValid = false;
@@ -554,19 +556,19 @@ bool FMVVMViewBlueprintCompiler::PreCompileBindingSources(UWidgetBlueprintGenera
 			FBindingSourceContext BindingSourceContext = CreatedBindingSourceContext.GetValue();
 			if (!IsPropertyPathValid(BindingSourceContext.PropertyPath))
 			{
-				AddErrorForBinding(Binding, BlueprintView, FString::Printf(TEXT("The property path '%s' is invalid."), *PropertyPathToString(BlueprintView, Binding.WidgetPath)));
+				AddErrorForBinding(Binding, WidgetBlueprint, FString::Printf(TEXT("The property path '%s' is invalid."), *PropertyPathToString(BlueprintView, Binding.WidgetPath)));
 				bIsBindingsValid = false;
 				continue;
 			}
 			if (BindingSourceContext.SourceClass == nullptr)
 			{
-				AddErrorForBinding(Binding, BlueprintView, TEXT("Internal error. The binding could not find its source."));
+				AddErrorForBinding(Binding, WidgetBlueprint, TEXT("Internal error. The binding could not find its source."));
 				bIsBindingsValid = false;
 				continue;
 			}
 			if (!BindingSourceContext.bIsRootWidget && BindingSourceContext.CompilerSourceContextIndex == INDEX_NONE)
 			{
-				AddErrorForBinding(Binding, BlueprintView, TEXT("Internal error. The binding could not find its source."));
+				AddErrorForBinding(Binding, WidgetBlueprint, TEXT("Internal error. The binding could not find its source."));
 				bIsBindingsValid = false;
 				continue;
 			}
@@ -760,6 +762,8 @@ bool FMVVMViewBlueprintCompiler::PreCompileBindings(UWidgetBlueprintGeneratedCla
 		return false;
 	}
 
+	const UWidgetBlueprint* WidgetBlueprint = WidgetBlueprintCompilerContext.WidgetBlueprint();
+
 	for (const FBindingSourceContext& BindingSourceContext : BindingSourceContexts)
 	{
 		FMVVMBlueprintViewBinding* BindingPtr = BlueprintView->GetBindingAt(BindingSourceContext.BindingIndex);
@@ -857,7 +861,7 @@ bool FMVVMViewBlueprintCompiler::PreCompileBindings(UWidgetBlueprintGeneratedCla
 		TValueOrError<FCompiledBindingLibraryCompiler::FFieldIdHandle, FString> AddFieldResult = AddFieldId(BindingSourceContext.SourceClass, true, Binding.BindingType, BindingSourceContext.FieldId.GetFieldName());
 		if (AddFieldResult.HasError())
 		{
-			AddErrorForBinding(Binding, BlueprintView, *FString::Printf(TEXT("The binding could not create its source. %s"), *AddFieldResult.GetError()));
+			AddErrorForBinding(Binding, WidgetBlueprint, *FString::Printf(TEXT("The binding could not create its source. %s"), *AddFieldResult.GetError()));
 			bIsBindingsValid = false;
 			continue;
 		}
@@ -868,7 +872,7 @@ bool FMVVMViewBlueprintCompiler::PreCompileBindings(UWidgetBlueprintGeneratedCla
 			SetterPath = CreateBindingDestinationPath(BlueprintView, Class, DestinationPath);
 			if (!IsPropertyPathValid(SetterPath))
 			{
-				AddErrorForBinding(Binding, BlueprintView, FString::Printf(TEXT("The property path '%s' is invalid."), *PropertyPathToString(BlueprintView, DestinationPath)));
+				AddErrorForBinding(Binding, WidgetBlueprint, FString::Printf(TEXT("The property path '%s' is invalid."), *PropertyPathToString(BlueprintView, DestinationPath)));
 				bIsBindingsValid = false;
 				continue;
 			}
@@ -884,7 +888,7 @@ bool FMVVMViewBlueprintCompiler::PreCompileBindings(UWidgetBlueprintGeneratedCla
 		const UFunction* ConversionFunction = ConversionFunctionReference.ResolveMember<UFunction>(Class);
 		if (!ConversionFunctionWrapper.IsNone() && ConversionFunction == nullptr)
 		{
-			AddErrorForBinding(Binding, BlueprintView, TEXT("A complex conversion function was created but could not be found."));
+			AddErrorForBinding(Binding, WidgetBlueprint, TEXT("A complex conversion function was created but could not be found."));
 			bIsBindingsValid = false;
 			continue;
 		}
@@ -892,7 +896,7 @@ bool FMVVMViewBlueprintCompiler::PreCompileBindings(UWidgetBlueprintGeneratedCla
 		TValueOrError<FCompilerBinding, FString> AddBindingResult = AddBinding(Class, BindingSourceContext.PropertyPath, SetterPath, ConversionFunction);
 		if (AddBindingResult.HasError())
 		{
-			AddErrorForBinding(Binding, BlueprintView, FString::Printf(TEXT("The binding could not be created. %s"), *AddBindingResult.GetError()));
+			AddErrorForBinding(Binding, WidgetBlueprint, FString::Printf(TEXT("The binding could not be created. %s"), *AddBindingResult.GetError()));
 			bIsBindingsValid = false;
 			continue;
 		}
diff --git a/Engine/Plugins/Runtime/ModelViewViewModel/Source/ModelViewViewModelBlueprint/Private/MVVMViewBlueprintCompiler.h b/Engine/Plugins/Runtime/ModelViewViewModel/Source/ModelViewViewModelBlueprint/Private/MVVMViewBlueprintCompiler.h
--- a/Engine/Plugins/Runtime/ModelViewViewModel/Source/ModelViewViewModelBlueprint/Private/MVVMViewBlueprintCompiler.h
+++ b/Engine/Plugins/Runtime/ModelViewViewModel/Source/ModelViewViewModelBlueprint/Private/MVVMViewBlueprintCompiler.h
@@ -71,7 +71,7 @@ private:
 
 	const FCompilerSourceCreatorContext* FindViewModelSource(FGuid Id) const;
 
-	void AddErrorForBinding(FMVVMBlueprintViewBinding& Binding, const UMVVMBlueprintView* BlueprintView, const FString& Message) const;
+	void AddErrorForBinding(FMVVMBlueprintViewBinding& Binding, const UWidgetBlueprint* WidgetBlueprint, const FString& Message) const;
 
 	TValueOrError<FBindingSourceContext, FString> CreateBindingSourceContext(const UMVVMBlueprintView* BlueprintView, const UWidgetBlueprintGeneratedClass* Class, const FMVVMBlueprintPropertyPath& PropertyPath);
 	TArray<FMVVMConstFieldVariant> CreateBindingDestinationPath(const UMVVMBlueprintView* BlueprintView, const UWidgetBlueprintGeneratedClass* Class, const FMVVMBlueprintPropertyPath& PropertyPath) const;
diff --git a/Engine/Plugins/Runtime/ModelViewViewModel/Source/ModelViewViewModelBlueprint/Public/Bindings/MVVMConversionFunctionHelper.h b/Engine/Plugins/Runtime/ModelViewViewModel/Source/ModelViewViewModelBlueprint/Public/Bindings/MVVMConversionFunctionHelper.h
--- a/Engine/Plugins/Runtime/ModelViewViewModel/Source/ModelViewViewModelBlueprint/Public/Bindings/MVVMConversionFunctionHelper.h
+++ b/Engine/Plugins/Runtime/ModelViewViewModel/Source/ModelViewViewModelBlueprint/Public/Bindings/MVVMConversionFunctionHelper.h
@@ -7,20 +7,28 @@
 
 class UClass;
 class UEdGraph;
+class UK2Node_CallFunction;
 class UMVVMBlueprintView;
 class UWidgetBlueprint;
 struct FMVVMBlueprintViewBinding;
 
 namespace UE::MVVM::ConversionFunctionHelper
 {
+	/** Find all BlueprintPropertyPath used by the given conversion function node. */
+	MODELVIEWVIEWMODELBLUEPRINT_API TArray<FMVVMBlueprintPropertyPath> GetAllArgumentPropertyPaths(const UWidgetBlueprint* WidgetBlueprint, const UK2Node_CallFunction* FunctionNode);
+	
+	/** Find the property path of a given argument in the conversion function. */
+	MODELVIEWVIEWMODELBLUEPRINT_API FMVVMBlueprintPropertyPath GetPropertyPathForArgument(const UWidgetBlueprint* WidgetBlueprint, const UK2Node_CallFunction* Function, FName ArgumentName);
 
-	/** Find all BlueprintPropertyPath used in the Graph. */
-	MODELVIEWVIEWMODELBLUEPRINT_API TArray<FMVVMBlueprintPropertyPath> FindAllPropertyPathInGraph(const UEdGraph* Graph, const UMVVMBlueprintView* View, UClass* Class);
-
-	/** Get the name of the conversion function wrapper if autogenerated. */
-	MODELVIEWVIEWMODELBLUEPRINT_API FName GetWrapperName(const UMVVMBlueprintView* View, const FMVVMBlueprintViewBinding& Binding, bool bSourceToDestination);
+	/** Create the name of the conversion function wrapper this binding should have. */
+	MODELVIEWVIEWMODELBLUEPRINT_API FName CreateWrapperName(const FMVVMBlueprintViewBinding& Binding, bool bSourceToDestination);
 
 	/** Get an existing conversion graph for the corresponding binding if it exist. */
 	MODELVIEWVIEWMODELBLUEPRINT_API UEdGraph* GetGraph(const UWidgetBlueprint* WidgetBlueprint, const FMVVMBlueprintViewBinding& Binding, bool bSourceToDestination);
 
+	/** Find the conversion function node from the given graph. */
+	MODELVIEWVIEWMODELBLUEPRINT_API UK2Node_CallFunction* GetFunctionNode(UEdGraph* Graph);
+
+	/** Mark the given function node as the conversion function node. */
+	MODELVIEWVIEWMODELBLUEPRINT_API void MarkAsConversionFunction(const UK2Node_CallFunction* FunctionNode, const FMVVMBlueprintViewBinding& Binding);
 } //namespace
\ No newline at end of file
diff --git a/Engine/Plugins/Runtime/ModelViewViewModel/Source/ModelViewViewModelBlueprint/Public/MVVMBlueprintView.h b/Engine/Plugins/Runtime/ModelViewViewModel/Source/ModelViewViewModelBlueprint/Public/MVVMBlueprintView.h
--- a/Engine/Plugins/Runtime/ModelViewViewModel/Source/ModelViewViewModelBlueprint/Public/MVVMBlueprintView.h
+++ b/Engine/Plugins/Runtime/ModelViewViewModel/Source/ModelViewViewModelBlueprint/Public/MVVMBlueprintView.h
@@ -70,6 +70,8 @@ public:
 	void WidgetRenamed(FName OldObjectName, FName NewObjectName);
 #endif
 
+	virtual void PostLoad() override;
+
 	DECLARE_EVENT(UMVVMBlueprintView, FOnBindingsUpdated);
 	FOnBindingsUpdated OnBindingsUpdated;
 
diff --git a/Engine/Plugins/Runtime/ModelViewViewModel/Source/ModelViewViewModelBlueprint/Public/MVVMBlueprintViewBinding.h b/Engine/Plugins/Runtime/ModelViewViewModel/Source/ModelViewViewModelBlueprint/Public/MVVMBlueprintViewBinding.h
--- a/Engine/Plugins/Runtime/ModelViewViewModel/Source/ModelViewViewModelBlueprint/Public/MVVMBlueprintViewBinding.h
+++ b/Engine/Plugins/Runtime/ModelViewViewModel/Source/ModelViewViewModelBlueprint/Public/MVVMBlueprintViewBinding.h
@@ -7,6 +7,7 @@
 
 #include "MVVMBlueprintViewBinding.generated.h"
 
+class UWidgetBlueprint;
 class UMVVMBlueprintView;
 
 /**
@@ -61,6 +62,9 @@ struct MODELVIEWVIEWMODELBLUEPRINT_API FMVVMBlueprintViewBinding
 	UPROPERTY(VisibleAnywhere, Category = "MVVM", Transient)
 	TArray<FText> Errors;
 
+	/** The unique ID of this binding. */
+	FGuid BindingId;
+
 	/** Whether the binding is enabled or disabled by default. The instance may enable the binding at runtime. */
 	UPROPERTY(EditAnywhere, Category = "MVVM")
 	bool bEnabled = true;
@@ -72,11 +76,11 @@ struct MODELVIEWVIEWMODELBLUEPRINT_API FMVVMBlueprintViewBinding
 	/**
 	 * Get an internal name. For use in the UI, use GetDisplayNameString()
 	 */
-	FName GetFName(const UMVVMBlueprintView* View) const;
+	FName GetFName() const;
 
 	/** 
 	 * Get a string that identifies this binding. 
-	 * This is of the form: ViewModel.Property -> Widget.Property
+	 * This is of the form: Widget.Property <- ViewModel.Property
 	 */
-	FString GetDisplayNameString(const UMVVMBlueprintView* View) const;
+	FString GetDisplayNameString(const UWidgetBlueprint* WidgetBlueprint) const;
 };
diff --git a/Engine/Plugins/Runtime/ModelViewViewModel/Source/ModelViewViewModelEditor/Private/MVVMEditorSubsystem.cpp b/Engine/Plugins/Runtime/ModelViewViewModel/Source/ModelViewViewModelEditor/Private/MVVMEditorSubsystem.cpp
--- a/Engine/Plugins/Runtime/ModelViewViewModel/Source/ModelViewViewModelEditor/Private/MVVMEditorSubsystem.cpp
+++ b/Engine/Plugins/Runtime/ModelViewViewModel/Source/ModelViewViewModelEditor/Private/MVVMEditorSubsystem.cpp
@@ -12,6 +12,7 @@
 #include "Engine/Engine.h"
 #include "Kismet/BlueprintFunctionLibrary.h"
 #include "Kismet2/BlueprintEditorUtils.h"
+#include "K2Node_BreakStruct.h"
 #include "K2Node_CallFunction.h"
 #include "K2Node_FunctionEntry.h"
 #include "K2Node_FunctionResult.h"
@@ -162,17 +163,9 @@ namespace UE::MVVM::Private
 			WidgetBlueprint->FunctionGraphs.RemoveSingle(Graph);
 		}
 	}
-} //namespace
 
-FName UMVVMEditorSubsystem::GetConversionFunctionWrapperName(const UWidgetBlueprint* WidgetBlueprint, const FMVVMBlueprintViewBinding& Binding, bool bSourceToDestination) const
-{
-	if (UMVVMBlueprintView* View = GetView(WidgetBlueprint))
-	{
-		return UE::MVVM::ConversionFunctionHelper::GetWrapperName(View, Binding, bSourceToDestination);
-	}
+} //namespace Private
 
-	return FName();
-}
 
 UMVVMBlueprintView* UMVVMEditorSubsystem::RequestView(UWidgetBlueprint* WidgetBlueprint) const
 {
@@ -297,23 +290,41 @@ UEdGraph* UMVVMEditorSubsystem::GetConversionFunctionGraph(const UWidgetBlueprin
 	return UE::MVVM::ConversionFunctionHelper::GetGraph(WidgetBlueprint, Binding, bSourceToDestination);
 }
 
-UFunction* UMVVMEditorSubsystem::GetConversionFunction(const UWidgetBlueprint* WidgetBlueprint, const FMVVMBlueprintViewBinding& Binding, bool bSourceToDestination) const
+UK2Node_CallFunction* UMVVMEditorSubsystem::GetConversionFunctionNode(const UWidgetBlueprint* WidgetBlueprint, const FMVVMBlueprintViewBinding& Binding, bool bSourceToDestination) const
 {
-	if (const UEdGraph* Graph = GetConversionFunctionGraph(WidgetBlueprint, Binding, bSourceToDestination))
+	if (UEdGraph* Graph = UE::MVVM::ConversionFunctionHelper::GetGraph(WidgetBlueprint, Binding, bSourceToDestination))
 	{
-		TArray<UK2Node_CallFunction*> FunctionNodes;
-		Graph->GetNodesOfClass<UK2Node_CallFunction>(FunctionNodes);
+		return UE::MVVM::ConversionFunctionHelper::GetFunctionNode(Graph);
+	}
+	return nullptr;
+}
 
-		if (FunctionNodes.Num() == 1)
-		{
-			return FunctionNodes[0]->GetTargetFunction();
-		}
+UFunction* UMVVMEditorSubsystem::GetConversionFunction(const UWidgetBlueprint* WidgetBlueprint, const FMVVMBlueprintViewBinding& Binding, bool bSourceToDestination) const
+{
+	UK2Node_CallFunction* CallFunctionNode = GetConversionFunctionNode(WidgetBlueprint, Binding, bSourceToDestination);
+	if (CallFunctionNode)
+	{
+		return CallFunctionNode->GetTargetFunction();
 	}
 	
+	// simple conversion functions don't have a wrapper graph
 	FMemberReference ConversionFunctionReference = bSourceToDestination ? Binding.Conversion.SourceToDestinationFunction : Binding.Conversion.DestinationToSourceFunction;
 	return ConversionFunctionReference.ResolveMember<UFunction>(WidgetBlueprint->SkeletonGeneratedClass);
 }
 
+UEdGraphPin* UMVVMEditorSubsystem::GetConversionFunctionArgumentPin(const UWidgetBlueprint* WidgetBlueprint, const FMVVMBlueprintViewBinding& Binding, FName ParameterName, bool bSourceToDestination)
+{
+	UK2Node_CallFunction* CallFunctionNode = GetConversionFunctionNode(WidgetBlueprint, Binding, bSourceToDestination);
+	if (CallFunctionNode == nullptr)
+	{
+		return nullptr;
+	}
+
+	UEdGraphPin* Pin = CallFunctionNode->FindPin(ParameterName, EGPD_Input);
+	return Pin;
+}
+
+
 UEdGraph* UMVVMEditorSubsystem::CreateConversionFunctionWrapperGraph(UWidgetBlueprint* WidgetBlueprint, const FMVVMBlueprintViewBinding& Binding, const UFunction* ConversionFunction, bool bSourceToDestination)
 {
 	if (UMVVMBlueprintView* View = GetView(WidgetBlueprint))
@@ -324,7 +335,7 @@ UEdGraph* UMVVMEditorSubsystem::CreateConversionFunctionWrapperGraph(UWidgetBlue
 			TArray<const FProperty*> Arguments = ArgumentsResult.StealValue();
 			if (Arguments.Num() > 1)
 			{
-				const FName WrapperName = GetConversionFunctionWrapperName(WidgetBlueprint, Binding, bSourceToDestination);
+				const FName WrapperName = UE::MVVM::ConversionFunctionHelper::CreateWrapperName(Binding, bSourceToDestination);
 
 				UEdGraph* NewGraph = UE::MVVM::Private::CreateFunctionGraph(
 					WidgetBlueprint,
@@ -342,6 +353,8 @@ UEdGraph* UMVVMEditorSubsystem::CreateConversionFunctionWrapperGraph(UWidgetBlue
 				CallFunctionNode->NodePosX = (FunctionEntry->NodePosX + FunctionResult->NodePosX) / 2;
 				CallFunctionCreator.Finalize();
 
+				UE::MVVM::ConversionFunctionHelper::MarkAsConversionFunction(CallFunctionNode, Binding);
+
 				UEdGraphPin* FunctionEntryThenPin = FunctionEntry->FindPin(UEdGraphSchema_K2::PN_Then, EGPD_Output);
 				UEdGraphPin* FunctionResultExecPin = FunctionResult->GetExecPin();
 			
@@ -386,29 +399,6 @@ UEdGraph* UMVVMEditorSubsystem::CreateConversionFunctionWrapperGraph(UWidgetBlue
 	return nullptr;
 }
 
-UEdGraphPin* UMVVMEditorSubsystem::FindConversionFunctionArgumentPin(const UWidgetBlueprint* WidgetBlueprint, const FMVVMBlueprintViewBinding& Binding, FName ParameterName, bool bSourceToDestination)
-{
-	const UEdGraph* Graph = GetConversionFunctionGraph(WidgetBlueprint, Binding, bSourceToDestination);
-	if (Graph == nullptr)
-	{
-		return nullptr;
-	}
-
-	TArray<UK2Node_CallFunction*> FunctionNodes;
-	Graph->GetNodesOfClass<UK2Node_CallFunction>(FunctionNodes);
-
-	if (FunctionNodes.Num() != 1)
-	{
-		// ambiguous result, no idea what our function node is
-		return nullptr;
-	}
-
-	UK2Node_CallFunction* CallFunctionNode = FunctionNodes[0];
-	UEdGraphPin* Pin = CallFunctionNode->FindPin(ParameterName, EGPD_Input);
-	
-	return Pin;
-}
-
 void UMVVMEditorSubsystem::SetSourceToDestinationConversionFunction(UWidgetBlueprint* WidgetBlueprint, FMVVMBlueprintViewBinding& Binding, const UFunction* ConversionFunction)
 {
 	if (UMVVMBlueprintView* View = GetView(WidgetBlueprint))
@@ -718,74 +708,14 @@ FMVVMBlueprintPropertyPath UMVVMEditorSubsystem::GetPathForConversionFunctionArg
 			return Binding.WidgetPath;
 		}
 	}
-	 
-	UEdGraphNode* ConversionNode = UE::MVVM::Private::FindConversionNode(WrapperGraph);
 
-	const UEdGraphPin* ArgumentPin = ConversionNode->FindPin(ArgumentName, EGPD_Input);
-	if (ArgumentPin == nullptr || ArgumentPin->LinkedTo.Num() == 0)
+	const UK2Node_CallFunction* FunctionNode = GetConversionFunctionNode(WidgetBlueprint, Binding, bSourceToDestination);
+	if (FunctionNode == nullptr)
 	{
 		return FMVVMBlueprintPropertyPath();
 	}
 
-	TArray<FMemberReference> PathParts;
-	TDeque<UEdGraphNode*> NodesToSearch;
-	NodesToSearch.PushLast(ArgumentPin->LinkedTo[0]->GetOwningNode());
-
-	while (NodesToSearch.Num() > 0)
-	{
-		UEdGraphNode* Node = NodesToSearch[0];
-		NodesToSearch.PopFirst();
-
-		if (UK2Node_VariableGet* Getter = Cast<UK2Node_VariableGet>(Node))
-		{
-			PathParts.Add(Getter->VariableReference);
-		}
-
-		for (UEdGraphPin* Pin : Node->Pins)
-		{
-			if (Pin->Direction == EGPD_Input && Pin->LinkedTo.Num() > 0)
-			{
-				NodesToSearch.PushLast(Pin->LinkedTo[0]->GetOwningNode());
-			}
-		}
-	}
-
-	if (PathParts.Num() == 0)
-	{
-		return FMVVMBlueprintPropertyPath();
-	}
-
-	Algo::Reverse(PathParts);
-
-	FMVVMBlueprintPropertyPath ResultPath;
-
-	FMemberReference RootRef = PathParts[0];
-	PathParts.RemoveAt(0);
-
-	if (const FObjectProperty* Property = CastField<FObjectProperty>(RootRef.ResolveMember<FProperty>(WidgetBlueprint->SkeletonGeneratedClass)))
-	{
-		if (Property->PropertyClass->ImplementsInterface(UNotifyFieldValueChanged::StaticClass()))
-		{
-			UE::MVVM::FBindingSource Source = UE::MVVM::FBindingSource::CreateForViewModel(WidgetBlueprint, Property->GetFName());
-			ResultPath.SetViewModelId(Source.ViewModelId);
-		}
-		else if (Property->PropertyClass->IsChildOf<UWidget>() || Property->PropertyClass->IsChildOf<UWidgetBlueprint>())
-		{
-			UE::MVVM::FBindingSource Source = UE::MVVM::FBindingSource::CreateForWidget(WidgetBlueprint, Property->GetFName());
-			ResultPath.SetWidgetName(Source.Name);
-		}
-	}
-
-	if (PathParts.Num() == 0)
-	{
-		return ResultPath;
-	}
-
-	// TODO: This only takes one part, because FMVVMBlueprintPropertyPath only takes one part
-	UE::MVVM::FMVVMConstFieldVariant Variant = UE::MVVM::FMVVMConstFieldVariant(PathParts.Last().ResolveMember<FProperty>(WidgetBlueprint->SkeletonGeneratedClass));
-	ResultPath.SetBasePropertyPath(Variant);
-
-	return ResultPath;
+	return UE::MVVM::ConversionFunctionHelper::GetPropertyPathForArgument(WidgetBlueprint, FunctionNode, ArgumentName);
 }
 
 void UMVVMEditorSubsystem::SetPathForConversionFunctionArgument(UWidgetBlueprint* WidgetBlueprint, FMVVMBlueprintViewBinding& Binding, FName ArgumentName, const FMVVMBlueprintPropertyPath& Path, bool bSourceToDestination) const
@@ -933,16 +863,19 @@ void UMVVMEditorSubsystem::SetPathForConversionFunctionArgument(UWidgetBlueprint
 			const FStructProperty* StructProperty = CastField<FStructProperty>(Property);
 			if (!bLastField && StructProperty != nullptr)
 			{
-				// TODO: Add struct break node
-				ensureAlwaysMsgf(false, TEXT("Struct nodes in the middle of the path are not yet supported. Path: %s"), *Path.GetBasePropertyPath());
+				FGraphNodeCreator<UK2Node_BreakStruct> BreakCreator(*ConversionFunctionGraph);
+				UK2Node_BreakStruct* BreakNode = BreakCreator.CreateNode();
+				BreakNode->AllocateDefaultPins();
+				BreakCreator.Finalize();
+
+				NewNode = BreakNode;
+
 				PreviousClass = nullptr;
 			}
 			else if (PreviousClass != nullptr)
 			{
 				FGraphNodeCreator<UK2Node_VariableGet> GetterCreator(*ConversionFunctionGraph);
 				UK2Node_VariableGet* GetterNode = GetterCreator.CreateNode();
-				GetterNode->NodePosX = PosX;
-				GetterNode->NodePosY = PosY;
 				GetterNode->SetFromProperty(Property, false, PreviousClass);
 				GetterNode->AllocateDefaultPins();
 				GetterCreator.Finalize();
@@ -965,12 +898,12 @@ void UMVVMEditorSubsystem::SetPathForConversionFunctionArgument(UWidgetBlueprint
 
 			FGraphNodeCreator<UK2Node_CallFunction> CallFunctionCreator(*ConversionFunctionGraph);
 			UK2Node_CallFunction* FunctionNode = CallFunctionCreator.CreateNode();
-			FunctionNode->NodePosX = PosX;
-			FunctionNode->NodePosY = PosY;
 			FunctionNode->SetFromFunction(Function);
 			FunctionNode->AllocateDefaultPins();
 			CallFunctionCreator.Finalize();
 
+			NewNode = FunctionNode;
+
 			const FProperty* ReturnProperty = UE::MVVM::BindingHelper::GetReturnProperty(Function);
 			if (const FObjectProperty* ObjectProperty = CastField<FObjectProperty>(ReturnProperty))
 			{
@@ -980,8 +913,6 @@ void UMVVMEditorSubsystem::SetPathForConversionFunctionArgument(UWidgetBlueprint
 			{
 				PreviousClass = nullptr;
 			}
-
-			NewNode = FunctionNode;
 		}
 		else
 		{
@@ -1022,6 +953,9 @@ void UMVVMEditorSubsystem::SetPathForConversionFunctionArgument(UWidgetBlueprint
 			}
 		}
 
+		NewNode->NodePosX = PosX;
+		NewNode->NodePosY = PosY;
+
 		PosX += 300; 
 	}
 
diff --git a/Engine/Plugins/Runtime/ModelViewViewModel/Source/ModelViewViewModelEditor/Private/Widgets/SMVVMFunctionParameter.cpp b/Engine/Plugins/Runtime/ModelViewViewModel/Source/ModelViewViewModelEditor/Private/Widgets/SMVVMFunctionParameter.cpp
--- a/Engine/Plugins/Runtime/ModelViewViewModel/Source/ModelViewViewModelEditor/Private/Widgets/SMVVMFunctionParameter.cpp
+++ b/Engine/Plugins/Runtime/ModelViewViewModel/Source/ModelViewViewModelEditor/Private/Widgets/SMVVMFunctionParameter.cpp
@@ -44,7 +44,7 @@ void SFunctionParameter::Construct(const FArguments& InArgs)
 
 	TSharedRef<SWidget> ValueWidget = SNullWidget::NullWidget;
 
-	UEdGraphPin* Pin = EditorSubsystem->FindConversionFunctionArgumentPin(InArgs._WidgetBlueprint, *Binding, ParameterName, bSourceToDestination);
+	UEdGraphPin* Pin = EditorSubsystem->GetConversionFunctionArgumentPin(InArgs._WidgetBlueprint, *Binding, ParameterName, bSourceToDestination);
 	if (Pin != nullptr)
 	{
 		// create a new pin widget so that we can get the default value widget out of it
diff --git a/Engine/Plugins/Runtime/ModelViewViewModel/Source/ModelViewViewModelEditor/Private/Widgets/SMVVMViewBindingListView.cpp b/Engine/Plugins/Runtime/ModelViewViewModel/Source/ModelViewViewModelEditor/Private/Widgets/SMVVMViewBindingListView.cpp
--- a/Engine/Plugins/Runtime/ModelViewViewModel/Source/ModelViewViewModelEditor/Private/Widgets/SMVVMViewBindingListView.cpp
+++ b/Engine/Plugins/Runtime/ModelViewViewModel/Source/ModelViewViewModelEditor/Private/Widgets/SMVVMViewBindingListView.cpp
@@ -472,7 +472,7 @@ private:
 					ErrorItems.Add(MakeShared<FText>(ErrorText));
 				}
 
-				const FText BindingDisplayName = FText::FromString(ViewBinding->GetDisplayNameString(View));
+				const FText BindingDisplayName = FText::FromString(ViewBinding->GetDisplayNameString(WidgetBlueprint));
 				ErrorDialog = SNew(SCustomDialog)
 					.Title(FText::Format(LOCTEXT("Compilation Errors", "Compilation Errors for {0}"), BindingDisplayName))
 					.Buttons({
@@ -865,7 +865,7 @@ class SFunctionParameterRow : public STableRow<TSharedPtr<FBindingEntry>>
 
 		bool bSimpleConversionFunction = false;
 		
-		if (UEdGraphPin* Pin = EditorSubsystem->FindConversionFunctionArgumentPin(InWidgetBlueprint, *Binding, Entry->GetName(), bSourceToDestination))
+		if (UEdGraphPin* Pin = EditorSubsystem->GetConversionFunctionArgumentPin(InWidgetBlueprint, *Binding, Entry->GetName(), bSourceToDestination))
 		{
 			PrimaryBrush = FBlueprintEditor::GetVarIconAndColorFromPinType(Pin->PinType, PrimaryColor, SecondaryBrush, SecondaryColor);
 			DisplayName = Pin->GetDisplayName();
@@ -1163,6 +1163,8 @@ void SBindingsList::OnDeleteSelected()
 	{
 		if (UMVVMBlueprintView* BlueprintView = MVVMExtensionPtr->GetBlueprintView())
 		{
+			const UWidgetBlueprint* WidgetBlueprint = MVVMExtensionPtr->GetWidgetBlueprint();
+
 			TArray<const FMVVMBlueprintViewBinding*> BindingsToRemove;
 			Private::GatherChildBindings(BlueprintView, Selection, BindingsToRemove);
 
@@ -1174,7 +1176,7 @@ void SBindingsList::OnDeleteSelected()
 			TArray<FText> BindingDisplayNames;
 			for (const FMVVMBlueprintViewBinding* Binding : BindingsToRemove)
 			{
-				BindingDisplayNames.Add(FText::FromString(Binding->GetDisplayNameString(BlueprintView)));
+				BindingDisplayNames.Add(FText::FromString(Binding->GetDisplayNameString(WidgetBlueprint)));
 			}
 
 			const FText Message = FText::Format(BindingDisplayNames.Num() == 1 ?
diff --git a/Engine/Plugins/Runtime/ModelViewViewModel/Source/ModelViewViewModelEditor/Public/MVVMEditorSubsystem.h b/Engine/Plugins/Runtime/ModelViewViewModel/Source/ModelViewViewModelEditor/Public/MVVMEditorSubsystem.h
--- a/Engine/Plugins/Runtime/ModelViewViewModel/Source/ModelViewViewModelEditor/Public/MVVMEditorSubsystem.h
+++ b/Engine/Plugins/Runtime/ModelViewViewModel/Source/ModelViewViewModelEditor/Public/MVVMEditorSubsystem.h
@@ -13,6 +13,7 @@
 #include "MVVMEditorSubsystem.generated.h"
 
 class UEdGraph;
+class UK2Node_CallFunction;
 class UWidgetBlueprint;
 
 /** */
@@ -70,12 +71,15 @@ public:
 	UFUNCTION(BlueprintCallable, Category = "MVVM")
 	UFunction* GetConversionFunction(const UWidgetBlueprint* WidgetBlueprint, const FMVVMBlueprintViewBinding& Binding, bool bSourceToDestination) const;
 
+	UFUNCTION(BlueprintCallable, Category = "MVVM")
+	UK2Node_CallFunction* GetConversionFunctionNode(const UWidgetBlueprint* WidgetBlueprint, const FMVVMBlueprintViewBinding& Binding, bool bSourceToDestination) const;
+
 	UFUNCTION(BlueprintCallable, Category = "MVVM")
 	TArray<UFunction*> GetAvailableConversionFunctions(const UWidgetBlueprint* WidgetBlueprint, const FMVVMBlueprintPropertyPath& Source, const FMVVMBlueprintPropertyPath& Destination) const;
 
 	FMVVMBlueprintPropertyPath GetPathForConversionFunctionArgument(const UWidgetBlueprint* WidgetBlueprint, const FMVVMBlueprintViewBinding& Binding, FName ArgumentName, bool bSourceToDestination) const;
 	void SetPathForConversionFunctionArgument(UWidgetBlueprint* WidgetBlueprint, FMVVMBlueprintViewBinding& Binding, FName ArgumentName, const FMVVMBlueprintPropertyPath& Path, bool bSourceToDestination) const;
-	UEdGraphPin* FindConversionFunctionArgumentPin(const UWidgetBlueprint* WidgetBlueprint, const FMVVMBlueprintViewBinding& Binding, FName ParameterName, bool bSourceToDestination);
+	UEdGraphPin* GetConversionFunctionArgumentPin(const UWidgetBlueprint* WidgetBlueprint, const FMVVMBlueprintViewBinding& Binding, FName ParameterName, bool bSourceToDestination);
 
 	TArray<UE::MVVM::FBindingSource> GetBindableWidgets(const UWidgetBlueprint* WidgetBlueprint) const;
 	
