MVVM - Field selector can now display long paths, eg. Slot > Margin > Left.

Field selector now maintains conversion function selection after searching.

#rnx

diff --git a/Engine/Plugins/Runtime/ModelViewViewModel/Source/ModelViewViewModelEditor/Private/Widgets/SMVVMFieldEntry.cpp b/Engine/Plugins/Runtime/ModelViewViewModel/Source/ModelViewViewModelEditor/Private/Widgets/SMVVMFieldEntry.cpp
--- a/Engine/Plugins/Runtime/ModelViewViewModel/Source/ModelViewViewModelEditor/Private/Widgets/SMVVMFieldEntry.cpp
+++ b/Engine/Plugins/Runtime/ModelViewViewModel/Source/ModelViewViewModelEditor/Private/Widgets/SMVVMFieldEntry.cpp
@@ -3,6 +3,7 @@
 #include "SMVVMFieldEntry.h"
 #include "SMVVMFieldIcon.h"
 #include "Widgets/SBoxPanel.h"
+#include "Widgets/Images/SImage.h"
 
 #define LOCTEXT_NAMESPACE "MVVMFieldEntry"
 
@@ -32,13 +33,16 @@ FText GetFieldToolTip(const FMVVMConstFieldVariant& Field)
 {
 	if (!Field.IsEmpty())
 	{
-		if (Field.IsFunction())
+		if (Field.IsFunction() && Field.GetFunction() != nullptr)
 		{
 			return Field.GetFunction()->GetToolTipText();
 		}
-		if (Field.IsProperty())
+		if (Field.IsProperty() && Field.GetProperty() != nullptr)
 		{
-			return FText::Join(FText::FromString(TEXT("\n")), Field.GetProperty()->GetToolTipText(), FText::FromString(Field.GetProperty()->GetCPPType()));
+			return FText::Join(FText::FromString(TEXT("\n")), 
+				Field.GetProperty()->GetToolTipText(), 
+					FText::FromString(Field.GetProperty()->GetCPPType())
+				);
 		}
 	}
 
@@ -50,45 +54,63 @@ FText GetFieldToolTip(const FMVVMConstFieldVariant& Field)
 void SFieldEntry::Construct(const FArguments& InArgs)
 {
 	Field = InArgs._Field;
+	TextStyle = InArgs._TextStyle;
 
 	ChildSlot
 	[
-		SNew(SHorizontalBox)
-		+ SHorizontalBox::Slot()
+		SAssignNew(FieldBox, SHorizontalBox)
+	];
+
+	Refresh();
+}
+
+void SFieldEntry::Refresh()
+{
+	FieldBox->ClearChildren();
+
+	TArray<FMVVMConstFieldVariant> Fields = Field.GetFields();
+	for (int32 Index = 0; Index < Fields.Num(); ++Index)
+	{
+		FieldBox->AddSlot()
 		.HAlign(HAlign_Left)
 		.VAlign(VAlign_Center)
 		.AutoWidth()
 		[
-			SAssignNew(Icon, SFieldIcon)
-		]
-		+ SHorizontalBox::Slot()
+			SNew(SFieldIcon)
+			.Field(Fields[Index])
+		];
+
+		FieldBox->AddSlot()
 		.HAlign(HAlign_Fill)
 		.VAlign(VAlign_Center)
 		.Padding(4, 0, 0, 0)
+		.AutoWidth()
 		[
-			SAssignNew(Label, STextBlock)
-			.TextStyle(InArgs._TextStyle)
+			SNew(STextBlock)
+			.TextStyle(TextStyle)
 			.Clipping(EWidgetClipping::OnDemand)
-		]
-	];
-
-	Refresh();
-}
+			.Text(Private::GetFieldDisplayName(Fields[Index]))
+		];
 
-void SFieldEntry::Refresh()
-{
-	UE::MVVM::FMVVMConstFieldVariant Variant;
+		// if not the last, then we need to add a chevron separator
+		if (Index != Fields.Num() - 1)
+		{
+			FieldBox->AddSlot()
+			.HAlign(HAlign_Left)
+			.VAlign(VAlign_Center)
+			.Padding(6, 0)
+			.AutoWidth()
+			[
+				SNew(SImage)
+				.Image(FAppStyle::Get().GetBrush("Icons.ChevronRight"))
+			];
+		}
+	}
 
-	TArray<UE::MVVM::FMVVMConstFieldVariant> Fields = Field.GetFields();
 	if (Fields.Num() > 0)
 	{
-		Variant = Fields.Last();
+		SetToolTipText(Private::GetFieldToolTip(Fields.Last()));
 	}
-
-	SetToolTipText(Private::GetFieldToolTip(Variant));
-
-	Icon->RefreshBinding(Variant);
-	Label->SetText(Private::GetFieldDisplayName(Variant));
 }
 
 void SFieldEntry::SetField(const FMVVMBlueprintPropertyPath& InField)
@@ -96,7 +118,7 @@ void SFieldEntry::SetField(const FMVVMBlueprintPropertyPath& InField)
 	Field = InField;
 
 	Refresh();
-}
+} 
 
 } // namespace UE::MVVM
 
diff --git a/Engine/Plugins/Runtime/ModelViewViewModel/Source/ModelViewViewModelEditor/Private/Widgets/SMVVMFieldEntry.h b/Engine/Plugins/Runtime/ModelViewViewModel/Source/ModelViewViewModelEditor/Private/Widgets/SMVVMFieldEntry.h
--- a/Engine/Plugins/Runtime/ModelViewViewModel/Source/ModelViewViewModelEditor/Private/Widgets/SMVVMFieldEntry.h
+++ b/Engine/Plugins/Runtime/ModelViewViewModel/Source/ModelViewViewModelEditor/Private/Widgets/SMVVMFieldEntry.h
@@ -11,7 +11,7 @@
 #include "Widgets/SCompoundWidget.h"
 #include "Widgets/Text/STextBlock.h"
 
-class STextBlock;
+class SHorizontalBox;
 
 namespace UE::MVVM
 {
@@ -36,9 +36,9 @@ public:
 	void SetField(const FMVVMBlueprintPropertyPath& InField);
 
 private:
+	const FTextBlockStyle* TextStyle = nullptr;
 	FMVVMBlueprintPropertyPath Field;
-	TSharedPtr<SFieldIcon> Icon;
-	TSharedPtr<STextBlock> Label;
+	TSharedPtr<SHorizontalBox> FieldBox;
 };
 
 } // namespace UE::MVVM
diff --git a/Engine/Plugins/Runtime/ModelViewViewModel/Source/ModelViewViewModelEditor/Private/Widgets/SMVVMFieldSelector.cpp b/Engine/Plugins/Runtime/ModelViewViewModel/Source/ModelViewViewModelEditor/Private/Widgets/SMVVMFieldSelector.cpp
--- a/Engine/Plugins/Runtime/ModelViewViewModel/Source/ModelViewViewModelEditor/Private/Widgets/SMVVMFieldSelector.cpp
+++ b/Engine/Plugins/Runtime/ModelViewViewModel/Source/ModelViewViewModelEditor/Private/Widgets/SMVVMFieldSelector.cpp
@@ -83,11 +83,15 @@ void SFieldSelector::Construct(const FArguments& InArgs, const UWidgetBlueprint*
 		.ButtonContent()
 		[
 			SNew(SOverlay)
+
+			// is a property set?
 			+ SOverlay::Slot()
 			[
 				SAssignNew(SourceEntryBox, SHorizontalBox)
-				.Visibility_Lambda([this]() { return CachedSelectedConversionFunction == nullptr && !CachedSelectedField.IsEmpty() ? EVisibility::Visible : EVisibility::Collapsed; })
+				.Visibility_Lambda([this]() { return CachedSelectedConversionFunction == nullptr && CachedSelectedField.GetFields().Num() > 0 ? EVisibility::Visible : EVisibility::Collapsed; })
 			]
+
+			// is a conversion function set?
 			+ SOverlay::Slot()
 			[
 				SNew(SHorizontalBox)
@@ -110,6 +114,8 @@ void SFieldSelector::Construct(const FArguments& InArgs, const UWidgetBlueprint*
 					.Text_Lambda([this]() { return CachedSelectedConversionFunction != nullptr ? CachedSelectedConversionFunction->GetDisplayNameText() : FText::GetEmpty(); })
 				]
 			]
+
+			// nothing selected
 			+ SOverlay::Slot()
 			[
 				SNew(SBox)
@@ -141,7 +147,7 @@ void SFieldSelector::Construct(const FArguments& InArgs, const UWidgetBlueprint*
 			];
 
 		SourceEntryBox->AddSlot()
-			.Padding(8, 0)
+			.Padding(6, 0)
 			.HAlign(HAlign_Left)
 			.VAlign(VAlign_Center)
 			.AutoWidth()
@@ -232,17 +238,8 @@ bool SFieldSelector::IsClearEnabled() const
 
 FReply SFieldSelector::OnClearClicked()
 {
-	if (FixedSource.IsSet())
-	{
-		FMVVMBlueprintPropertyPath Path = CachedSelectedField;
-		Path.ResetBasePropertyPath();
-		SetPropertySelection(Path);
-	}
-	else
-	{
-		SetPropertySelection(FMVVMBlueprintPropertyPath());
-		SetConversionFunctionSelection(nullptr);
-	}
+	SetPropertySelection(FMVVMBlueprintPropertyPath());
+	SetConversionFunctionSelection(nullptr);
 	return FReply::Handled();
 }
 
@@ -267,6 +264,31 @@ void SFieldSelector::SetPropertySelection(const FMVVMBlueprintPropertyPath& Sele
 	Refresh();
 }
 
+TSharedPtr<SFieldSelector::FConversionFunctionItem> SFieldSelector::FindConversionFunctionCategory(const TArray<TSharedPtr<FConversionFunctionItem>>& Items, TArrayView<FString> CategoryNameParts) const
+{
+	if (CategoryNameParts.Num() > 0)
+	{
+		for (const TSharedPtr<FConversionFunctionItem>& Item : Items)
+		{
+			if (Item->GetCategoryName() == CategoryNameParts[0])
+			{
+				TArrayView<FString> RemainingParts = CategoryNameParts.RightChop(1);
+
+				// last category part, this is what we're looking for
+				if (RemainingParts.Num() == 0)
+				{
+					return Item;
+				}
+
+				// recurse into children
+				return FindConversionFunctionCategory(Item->Children, RemainingParts);
+			}
+		}
+	}
+
+	return TSharedPtr<FConversionFunctionItem>();
+}
+
 void SFieldSelector::OnSearchBoxTextChanged(const FText& NewText)
 {
 	if (BindingList.IsValid())
@@ -281,14 +303,38 @@ void SFieldSelector::OnSearchBoxTextChanged(const FText& NewText)
 
 	if (ConversionFunctionCategoryTree.IsValid())
 	{
+		TArray<TSharedPtr<FConversionFunctionItem>> OldSelectedCategories;
+		ConversionFunctionCategoryTree->GetSelectedItems(OldSelectedCategories);
+
 		FilterConversionFunctionCategories();
 		ConversionFunctionCategoryTree->RequestTreeRefresh();
+
+		TArray<TSharedPtr<FConversionFunctionItem>> NewSelectedCategories;
+
+		// reselect old selection
+		for (const TSharedPtr<FConversionFunctionItem>& OldItem : OldSelectedCategories)
+		{
+			if (TSharedPtr<FConversionFunctionItem> NewItem = FindConversionFunctionCategory(FilteredConversionFunctionRoot, OldItem->CategoryPath))
+			{
+				NewSelectedCategories.Add(NewItem);
+			}
+		}
+		
+		if (NewSelectedCategories.Num() > 0)
+		{
+			ConversionFunctionCategoryTree->SetItemSelection(NewSelectedCategories, true);
+		}
 	}
 
 	if (ConversionFunctionList.IsValid())
 	{
+		TArray<const UFunction*> OldSelectedFunctions;
+		ConversionFunctionList->GetSelectedItems(OldSelectedFunctions);
+
 		FilterConversionFunctions();
 		ConversionFunctionList->RequestListRefresh();
+
+		ConversionFunctionList->SetItemSelection(OldSelectedFunctions, true);
 	}
 }
 
@@ -424,7 +470,7 @@ int32 SFieldSelector::FilterConversionFunctionCategoryChildren(const TArray<FStr
 		if (NumChildren > 0)
 		{
 			TSharedPtr<FConversionFunctionItem>& NewItem = OutDestArray.Add_GetRef(MakeShared<FConversionFunctionItem>());
-			NewItem->CategoryName = SourceItem->CategoryName;
+			NewItem->CategoryPath = SourceItem->CategoryPath;
 			NewItem->Function = SourceItem->Function; 
 			NewItem->Children = FilteredChildren;
 			NewItem->NumFunctions = NumChildren;
@@ -454,7 +500,7 @@ void SFieldSelector::FilterConversionFunctionCategories()
 	}
 
 	TSharedPtr<FConversionFunctionItem> RootItem = FilteredConversionFunctionRoot.Add_GetRef(MakeShared<FConversionFunctionItem>());
-	RootItem->CategoryName = TEXT("Conversion Functions");
+	RootItem->CategoryPath = { TEXT("Conversion Functions") };
 	
 	RootItem->NumFunctions = FilterConversionFunctionCategoryChildren(FilterStrings, ConversionFunctionRoot[0]->Children, FilteredConversionFunctionRoot[0]->Children);
 	
@@ -467,15 +513,15 @@ int32 SFieldSelector::SortConversionFunctionItemsRecursive(TArray<TSharedPtr<FCo
 
 	Items.Sort([](const TSharedPtr<FConversionFunctionItem>& A, const TSharedPtr<FConversionFunctionItem>& B)
 	{
-		if (!A->CategoryName.IsEmpty() && !B->CategoryName.IsEmpty())
+		if (!A->GetCategoryName().IsEmpty() && !B->GetCategoryName().IsEmpty())
 		{
-			return A->CategoryName < B->CategoryName;
+			return A->GetCategoryName() < B->GetCategoryName();
 		}
-		if (!A->CategoryName.IsEmpty() && B->CategoryName.IsEmpty())
+		if (!A->GetCategoryName().IsEmpty() && B->GetCategoryName().IsEmpty())
 		{
 			return true;
 		}
-		if (A->CategoryName.IsEmpty() && !B->CategoryName.IsEmpty())
+		if (A->GetCategoryName().IsEmpty() && !B->GetCategoryName().IsEmpty())
 		{
 			return false;
 		}
@@ -518,28 +564,38 @@ void SFieldSelector::GenerateConversionFunctionItems()
 		}
 	}
 
+	auto AddFunctionToItem = [](const UFunction* Function, const TSharedPtr<FConversionFunctionItem>& Parent)
+	{
+		TSharedPtr<FConversionFunctionItem>& Item = Parent->Children.Add_GetRef(MakeShared<FConversionFunctionItem>());
+		Item->Function = Function;
+		Item->NumFunctions = 1;
+		Parent->NumFunctions += 1;
+	};
+
 	TSharedPtr<FConversionFunctionItem> CurrentSelectedItem;
 
-	TArray<FString> SubCategories;
+	TArray<FString> CategoryPath;
 
 	ConversionFunctionRoot.Reset();
 	TSharedPtr<FConversionFunctionItem>& RootItem = ConversionFunctionRoot.Add_GetRef(MakeShared<FConversionFunctionItem>());
-	RootItem->CategoryName = TEXT("Conversion Functions");
+	RootItem->CategoryPath = { TEXT("Conversion Functions") };
 
 	for (const UFunction* Function : AllConversionFunctions)
 	{
 		FText CategoryName = Function->GetMetaDataText("Category");
 		if (CategoryName.IsEmpty())
 		{
-			TSharedPtr<FConversionFunctionItem>& Item = RootItem->Children.Add_GetRef(MakeShared<FConversionFunctionItem>());
-			Item->Function = Function;
-			Item->NumFunctions = 1;
-			RootItem->NumFunctions += 1;
+			AddFunctionToItem(Function, RootItem);
 			continue;
 		}
 
-		SubCategories.Reset();
-		CategoryName.ToString().ParseIntoArray(SubCategories, TEXT("|"));
+		// split into subcategories and trim
+		CategoryPath.Reset();
+		CategoryName.ToString().ParseIntoArray(CategoryPath, TEXT("|"));
+		for (FString& SubCategory : CategoryPath)
+		{
+			SubCategory.TrimStartAndEndInline();
+		}
 
 		TSharedPtr<FConversionFunctionItem> ParentItem = RootItem;
 
@@ -548,18 +604,14 @@ void SFieldSelector::GenerateConversionFunctionItems()
 		// Math 
 		//   > Boolean
 		//     > AND
-		for (const FString& SubCategory : SubCategories)
+		for (int32 PathIndex = 0; PathIndex < CategoryPath.Num(); ++PathIndex)
 		{
-			const FString Trimmed = SubCategory.TrimStartAndEnd();
-
 			ParentItem->NumFunctions += 1;
-			ParentItem = FindOrCreateItemForCategory(ParentItem->Children, Trimmed);
+
+			ParentItem = FindOrCreateItemForCategory(ParentItem->Children, MakeArrayView(CategoryPath.GetData(), PathIndex + 1));
 		}
 
-		TSharedPtr<FConversionFunctionItem>& Item = ParentItem->Children.Add_GetRef(MakeShared<FConversionFunctionItem>());
-		Item->Function = Function;
-		Item->NumFunctions = 1;
-		ParentItem->NumFunctions += 1;
+		AddFunctionToItem(Function, ParentItem);
 	}
 
 	int32 NumItems = SortConversionFunctionItemsRecursive(ConversionFunctionRoot);
@@ -570,26 +622,26 @@ void SFieldSelector::ExpandAllToItem(const UFunction* Function)
 {
 	TArray<TSharedPtr<FConversionFunctionItem>> Path;
 
-	FText CategoryName = Function->GetMetaDataText("Category");
-	if (CategoryName.IsEmpty())
+	FText FullCategoryName = Function->GetMetaDataText("Category");
+	if (FullCategoryName.IsEmpty())
 	{
 		Path.Add(FilteredConversionFunctionRoot[0]);
 	}
 	else
 	{
-		TArray<FString> SubCategories;
-		CategoryName.ToString().ParseIntoArray(SubCategories, TEXT("|"));
+		TArray<FString> CategoryPath;
+		FullCategoryName.ToString().ParseIntoArray(CategoryPath, TEXT("|"));
 
 		TSharedPtr<FConversionFunctionItem> CurrentParent = FilteredConversionFunctionRoot[0];
 
-		for (const FString& SubCategory : SubCategories)
+		for (const FString& SubCategory : CategoryPath)
 		{
 			const FString Trimmed = SubCategory.TrimStartAndEnd();
 
 			TSharedPtr<FConversionFunctionItem>* FoundItem = 
 				CurrentParent->Children.FindByPredicate([Trimmed, Function](const TSharedPtr<FConversionFunctionItem>& Item)
 				{
-					return Item->CategoryName == Trimmed || Item->Function == Function;
+					return Item->GetCategoryName() == Trimmed || Item->Function == Function;
 				});
 
 			if (FoundItem != nullptr)
@@ -814,8 +866,14 @@ TSharedRef<SWidget> SFieldSelector::OnGetMenuContent()
 		// Single fixed source, don't show the separate source panel.
 		BindingList->AddSource(FixedSource.GetValue());
 	}
-
-	BindingList->SetSelectedProperty(CachedSelectedField);
+	else if (!CachedSelectedField.IsEmpty())
+	{
+		FBindingSource Source = CachedSelectedField.IsFromViewModel() ? 
+			FBindingSource::CreateForViewModel(WidgetBlueprint.Get(), CachedSelectedField.GetViewModelId()) : 
+			FBindingSource::CreateForWidget(WidgetBlueprint.Get(), CachedSelectedField.GetWidgetName());
+		BindingList->AddSource(Source);
+		BindingList->SetSelectedProperty(CachedSelectedField);
+	}
 	
 	TSharedPtr<SVerticalBox> BindingListVBox;
 		
@@ -975,32 +1033,36 @@ EFieldVisibility SFieldSelector::GetFieldVisibilityFlags() const
 	return Flags;
 }
 
-TSharedPtr<SFieldSelector::FConversionFunctionItem> SFieldSelector::FindOrCreateItemForCategory(TArray<TSharedPtr<FConversionFunctionItem>>& Items, const FString& Name)
+TSharedPtr<SFieldSelector::FConversionFunctionItem> SFieldSelector::FindOrCreateItemForCategory(TArray<TSharedPtr<FConversionFunctionItem>>& Items, TArrayView<FString> CategoryPath)
 {
+	check(CategoryPath.Num() > 0);
+
+	const FString& CategoryName = CategoryPath.Last();
+
 	int32 Idx = 0;
 	for (; Idx < Items.Num(); ++Idx)
 	{
 		// found item
-		if (Items[Idx]->CategoryName == Name)
+		if (Items[Idx]->GetCategoryName() == CategoryName)
 		{
 			return Items[Idx];
 		}
 
 		// passed the place where it should have been, break out
-		if (Items[Idx]->CategoryName > Name)
+		if (Items[Idx]->GetCategoryName() > CategoryName)
 		{
 			break;
 		}
 	}
 
 	TSharedPtr<FConversionFunctionItem> NewItem = Items.Insert_GetRef(MakeShared<FConversionFunctionItem>(), Idx);
-	NewItem->CategoryName = Name;
+	NewItem->CategoryPath = CategoryPath;
 	return NewItem;
 }
 
 TSharedRef<ITableRow> SFieldSelector::GenerateConversionFunctionCategoryRow(TSharedPtr<FConversionFunctionItem> Item, const TSharedRef<STableViewBase>& OwnerTable)
 {
-	FText DisplayName = FText::FormatOrdered(FText::FromString("{0} ({1})"), FText::FromString(Item->CategoryName), FText::FromString(LexToString(Item->NumFunctions)));
+	FText DisplayName = FText::FormatOrdered(FText::FromString("{0} ({1})"), FText::FromString(Item->GetCategoryName()), FText::FromString(LexToString(Item->NumFunctions)));
 
 	return SNew(STableRow<TSharedPtr<FConversionFunctionItem>>, OwnerTable)
 	[
@@ -1022,7 +1084,7 @@ TSharedRef<ITableRow> SFieldSelector::GenerateConversionFunctionCategoryRow(TSha
 			SNew(STextBlock)
 			.Font(Item == FilteredConversionFunctionRoot[0] ? FAppStyle::Get().GetFontStyle("NormalText") : FAppStyle::Get().GetFontStyle("BoldFont"))
 			.Text(DisplayName)
-			.ToolTipText(FText::FromString(Item->CategoryName))
+			.ToolTipText(FText::FromString(Item->GetCategoryName()))
 			.HighlightText_Lambda([this]() { return SearchBox.IsValid() ? SearchBox->GetText() : FText::GetEmpty(); })
 		]
 	];
@@ -1033,7 +1095,7 @@ void SFieldSelector::GetConversionFunctionCategoryChildren(TSharedPtr<FConversio
 	Algo::TransformIf(Item->Children, OutItems,
 		[](const TSharedPtr<FConversionFunctionItem>& Item)
 		{
-			return !Item->CategoryName.IsEmpty();
+			return !Item->GetCategoryName().IsEmpty();
 		},
 		[](const TSharedPtr<FConversionFunctionItem>& Item)
 		{
diff --git a/Engine/Plugins/Runtime/ModelViewViewModel/Source/ModelViewViewModelEditor/Private/Widgets/SMVVMFieldSelector.h b/Engine/Plugins/Runtime/ModelViewViewModel/Source/ModelViewViewModelEditor/Private/Widgets/SMVVMFieldSelector.h
--- a/Engine/Plugins/Runtime/ModelViewViewModel/Source/ModelViewViewModelEditor/Private/Widgets/SMVVMFieldSelector.h
+++ b/Engine/Plugins/Runtime/ModelViewViewModel/Source/ModelViewViewModelEditor/Private/Widgets/SMVVMFieldSelector.h
@@ -105,23 +105,25 @@ private:
 
 	struct FConversionFunctionItem
 	{
-		FString CategoryName;
-		const UFunction* Function = nullptr;
+		FString GetCategoryName() { return CategoryPath.Num() > 0 ? CategoryPath.Last() : FString(); }
+
+		TArray<FString> CategoryPath;
+		const UFunction* Function;
 		TArray<TSharedPtr<FConversionFunctionItem>> Children;
 		int32 NumFunctions = 0;
 	};
 
-	TSharedPtr<FConversionFunctionItem> FindOrCreateItemForCategory(TArray<TSharedPtr<FConversionFunctionItem>>& Items, const FString& Name);
+	TSharedPtr<FConversionFunctionItem> FindOrCreateItemForCategory(TArray<TSharedPtr<FConversionFunctionItem>>& Items, TArrayView<FString> CategoryPath);
 	TSharedRef<ITableRow> GenerateConversionFunctionCategoryRow(TSharedPtr<FConversionFunctionItem> Item, const TSharedRef<STableViewBase>& OwnerTable);
 	void GenerateConversionFunctionItems();
 	void FilterConversionFunctionCategories();
 	/** Recursively filter the items in SourceArray and place them into DestArray. Returns true if any items were added. */
 	int32 FilterConversionFunctionCategoryChildren(const TArray<FString>& FilterStrings, const TArray<TSharedPtr<FConversionFunctionItem>>& SourceArray, TArray<TSharedPtr<FConversionFunctionItem>>& OutDestArray);
-
 	void FilterConversionFunctions();
 
 	void OnConversionFunctionCategorySelected(TSharedPtr<FConversionFunctionItem> Item, ESelectInfo::Type);
 	void GetConversionFunctionCategoryChildren(TSharedPtr<FConversionFunctionItem> Item, TArray<TSharedPtr<FConversionFunctionItem>>& OutItems) const;
+	TSharedPtr<FConversionFunctionItem> FindConversionFunctionCategory(const TArray<TSharedPtr<FConversionFunctionItem>>& Items, TArrayView<FString> CategoryNameParts) const;
 	
 	void AddConversionFunctionChildrenRecursive(const TSharedPtr<FConversionFunctionItem>& Parent, TArray<const UFunction*>& OutFunctions);
 
