MVVM - Replaced all user-facing FStrings with FText. This includes all compiler error messages.

Binding display names now take into account complex conversion functions.

Complex conversion functions now compile correctly.

#rnx

diff --git a/Engine/Plugins/Runtime/ModelViewViewModel/Source/ModelViewViewModel/Private/Bindings/MVVMBindingHelper.cpp b/Engine/Plugins/Runtime/ModelViewViewModel/Source/ModelViewViewModel/Private/Bindings/MVVMBindingHelper.cpp
--- a/Engine/Plugins/Runtime/ModelViewViewModel/Source/ModelViewViewModel/Private/Bindings/MVVMBindingHelper.cpp
+++ b/Engine/Plugins/Runtime/ModelViewViewModel/Source/ModelViewViewModel/Private/Bindings/MVVMBindingHelper.cpp
@@ -6,6 +6,8 @@
 #include "UObject/PropertyAccessUtil.h"
 #include "UObject/UnrealType.h"
 
+#define LOCTEXT_NAMESPACE "MVVMBindingHelper"
+
 #ifndef UE_MVVM_ALLOW_AUTO_INTEGRAL_CONVERSION
 #define UE_MVVM_ALLOW_AUTO_INTEGRAL_CONVERSION 1
 #endif
@@ -185,80 +187,80 @@ namespace UE::MVVM::BindingHelper
 
 	namespace Private
 	{
-		FString TryGetPropertyTypeCommon(const FProperty* InProperty)
+		FText TryGetPropertyTypeCommon(const FProperty* InProperty)
 		{
 			if (InProperty == nullptr)
 			{
-				return TEXT("The property is invalid.");
+				return LOCTEXT("PropertyInvalid", "The property is invalid.");
 			}
 			
 			if (InProperty->ArrayDim != 1)
 			{
-				return FString::Printf(TEXT("The property '%s' is a static array."), *InProperty->GetName());
+				return FText::Format(LOCTEXT("PropertyStaticArray", "The property '{0}' is a static array."), FText::FromString(InProperty->GetName()));
 			}
 
 			if (InProperty->HasAnyPropertyFlags(CPF_Deprecated))
 			{
-				return FString::Printf(TEXT("The property '%s' is depreated."), *InProperty->GetName());
+				return FText::Format(LOCTEXT("PropertyDeprecated", "The property '{0}' is depreated."), FText::FromString(InProperty->GetName()));
 			}
 
 			if (InProperty->HasAnyPropertyFlags(CPF_EditorOnly))
 			{
-				return FString::Printf(TEXT("The property '%s' is only available in the editor."), *InProperty->GetName());
+				return FText::Format(LOCTEXT("PropertyEditorOnly", "The property '{0}' is only available in the editor."), FText::FromString(InProperty->GetName()));
 			}
 			
 			if (!InProperty->HasAnyPropertyFlags(CPF_BlueprintVisible))
 			{
-				return FString::Printf(TEXT("The property '%s' is not visible to the Blueprint script."), *InProperty->GetName());
+				return FText::Format(LOCTEXT("PropertyNotBlueprintVisible", "The property '{0}' is not visible to the Blueprint script."), FText::FromString(InProperty->GetName()));
 			}
 
-			return FString();
+			return FText();
 		}
 
-		FString TryGetPropertyTypeCommon(const UFunction* InFunction)
+		FText TryGetPropertyTypeCommon(const UFunction* InFunction)
 		{
 			if (InFunction == nullptr)
 			{
-				return (TEXT("The function is invalid."));
+				return LOCTEXT("FunctionInvalid", "The function is invalid.");
 			}
 
 			if (!InFunction->HasAllFunctionFlags(FUNC_BlueprintCallable))
 			{
-				return (FString::Printf(TEXT("The function '%s' is not BlueprintCallable."), *InFunction->GetName()));
+				return FText::Format(LOCTEXT("FunctionNotBlueprintCallable", "The function '{0}' is not BlueprintCallable."), FText::FromString(InFunction->GetName()));
 			}
 
 			if (InFunction->HasAllFunctionFlags(FUNC_Net))
 			{
-				return (FString::Printf(TEXT("The function '%s' is networked."), *InFunction->GetName()));
+				return FText::Format(LOCTEXT("FunctionNetworked", "The function '{0}' is networked."), FText::FromString(InFunction->GetName()));
 			}
 
 			if (InFunction->HasAllFunctionFlags(FUNC_Event))
 			{
 
-				return (FString::Printf(TEXT("The function '%s' is an event."), *InFunction->GetName()));
+				return FText::Format(LOCTEXT("FunctionEvent", "The function '{0}' is an event."), FText::FromString(InFunction->GetName()));
 			}
 
 			if (InFunction->HasAllFunctionFlags(FUNC_EditorOnly))
 			{
-				return (FString::Printf(TEXT("The function '%s' is editor only"), *InFunction->GetName()));
+				return FText::Format(LOCTEXT("FunctionEditorOnly", "The function '{0}' is editor only"), FText::FromString(InFunction->GetName()));
 			}
 
 #if WITH_EDITOR
 			{
 				if (InFunction->HasMetaData(NAME_DeprecatedFunction))
 				{
-					return (FString::Printf(TEXT("The function '%s' is deprecated"), *InFunction->GetName()));
+					return FText::Format(LOCTEXT("FunctionDeprecated", "The function '{0}' is deprecated"), FText::FromString(InFunction->GetName()));
 				}
 			}
 #endif
 
-			return FString();
+			return FText();
 		}
 	} // namespace
 
-	TValueOrError<const FProperty*, FString> TryGetPropertyTypeForSourceBinding(const FProperty* Property)
+	TValueOrError<const FProperty*, FText> TryGetPropertyTypeForSourceBinding(const FProperty* Property)
 	{
-		FString Result = Private::TryGetPropertyTypeCommon(Property);
+		FText Result = Private::TryGetPropertyTypeCommon(Property);
 		if (!Result.IsEmpty())
 		{
 			return MakeError(MoveTemp(Result));
@@ -267,9 +269,9 @@ namespace UE::MVVM::BindingHelper
 		return MakeValue(Property);
 	}
 
-	TValueOrError<const FProperty*, FString> TryGetPropertyTypeForSourceBinding(const UFunction* Function)
+	TValueOrError<const FProperty*, FText> TryGetPropertyTypeForSourceBinding(const UFunction* Function)
 	{
-		FString Result = Private::TryGetPropertyTypeCommon(Function);
+		FText Result = Private::TryGetPropertyTypeCommon(Function);
 		if (!Result.IsEmpty())
 		{
 			return MakeError(MoveTemp(Result));
@@ -277,33 +279,33 @@ namespace UE::MVVM::BindingHelper
 
 		if (!Function->HasAllFunctionFlags(FUNC_Const))
 		{
-			return MakeError(FString::Printf(TEXT("The function '%s' is not const."), *Function->GetName()));
+			return MakeError(FText::Format(LOCTEXT("FunctionNotConst", "The function '{0}' is not const."), FText::FromString(Function->GetName())));
 		}
 
 		//if (!InFunction->HasAllFunctionFlags(FUNC_BlueprintPure))
 		//{
-		//	return MakeError(FString::Printf(TEXT("The function '%s' is not pure."), *InFunction->GetName()));
+		//	return MakeError(FText::Format(LOCTEXT("FunctionNotPure", "The function '{0}' is not pure."), FText::FromString(Function->GetName()));
 		//}
 
 		if (Function->NumParms != 1)
 		{
-			return MakeError(FString::Printf(TEXT("The function '%s' doesn't have a return value."), *Function->GetName()));
+			return MakeError(FText::Format(LOCTEXT("FunctionWithoutReturn", "The function '{0}' doesn't have a return value."), FText::FromString(Function->GetName())));
 		}
 
 		const FProperty* ReturnProperty = GetReturnProperty(Function);
 		if (ReturnProperty == nullptr)
 		{
-			return MakeError(FString::Printf(TEXT("The return value for function '%s' is invalid."), *Function->GetName()));
+			return MakeError(FText::Format(LOCTEXT("FunctionReturnInvalid", "The return value for function '{0}' is invalid."), FText::FromString(Function->GetName())));
 		}
 
 		return MakeValue(ReturnProperty);
 	}
 
-	TValueOrError<const FProperty*, FString> TryGetPropertyTypeForSourceBinding(const FMVVMConstFieldVariant& InField)
+	TValueOrError<const FProperty*, FText> TryGetPropertyTypeForSourceBinding(const FMVVMConstFieldVariant& InField)
 	{
 		if (InField.IsEmpty())
 		{
-			return MakeError(TEXT("No source was provided for the binding."));
+			return MakeError(LOCTEXT("NoSourceForBinding", "No source was provided for the binding."));
 		}
 
 		if (InField.IsProperty())
@@ -317,9 +319,9 @@ namespace UE::MVVM::BindingHelper
 		}
 	}
 
-	TValueOrError<const FProperty*, FString> TryGetPropertyTypeForDestinationBinding(const FProperty* Property)
+	TValueOrError<const FProperty*, FText> TryGetPropertyTypeForDestinationBinding(const FProperty* Property)
 	{
-		FString Result = Private::TryGetPropertyTypeCommon(Property);
+		FText Result = Private::TryGetPropertyTypeCommon(Property);
 		if (!Result.IsEmpty())
 		{
 			return MakeError(MoveTemp(Result));
@@ -327,15 +329,15 @@ namespace UE::MVVM::BindingHelper
 
 		if (Property->HasAnyPropertyFlags(CPF_BlueprintReadOnly))
 		{
-			return MakeError(FString::Printf(TEXT("The property '%s' is read only."), *Property->GetName()));
+			return MakeError(FText::Format(LOCTEXT("PropertyReadOnly", "The property '{0}' is read only."), FText::FromString(Property->GetName())));
 		}
 
 		return MakeValue(Property);
 	}
 
-	TValueOrError<const FProperty*, FString> TryGetPropertyTypeForDestinationBinding(const UFunction* Function)
+	TValueOrError<const FProperty*, FText> TryGetPropertyTypeForDestinationBinding(const UFunction* Function)
 	{
-		FString Result = Private::TryGetPropertyTypeCommon(Function);
+		FText Result = Private::TryGetPropertyTypeCommon(Function);
 		if (!Result.IsEmpty())
 		{
 			return MakeError(MoveTemp(Result));
@@ -343,39 +345,39 @@ namespace UE::MVVM::BindingHelper
 
 		if (Function->NumParms != 1)
 		{
-			return MakeError(FString::Printf(TEXT("The function '%s' has more than one argument."), *Function->GetName()));
+			return MakeError(FText::Format(LOCTEXT("FunctionTooManyArguments", "The function '{0}' has more than one argument."), FText::FromString(Function->GetName())));
 		}
 
 		if (Function->HasAllFunctionFlags(FUNC_Const))
 		{
-			return MakeError(FString::Printf(TEXT("The function '%s' is const."), *Function->GetName()));
+			return MakeError(FText::Format(LOCTEXT("FunctionConst", "The function '{0}' is const."), FText::FromString(Function->GetName())));
 		}
 
 		if (Function->HasAllFunctionFlags(FUNC_BlueprintPure))
 		{
-			return MakeError(FString::Printf(TEXT("The function '%s' is pure."), *Function->GetName()));
+			return MakeError(FText::Format(LOCTEXT("FunctionPure", "The function '{0}' is pure."), FText::FromString(Function->GetName())));
 		}
 
 		const FProperty* ReturnProperty = GetReturnProperty(Function);
 		if (ReturnProperty != nullptr)
 		{
-			return MakeError(FString::Printf(TEXT("The function '%s' has a return value."), *Function->GetName()));
+			return MakeError(FText::Format(LOCTEXT("FunctionHasReturn", "The function '{0}' has a return value."), FText::FromString(Function->GetName())));
 		}
 
 		const FProperty* FirstArgumentProperty = GetFirstArgumentProperty(Function);
 		if (FirstArgumentProperty == nullptr)
 		{
-			return MakeError(FString::Printf(TEXT("The function '%s' doesn't not have a valid (single) argument."), *Function->GetName()));
+			return MakeError(FText::Format(LOCTEXT("FunctionNoArgument", "The function '{0}' doesn't have a valid (single) argument."), FText::FromString(Function->GetName())));
 		}
 
 		return MakeValue(FirstArgumentProperty);
 	}
 
-	TValueOrError<const FProperty*, FString> TryGetPropertyTypeForDestinationBinding(const FMVVMConstFieldVariant& InField)
+	TValueOrError<const FProperty*, FText> TryGetPropertyTypeForDestinationBinding(const FMVVMConstFieldVariant& InField)
 	{
 		if (InField.IsEmpty())
 		{
-			return MakeError(TEXT("No destination was provided for the binding."));
+			return MakeError(LOCTEXT("NoDestinationForBinding", "No destination was provided for the binding."));
 		}
 
 		if (InField.IsProperty())
@@ -389,9 +391,9 @@ namespace UE::MVVM::BindingHelper
 			return TryGetPropertyTypeForDestinationBinding(InField.GetFunction());
 		}
 	}
-	TValueOrError<const FProperty*, FString> TryGetReturnTypeForConversionFunction(const UFunction* InFunction)
+	TValueOrError<const FProperty*, FText> TryGetReturnTypeForConversionFunction(const UFunction* InFunction)
 	{
-		FString CommonResult = Private::TryGetPropertyTypeCommon(InFunction);
+		FText CommonResult = Private::TryGetPropertyTypeCommon(InFunction);
 		if (!CommonResult.IsEmpty())
 		{
 			return MakeError(MoveTemp(CommonResult));
@@ -399,21 +401,21 @@ namespace UE::MVVM::BindingHelper
 
 		if (!InFunction->HasAllFunctionFlags(FUNC_Static) && !InFunction->HasAllFunctionFlags(FUNC_Const | FUNC_BlueprintPure))
 		{
-			return MakeError(FString::Printf(TEXT("The function '%s' is not static or is not const and pure."), *InFunction->GetName()));
+			return MakeError(FText::Format(LOCTEXT("FunctionFlagsWrong", "The function '{0}' is not static or is not const and pure."), FText::FromString(InFunction->GetName())));
 		}
 
 		const FProperty* ReturnProperty = GetReturnProperty(InFunction);
 		if (ReturnProperty == nullptr)
 		{
-			return MakeError(FString::Printf(TEXT("The return value for function '%s' is invalid."), *InFunction->GetName()));
+			return MakeError(FText::Format(LOCTEXT("FunctionReturnInvalid", "The return value for function '{0}' is invalid."), FText::FromString(InFunction->GetName())));
 		}
 
 		return MakeValue(ReturnProperty);
 	}
 
-	TValueOrError<TArray<const FProperty*>, FString> TryGetArgumentsForConversionFunction(const UFunction* InFunction)
+	TValueOrError<TArray<const FProperty*>, FText> TryGetArgumentsForConversionFunction(const UFunction* InFunction)
 	{
-		FString CommonResult = Private::TryGetPropertyTypeCommon(InFunction);
+		FText CommonResult = Private::TryGetPropertyTypeCommon(InFunction);
 		if (!CommonResult.IsEmpty())
 		{
 			return MakeError(MoveTemp(CommonResult));
@@ -421,18 +423,18 @@ namespace UE::MVVM::BindingHelper
 
 		if (!InFunction->HasAllFunctionFlags(FUNC_Static) && !InFunction->HasAnyFunctionFlags(FUNC_Const | FUNC_BlueprintPure))
 		{
-			return MakeError(FString::Printf(TEXT("The function '%s' is not static or is not const and pure."), *InFunction->GetName()));
+			return MakeError(FText::Format(LOCTEXT("FunctionFlagsWrong", "The function '{0}' is not static or is not const and pure."), FText::FromString(InFunction->GetName())));
 		}
 
 		if (InFunction->NumParms < 2)
 		{
-			return MakeError(FString::Printf(TEXT("The function '%s' does not have the correct number of arguments."), *InFunction->GetName()));
+			return MakeError(FText::Format(LOCTEXT("FunctionWrongNumberOfArgs", "The function '{0}' does not have the correct number of arguments."), FText::FromString(InFunction->GetName())));
 		}
 
 		const FProperty* ReturnProperty = GetReturnProperty(InFunction);
 		if (ReturnProperty == nullptr)
 		{
-			return MakeError(FString::Printf(TEXT("The return value for function '%s' is invalid."), *InFunction->GetName()));
+			return MakeError(FText::Format(LOCTEXT("FunctionReturnInvalid", "The return value for function '{0}' is invalid."), FText::FromString(InFunction->GetName())));
 		}
 
 		TArray<const FProperty*> Arguments;
@@ -523,7 +525,7 @@ namespace UE::MVVM::BindingHelper
 			return false;
 		}
 
-		return Destination->SameType(Source)
+		return Source->SameType(Destination)
 				|| Private::IsNumericConversionRequired(Source, Destination)
 				|| Private::IsObjectPropertyCompatible(Source, Destination);
 	}
diff --git a/Engine/Plugins/Runtime/ModelViewViewModel/Source/ModelViewViewModel/Private/Bindings/MVVMFieldPathHelper.cpp b/Engine/Plugins/Runtime/ModelViewViewModel/Source/ModelViewViewModel/Private/Bindings/MVVMFieldPathHelper.cpp
--- a/Engine/Plugins/Runtime/ModelViewViewModel/Source/ModelViewViewModel/Private/Bindings/MVVMFieldPathHelper.cpp
+++ b/Engine/Plugins/Runtime/ModelViewViewModel/Source/ModelViewViewModel/Private/Bindings/MVVMFieldPathHelper.cpp
@@ -9,6 +9,8 @@
 #include "MVVMSubsystem.h"
 #include "Types/MVVMAvailableBinding.h"
 
+#define LOCTEXT_NAMESPACE "MVVMFieldPathHelper"
+
 namespace UE::MVVM::FieldPathHelper
 {
 
@@ -19,7 +21,7 @@ static const FName NAME_BlueprintGetter = "BlueprintGetter";
 static const FName NAME_BlueprintSetter = "BlueprintSetter";
 
 
-TValueOrError<UStruct*, FString> FindContainer(const FProperty* Property)
+TValueOrError<UStruct*, FText> FindContainer(const FProperty* Property)
 {
 	const FObjectPropertyBase* ObjectProperty = CastField<const FObjectPropertyBase>(Property);
 	const FStructProperty* StructProperty = CastField<const FStructProperty>(Property);
@@ -32,11 +34,11 @@ TValueOrError<UStruct*, FString> FindContainer(const FProperty* Property)
 	{
 		return MakeValue(StructProperty->Struct);
 	}
-	return MakeError(FString::Printf(TEXT("Only object or struct properties can be used as source paths. '%s' is a '%s'."), *Property->GetName(), *Property->GetClass()->GetName()));
+	return MakeError(FText::Format(LOCTEXT("InvalidPropertyTypeInSourcePath", "Only object or struct properties can be used as source paths. '{0}' is a '{1}'."), FText::FromString(Property->GetName()), FText::FromString(Property->GetClass()->GetName())));
 }
 
 
-TValueOrError<FMVVMConstFieldVariant, FString> TransformWithAccessor(UStruct* CurrentContainer, FMVVMConstFieldVariant CurrentField, bool bForReading)
+TValueOrError<FMVVMConstFieldVariant, FText> TransformWithAccessor(UStruct* CurrentContainer, FMVVMConstFieldVariant CurrentField, bool bForReading)
 {
 #if WITH_EDITORONLY_DATA
 	if (bForReading)
@@ -53,7 +55,10 @@ TValueOrError<FMVVMConstFieldVariant, FString> TransformWithAccessor(UStruct* Cu
 				}
 				else
 				{
-					return MakeError(FString::Printf(TEXT("The BlueprintGetter '%s' could not be found on object '%s'."), *BlueprintGetter, *CurrentContainer->GetName()));
+					return MakeError(FText::Format(LOCTEXT("BlueprintGetterNotFound", "The BlueprintGetter '{0}' could not be found on object '{1}'."), 
+						FText::FromString(BlueprintGetter), 
+						FText::FromString(CurrentContainer->GetName()))
+					);
 				}
 			}
 		}
@@ -72,7 +77,10 @@ TValueOrError<FMVVMConstFieldVariant, FString> TransformWithAccessor(UStruct* Cu
 				}
 				else
 				{
-					return MakeError(FString::Printf(TEXT("The BlueprintSetter '%s' could not be found on object %s."), *BlueprintSetter, *CurrentContainer->GetName()));
+					return MakeError(FText::Format(LOCTEXT("BlueprintSetterNotFound", "The BlueprintSetter '{0}' could not be found on object '{1}'."), 
+						FText::FromString(BlueprintSetter),
+						FText::FromString(CurrentContainer->GetName()))
+					);
 				}
 			}
 		}
@@ -84,19 +92,19 @@ TValueOrError<FMVVMConstFieldVariant, FString> TransformWithAccessor(UStruct* Cu
 } // namespace
 
 
-TValueOrError<TArray<FMVVMConstFieldVariant>, FString> GenerateFieldPathList(const UClass* InFrom, FStringView InFieldPath, bool bForReading)
+TValueOrError<TArray<FMVVMConstFieldVariant>, FText> GenerateFieldPathList(const UClass* InFrom, FStringView InFieldPath, bool bForReading)
 {
 	if (InFrom == nullptr)
 	{
-		return MakeError(TEXT("The source class is invalid."));
+		return MakeError(LOCTEXT("SourceClassInvalid", "The source class is invalid."));
 	}
 	if (InFieldPath.IsEmpty())
 	{
-		return MakeError(TEXT("The FieldPath is empty."));
+		return MakeError(LOCTEXT("FieldPathEmpty", "The field path is empty."));
 	}
 	if (InFieldPath[InFieldPath.Len() - 1] == TEXT('.'))
 	{
-		return MakeError(TEXT("The field path cannot end with a '.' character."));
+		return MakeError(LOCTEXT("FieldPathEndsInPeriod", "The field path cannot end with a '.' character."));
 	}
 
 	FMemMark Mark(FMemStack::Get());
@@ -112,11 +120,14 @@ TValueOrError<TArray<FMVVMConstFieldVariant>, FString> GenerateFieldPathList(con
 		FMVVMConstFieldVariant Field = BindingHelper::FindFieldByName(CurrentContainer, FMVVMBindingName(FName(FoundIndex, InFieldPath.GetData())));
 		if (Field.IsEmpty())
 		{
-			return MakeError(FString::Printf(TEXT("The field '%s' does not exist in the struct '%s'."), *FName(FoundIndex, InFieldPath.GetData()).ToString(), CurrentContainer ? *CurrentContainer->GetName() : TEXT("<none>")));
+			return MakeError(FText::Format(LOCTEXT("FieldDoesNotExistInStruct", "The field '%s' does not exist in the struct '%s'."), 
+				FText::FromName(FName(FoundIndex, InFieldPath.GetData())), 
+				CurrentContainer ? FText::FromString(CurrentContainer->GetName()) : LOCTEXT("None", "<none>"))
+			);
 		}
 		else if (Field.IsProperty())
 		{
-			TValueOrError<UStruct*, FString> FoundContainer = Private::FindContainer(Field.GetProperty());
+			TValueOrError<UStruct*, FText> FoundContainer = Private::FindContainer(Field.GetProperty());
 			if (FoundContainer.HasError())
 			{
 				return MakeError(FoundContainer.StealError());
@@ -127,7 +138,7 @@ TValueOrError<TArray<FMVVMConstFieldVariant>, FString> GenerateFieldPathList(con
 		if (Field.IsFunction())
 		{
 			const FProperty* ReturnProperty = BindingHelper::GetReturnProperty(Field.GetFunction());
-			TValueOrError<UStruct*, FString> FoundContainer = Private::FindContainer(ReturnProperty);
+			TValueOrError<UStruct*, FText> FoundContainer = Private::FindContainer(ReturnProperty);
 			if (FoundContainer.HasError())
 			{
 				return MakeError(FoundContainer.StealError());
@@ -145,7 +156,10 @@ TValueOrError<TArray<FMVVMConstFieldVariant>, FString> GenerateFieldPathList(con
 		FMVVMConstFieldVariant Field = BindingHelper::FindFieldByName(CurrentContainer, FMVVMBindingName(InFieldPath.GetData()));
 		if (Field.IsEmpty())
 		{
-			return MakeError(FString::Printf(TEXT("The field '%s' does not exist in the struct '%s'."), InFieldPath.GetData(), CurrentContainer ? *CurrentContainer->GetName() : TEXT("<none>")));
+			return MakeError(FText::Format(LOCTEXT("FieldDoesNotExistInStruct", "The field '{0}' does not exist in the struct '{1}'."), 
+				FText::FromStringView(InFieldPath), 
+				CurrentContainer ? FText::FromString(CurrentContainer->GetName()) : LOCTEXT("None", "<none>"))
+			);
 		}
 
 		Result.Add(Field);
@@ -169,11 +183,11 @@ TValueOrError<TArray<FMVVMConstFieldVariant>, FString> GenerateFieldPathList(con
  *				PropertyD and PropertyE cannot have Getter/BlueprintGetter if they are FStructProperty
  *		* We can only have one Setter/BlueprintSetter in the path
  */
-TValueOrError<TArray<FMVVMConstFieldVariant>, FString> GenerateFieldPathList(TArrayView<const FMVVMConstFieldVariant> InFieldPath, bool bForSourceBinding)
+TValueOrError<TArray<FMVVMConstFieldVariant>, FText> GenerateFieldPathList(TArrayView<const FMVVMConstFieldVariant> InFieldPath, bool bForSourceBinding)
 {
 	if (InFieldPath.Num() == 0)
 	{
-		return MakeError(TEXT("The FieldPath is empty."));
+		return MakeError(LOCTEXT("FieldPathEmpty", "The field path is empty."));
 	}
 
 	TArray<FMVVMConstFieldVariant> Result;
@@ -187,24 +201,27 @@ TValueOrError<TArray<FMVVMConstFieldVariant>, FString> GenerateFieldPathList(TAr
 
 		if (Field.IsEmpty())
 		{
-			return MakeError(FString::Printf(TEXT("The field '%d' does not exist."), Index));
+			return MakeError(FText::Format(LOCTEXT("FieldIsEmpty", "The field at index '{0}' is empty."), Index));
 		}
 
 		if (CurrentContainer == nullptr)
 		{
-			return MakeError(FString::Printf(TEXT("The field '%s' does not exist."), *Field.GetName().ToString()));
+			return MakeError(FText::Format(LOCTEXT("FieldDoesNotExist", "The field '{0}' does not exist."), FText::FromName(Field.GetName())));
 		}
 
 		const bool bIsChild = Field.GetOwner()->IsChildOf(CurrentContainer);
 		const bool bIsDownCast = CurrentContainer->IsChildOf(Field.GetOwner());
 		if (!(bIsChild || bIsDownCast))
 		{
-			return MakeError(FString::Printf(TEXT("The field '%s' does not exist in the struct '%s'."), *Field.GetName().ToString(), *CurrentContainer->GetName()));
+			return MakeError(FText::Format(LOCTEXT("FieldDoesNotExistInStruct", "The field '{0}' does not exist in the struct '{1}'."), 
+				FText::FromName(Field.GetName()), 
+				FText::FromString(CurrentContainer->GetName()))
+			);
 		}
 		
 		if (Field.IsProperty())
 		{
-			TValueOrError<FMVVMConstFieldVariant, FString> TransformedField = Private::TransformWithAccessor(CurrentContainer, Field, (bForSourceBinding || !bLastField));
+			TValueOrError<FMVVMConstFieldVariant, FText> TransformedField = Private::TransformWithAccessor(CurrentContainer, Field, (bForSourceBinding || !bLastField));
 			if (TransformedField.HasError())
 			{
 				return MakeError(TransformedField.StealError());
@@ -214,7 +231,7 @@ TValueOrError<TArray<FMVVMConstFieldVariant>, FString> GenerateFieldPathList(TAr
 
 			if (!bLastField && Field.IsProperty())
 			{
-				TValueOrError<UStruct*, FString> FoundContainer = Private::FindContainer(Field.GetProperty());
+				TValueOrError<UStruct*, FText> FoundContainer = Private::FindContainer(Field.GetProperty());
 				if (FoundContainer.HasError())
 				{
 					return MakeError(FoundContainer.StealError());
@@ -226,7 +243,7 @@ TValueOrError<TArray<FMVVMConstFieldVariant>, FString> GenerateFieldPathList(TAr
 		if (!bLastField && Field.IsFunction())
 		{
 			const FProperty* ReturnProperty = BindingHelper::GetReturnProperty(Field.GetFunction());
-			TValueOrError<UStruct*, FString> FoundContainer = Private::FindContainer(ReturnProperty);
+			TValueOrError<UStruct*, FText> FoundContainer = Private::FindContainer(ReturnProperty);
 			if (FoundContainer.HasError())
 			{
 				return MakeError(FoundContainer.StealError());
@@ -241,24 +258,24 @@ TValueOrError<TArray<FMVVMConstFieldVariant>, FString> GenerateFieldPathList(TAr
 }
 
 
-TValueOrError<FParsedBindingInfo, FString> GetBindingInfoFromFieldPath(const UClass* InAccessor, TArrayView<const FMVVMConstFieldVariant> InFieldPath)
+TValueOrError<FParsedBindingInfo, FText> GetBindingInfoFromFieldPath(const UClass* InAccessor, TArrayView<const FMVVMConstFieldVariant> InFieldPath)
 {
 	if (InAccessor == nullptr)
 	{
-		return MakeError(TEXT("The Accessor is null."));
+		return MakeError(LOCTEXT("AccessorClassNull", "The accessor class is null."));
 	}
 	if (InFieldPath.Num() == 0)
 	{
-		return MakeError(TEXT("The FieldPath is empty."));
+		return MakeError(LOCTEXT("FieldPathEmpty", "The field path is empty."));
 	}
 
 	FParsedBindingInfo Result;
 	FParsedBindingInfo::FFieldVariantArray LocalFullPath;
 	bool bLookForNotifyFieldId = false;
-	UStruct* CurrentContainerType = InFieldPath[0].GetOwner();
+	const UStruct* CurrentContainerType = InFieldPath[0].GetOwner();
 	if (!InAccessor->IsChildOf(CurrentContainerType))
 	{
-		return MakeError(TEXT("FieldPath that doesn't start with the accessor is not supported."));
+		return MakeError(LOCTEXT("FieldPathDoesntStartWithAccessor", "Field path that doesn't start with the accessor is not supported."));
 	}
 
 	for (int32 Index = 0; Index < InFieldPath.Num(); ++Index)
@@ -268,19 +285,24 @@ TValueOrError<FParsedBindingInfo, FString> GetBindingInfoFromFieldPath(const UCl
 
 		if (Field.IsEmpty())
 		{
-			return MakeError(FString::Printf(TEXT("The field '%d' does not exist."), Index));
+			return MakeError(FText::Format(LOCTEXT("FieldIsEmpty", "The field at index '{0}' is empty."), Index));
 		}
 
 		if (CurrentContainerType == nullptr)
 		{
-			return MakeError(FString::Printf(TEXT("The field '%s' does not exist."), *Field.GetName().ToString()));
+			return MakeError(FText::Format(LOCTEXT("FieldDoesNotExist", "The field '{0}' does not exist."), 
+				FText::FromName(Field.GetName()))
+			);
 		}
 
 		const bool bIsChild = Field.GetOwner()->IsChildOf(CurrentContainerType);
 		const bool bIsDownCast = CurrentContainerType->IsChildOf(Field.GetOwner());
 		if (!(bIsChild || bIsDownCast))
 		{
-			return MakeError(FString::Printf(TEXT("The field '%s' does not exist in the struct '%s'."), *Field.GetName().ToString(), *CurrentContainerType->GetName()));
+			return MakeError(FText::Format(LOCTEXT("FieldDoesNotExistInStruct", "The field '{0}' does not exist in the struct '{1}'."), 
+				FText::FromName(Field.GetName()), 
+				FText::FromString(CurrentContainerType->GetName()))
+			);
 		}
 
 		const FProperty* FieldProperty = nullptr;
@@ -295,7 +317,9 @@ TValueOrError<FParsedBindingInfo, FString> GetBindingInfoFromFieldPath(const UCl
 
 		if (FieldProperty == nullptr)
 		{
-			return MakeError(FString::Printf(TEXT("The field '%s' has an invalid property."), *Field.GetName().ToString()));
+			return MakeError(FText::Format(LOCTEXT("FieldHasInvalidProperty", "The field '{0}' has an invalid property."), 
+				FText::FromName(Field.GetName()))
+			);
 		}
 
 
@@ -335,7 +359,7 @@ TValueOrError<FParsedBindingInfo, FString> GetBindingInfoFromFieldPath(const UCl
 
 		if (!bLastField)
 		{
-			TValueOrError<UStruct*, FString> FoundContainer = Private::FindContainer(FieldProperty);
+			TValueOrError<UStruct*, FText> FoundContainer = Private::FindContainer(FieldProperty);
 			if (FoundContainer.HasError())
 			{
 				return MakeError(FoundContainer.StealError());
@@ -346,7 +370,7 @@ TValueOrError<FParsedBindingInfo, FString> GetBindingInfoFromFieldPath(const UCl
 
 	if (Result.ViewModelIndex > 0)
 	{
-		TValueOrError<TArray<FMVVMConstFieldVariant>, FString> FinalPath = GenerateFieldPathList(Result.NotifyFieldInterfacePath, true);
+		TValueOrError<TArray<FMVVMConstFieldVariant>, FText> FinalPath = GenerateFieldPathList(Result.NotifyFieldInterfacePath, true);
 		if (FinalPath.HasError())
 		{
 			return MakeError(FinalPath.StealError());
@@ -387,6 +411,31 @@ FString ToString(TArrayView<const FMVVMConstFieldVariant> Fields)
 	return Builder.ToString();
 }
 
+FText ToText(TArrayView<const FMVVMFieldVariant> Fields)
+{
+	TArray<FText> FieldNames;
+	FieldNames.Reserve(Fields.Num());
+
+	for (const FMVVMFieldVariant& Field : Fields)
+	{
+		FieldNames.Add(FText::FromName(Field.GetName()));
+	}
+
+	return FText::Join(LOCTEXT("PropertyPathDelim", "."), FieldNames);
+}
+
+FText ToText(TArrayView<const FMVVMConstFieldVariant> Fields)
+{
+	TArray<FText> FieldNames;
+	FieldNames.Reserve(Fields.Num());
+
+	for (const FMVVMConstFieldVariant& Field : Fields)
+	{
+		FieldNames.Add(FText::FromName(Field.GetName()));
+	}
+
+	return FText::Join(LOCTEXT("PropertyPathDelim", "."), FieldNames);
+}
 
 TValueOrError<UObject*, void> EvaluateObjectProperty(const FFieldContext& InSource)
 {
@@ -430,3 +479,5 @@ TValueOrError<UObject*, void> EvaluateObjectProperty(const FFieldContext& InSour
 }
 
 } // namespace
+
+#undef LOCTEXT_NAMESPACE
\ No newline at end of file
diff --git a/Engine/Plugins/Runtime/ModelViewViewModel/Source/ModelViewViewModel/Private/MVVMSubsystem.cpp b/Engine/Plugins/Runtime/ModelViewViewModel/Source/ModelViewViewModel/Private/MVVMSubsystem.cpp
--- a/Engine/Plugins/Runtime/ModelViewViewModel/Source/ModelViewViewModel/Private/MVVMSubsystem.cpp
+++ b/Engine/Plugins/Runtime/ModelViewViewModel/Source/ModelViewViewModel/Private/MVVMSubsystem.cpp
@@ -10,6 +10,7 @@
 
 #include UE_INLINE_GENERATED_CPP_BY_NAME(MVVMSubsystem)
 
+#define LOCTEXT_NAMESPACE "MVVMSubsystem"
 
 void UMVVMSubsystem::Initialize(FSubsystemCollectionBase& Collection)
 {
@@ -255,11 +256,11 @@ FMVVMAvailableBinding UMVVMSubsystem::GetAvailableBinding(const UClass* Class, F
 }
 
 
-TValueOrError<bool, FString> UMVVMSubsystem::IsBindingValid(FBindingArgs Args) const
+TValueOrError<bool, FText> UMVVMSubsystem::IsBindingValid(FBindingArgs Args) const
 {
 	if (UE::MVVM::IsForwardBinding(Args.Mode))
 	{
-		TValueOrError<bool, FString> ForwardDirection = IsBindingValid(Args.ForwardArgs);
+		TValueOrError<bool, FText> ForwardDirection = IsBindingValid(Args.ForwardArgs);
 		if (ForwardDirection.HasError())
 		{
 			return ForwardDirection;
@@ -268,7 +269,7 @@ TValueOrError<bool, FString> UMVVMSubsystem::IsBindingValid(FBindingArgs Args) c
 
 	if (UE::MVVM::IsBackwardBinding(Args.Mode))
 	{
-		TValueOrError<bool, FString> BackwardDirection = IsBindingValid(Args.BackwardArgs);
+		TValueOrError<bool, FText> BackwardDirection = IsBindingValid(Args.BackwardArgs);
 		if (BackwardDirection.HasError())
 		{
 			return BackwardDirection;
@@ -278,15 +279,15 @@ TValueOrError<bool, FString> UMVVMSubsystem::IsBindingValid(FBindingArgs Args) c
 	return MakeValue(true);
 }
 
-TValueOrError<bool, FString> UMVVMSubsystem::IsBindingValid(FDirectionalBindingArgs Args) const
+TValueOrError<bool, FText> UMVVMSubsystem::IsBindingValid(FDirectionalBindingArgs Args) const
 {
 	return IsBindingValid(Args.ToConst());
 }
 
-TValueOrError<bool, FString> UMVVMSubsystem::IsBindingValid(FConstDirectionalBindingArgs Args) const
+TValueOrError<bool, FText> UMVVMSubsystem::IsBindingValid(FConstDirectionalBindingArgs Args) const
 {
 	// Test Source
-	TValueOrError<const FProperty*, FString> SourceResult = UE::MVVM::BindingHelper::TryGetPropertyTypeForSourceBinding(Args.SourceBinding);
+	TValueOrError<const FProperty*, FText> SourceResult = UE::MVVM::BindingHelper::TryGetPropertyTypeForSourceBinding(Args.SourceBinding);
 	if (SourceResult.HasError())
 	{
 		return MakeError(SourceResult.StealError());
@@ -295,11 +296,11 @@ TValueOrError<bool, FString> UMVVMSubsystem::IsBindingValid(FConstDirectionalBin
 	const FProperty* SourceProperty = SourceResult.StealValue();
 	if (SourceProperty == nullptr)
 	{
-		return MakeError(TEXT("There is no value to bind at the source."));
+		return MakeError(LOCTEXT("NoValueToBindAtSource", "There is no value to bind at the source."));
 	}
 
 	// Test Destination
-	TValueOrError<const FProperty*, FString> DestinationResult = UE::MVVM::BindingHelper::TryGetPropertyTypeForDestinationBinding(Args.DestinationBinding);
+	TValueOrError<const FProperty*, FText> DestinationResult = UE::MVVM::BindingHelper::TryGetPropertyTypeForDestinationBinding(Args.DestinationBinding);
 	if (DestinationResult.HasError())
 	{
 		return MakeError(DestinationResult.StealError());
@@ -308,13 +309,13 @@ TValueOrError<bool, FString> UMVVMSubsystem::IsBindingValid(FConstDirectionalBin
 	const FProperty* DestinationProperty = DestinationResult.StealValue();
 	if (DestinationProperty == nullptr)
 	{
-		return MakeError(TEXT("There is no value to bind at the destination."));
+		return MakeError(LOCTEXT("NoValueToBindAtDestination", "There is no value to bind at the destination."));
 	}
 
 	// Test the conversion function
 	if (Args.ConversionFunction)
 	{
-		TValueOrError<const FProperty*, FString> ReturnResult = UE::MVVM::BindingHelper::TryGetReturnTypeForConversionFunction(Args.ConversionFunction);
+		TValueOrError<const FProperty*, FText> ReturnResult = UE::MVVM::BindingHelper::TryGetReturnTypeForConversionFunction(Args.ConversionFunction);
 		if (ReturnResult.HasError())
 		{
 			return MakeError(ReturnResult.StealError());
@@ -324,20 +325,20 @@ TValueOrError<bool, FString> UMVVMSubsystem::IsBindingValid(FConstDirectionalBin
 		const bool bIsComplexConversionFunction = UE::MVVM::BindingHelper::IsValidForComplexRuntimeConversion(Args.ConversionFunction);
 		if (!bIsSimpleConversionFunction && !bIsComplexConversionFunction)
 		{
-			return MakeError(FString::Printf(TEXT("The conversion function '%s' is not valid as a conversion function"), *Args.ConversionFunction->GetName()));
+			return MakeError(FText::Format(LOCTEXT("ConversionFunctionNotValid", "The conversion function '{0}' is not valid as a conversion function"), FText::FromString(Args.ConversionFunction->GetName())));
 		}
 
 		// The simple compiled version look like Setter(Conversion(Getter())).
 		if (bIsSimpleConversionFunction)
 		{
-			TValueOrError<TArray<const FProperty*>, FString> ArgumentsResult = UE::MVVM::BindingHelper::TryGetArgumentsForConversionFunction(Args.ConversionFunction);
+			TValueOrError<TArray<const FProperty*>, FText> ArgumentsResult = UE::MVVM::BindingHelper::TryGetArgumentsForConversionFunction(Args.ConversionFunction);
 			if (ArgumentsResult.HasError())
 			{
 				return MakeError(ArgumentsResult.StealError());
 			}
 			if (ArgumentsResult.GetValue().Num() != 1)
 			{
-				return MakeError(TEXT("Internal error. The number of arguments should be 1."));
+				return MakeError(LOCTEXT("NumberOfArgumentsWrong", "Internal error. The number of arguments should be 1."));
 			}
 
 			bool bAnyCompatible = false;
@@ -353,9 +354,18 @@ TValueOrError<bool, FString> UMVVMSubsystem::IsBindingValid(FConstDirectionalBin
 
 			if (!bAnyCompatible)
 			{
-				FString ArgumentsString = FString::JoinBy(ArgumentProperties, TEXT(", "), [](const FProperty* Property) { return Property->GetCPPType(); });
+				TArray<FText> ArgumentTypes;
+				ArgumentTypes.Reserve(ArgumentProperties.Num());
+				for (const FProperty* ArgProperty : ArgumentProperties)
+				{
+					ArgumentTypes.Add(FText::FromString(ArgProperty->GetCPPType()));
+				}
 
-				return MakeError(FString::Printf(TEXT("The source property '%s' (%s) does not match any of the argument types of the conversion function (%s)."), *SourceProperty->GetName(), *SourceProperty->GetCPPType(), *ArgumentsString));
+				return MakeError(FText::Format(LOCTEXT("SourcePropertyDoesNotMatchArguments", "The source property '{0}' ({1}) does not match any of the argument types of the conversion function ({1})."), 
+					FText::FromString(SourceProperty->GetName()), 
+					FText::FromString(SourceProperty->GetCPPType()),
+					FText::Join(LOCTEXT("CommaDelim", ", "), ArgumentTypes))
+				);
 			}
 		}
 		else
@@ -365,15 +375,25 @@ TValueOrError<bool, FString> UMVVMSubsystem::IsBindingValid(FConstDirectionalBin
 			const FProperty* ReturnProperty = ReturnResult.GetValue();
 			if (!UE::MVVM::BindingHelper::ArePropertiesCompatible(ReturnProperty, DestinationProperty))
 			{
-				return MakeError(FString::Printf(TEXT("The destination property '%s' (%s) does not match the return type of the conversion function (%s)."), *DestinationProperty->GetName(), *DestinationProperty->GetCPPType(), *ReturnProperty->GetCPPType()));
+				return MakeError(FText::Format(LOCTEXT("DestinationPropertyDoesNotMatchReturn", "The destination property '{0}' ({1}) does not match the return type of the conversion function ({2})."), 
+					FText::FromString(DestinationProperty->GetName()), 
+					FText::FromString(DestinationProperty->GetCPPType()),
+					FText::FromString(ReturnProperty->GetCPPType()))
+				);
 			}
 		}
 	}
 	else if (!UE::MVVM::BindingHelper::ArePropertiesCompatible(SourceProperty, DestinationProperty))
 	{
-		return MakeError(FString::Printf(TEXT("The source property '%s' (%s) does not match the type of the destination property '%s' (%s). A conversion function is required."), *SourceProperty->GetName(), *SourceProperty->GetCPPType(), *DestinationProperty->GetName(), *DestinationProperty->GetCPPType()));
+		return MakeError(FText::Format(LOCTEXT("SourcePropertyDoesNotMatchDestination", "The source property '{0}' ({1}) does not match the type of the destination property '{2}' ({3}). A conversion function is required."), 
+			FText::FromString(SourceProperty->GetName()),
+			FText::FromString(SourceProperty->GetCPPType()), 
+			FText::FromString(DestinationProperty->GetName()), 
+			FText::FromString(DestinationProperty->GetCPPType()))
+		);
 	}
 
 	return MakeValue(true);
 }
 
+#undef LOCTEXT_NAMESPACE
diff --git a/Engine/Plugins/Runtime/ModelViewViewModel/Source/ModelViewViewModel/Public/Bindings/MVVMBindingHelper.h b/Engine/Plugins/Runtime/ModelViewViewModel/Source/ModelViewViewModel/Public/Bindings/MVVMBindingHelper.h
--- a/Engine/Plugins/Runtime/ModelViewViewModel/Source/ModelViewViewModel/Public/Bindings/MVVMBindingHelper.h
+++ b/Engine/Plugins/Runtime/ModelViewViewModel/Source/ModelViewViewModel/Public/Bindings/MVVMBindingHelper.h
@@ -79,18 +79,18 @@ namespace UE::MVVM::BindingHelper
 	 */
 	UE_NODISCARD MODELVIEWVIEWMODEL_API FMVVMFieldVariant FindFieldByName(const UStruct* Container, FMVVMBindingName BindingName);
 
-	UE_NODISCARD MODELVIEWVIEWMODEL_API TValueOrError<const FProperty*, FString> TryGetPropertyTypeForSourceBinding(const FMVVMConstFieldVariant& InField);
-	UE_NODISCARD MODELVIEWVIEWMODEL_API TValueOrError<const FProperty*, FString> TryGetPropertyTypeForDestinationBinding(const FMVVMConstFieldVariant& InField);
-	UE_NODISCARD MODELVIEWVIEWMODEL_API TValueOrError<const FProperty*, FString> TryGetPropertyTypeForSourceBinding(const FProperty* InProperty);
-	UE_NODISCARD MODELVIEWVIEWMODEL_API TValueOrError<const FProperty*, FString> TryGetPropertyTypeForDestinationBinding(const FProperty* InProperty);
-	UE_NODISCARD MODELVIEWVIEWMODEL_API TValueOrError<const FProperty*, FString> TryGetPropertyTypeForSourceBinding(const UFunction* InFunction);
-	UE_NODISCARD MODELVIEWVIEWMODEL_API TValueOrError<const FProperty*, FString> TryGetPropertyTypeForDestinationBinding(const UFunction* InFunction);
+	UE_NODISCARD MODELVIEWVIEWMODEL_API TValueOrError<const FProperty*, FText> TryGetPropertyTypeForSourceBinding(const FMVVMConstFieldVariant& InField);
+	UE_NODISCARD MODELVIEWVIEWMODEL_API TValueOrError<const FProperty*, FText> TryGetPropertyTypeForDestinationBinding(const FMVVMConstFieldVariant& InField);
+	UE_NODISCARD MODELVIEWVIEWMODEL_API TValueOrError<const FProperty*, FText> TryGetPropertyTypeForSourceBinding(const FProperty* InProperty);
+	UE_NODISCARD MODELVIEWVIEWMODEL_API TValueOrError<const FProperty*, FText> TryGetPropertyTypeForDestinationBinding(const FProperty* InProperty);
+	UE_NODISCARD MODELVIEWVIEWMODEL_API TValueOrError<const FProperty*, FText> TryGetPropertyTypeForSourceBinding(const UFunction* InFunction);
+	UE_NODISCARD MODELVIEWVIEWMODEL_API TValueOrError<const FProperty*, FText> TryGetPropertyTypeForDestinationBinding(const UFunction* InFunction);
 
 	/** Returns the return property for a conversion function, or an error if the function is not valid as a conversion function. */
-	UE_NODISCARD MODELVIEWVIEWMODEL_API TValueOrError<const FProperty*, FString> TryGetReturnTypeForConversionFunction(const UFunction* InFunction);
+	UE_NODISCARD MODELVIEWVIEWMODEL_API TValueOrError<const FProperty*, FText> TryGetReturnTypeForConversionFunction(const UFunction* InFunction);
 
 	/** Returns the argument properties for a conversion function, or an error if the function is not valid as a conversion function. */
-	UE_NODISCARD MODELVIEWVIEWMODEL_API TValueOrError<TArray<const FProperty*>, FString> TryGetArgumentsForConversionFunction(const UFunction* InFunction);
+	UE_NODISCARD MODELVIEWVIEWMODEL_API TValueOrError<TArray<const FProperty*>, FText> TryGetArgumentsForConversionFunction(const UFunction* InFunction);
 
 	/** Are type the same or could be converted at runtime. */
 	UE_NODISCARD MODELVIEWVIEWMODEL_API bool ArePropertiesCompatible(const FProperty* Source, const FProperty* Destination);
diff --git a/Engine/Plugins/Runtime/ModelViewViewModel/Source/ModelViewViewModel/Public/Bindings/MVVMFieldPathHelper.h b/Engine/Plugins/Runtime/ModelViewViewModel/Source/ModelViewViewModel/Public/Bindings/MVVMFieldPathHelper.h
--- a/Engine/Plugins/Runtime/ModelViewViewModel/Source/ModelViewViewModel/Public/Bindings/MVVMFieldPathHelper.h
+++ b/Engine/Plugins/Runtime/ModelViewViewModel/Source/ModelViewViewModel/Public/Bindings/MVVMFieldPathHelper.h
@@ -45,17 +45,19 @@ namespace UE::MVVM::FieldPathHelper
 	};
 
 	/** */
-	UE_NODISCARD MODELVIEWVIEWMODEL_API TValueOrError<TArray<FMVVMConstFieldVariant>, FString> GenerateFieldPathList(const UClass* From, FStringView FieldPath, bool bForSourceBinding);
+	UE_NODISCARD MODELVIEWVIEWMODEL_API TValueOrError<TArray<FMVVMConstFieldVariant>, FText> GenerateFieldPathList(const UClass* From, FStringView FieldPath, bool bForSourceBinding);
 
 	/** */
-	UE_NODISCARD MODELVIEWVIEWMODEL_API TValueOrError<TArray<FMVVMConstFieldVariant>, FString> GenerateFieldPathList(TArrayView<const FMVVMConstFieldVariant> Fields, bool bForSourceBinding);
+	UE_NODISCARD MODELVIEWVIEWMODEL_API TValueOrError<TArray<FMVVMConstFieldVariant>, FText> GenerateFieldPathList(TArrayView<const FMVVMConstFieldVariant> Fields, bool bForSourceBinding);
 
 	/** */
-	UE_NODISCARD MODELVIEWVIEWMODEL_API TValueOrError<FParsedBindingInfo, FString> GetBindingInfoFromFieldPath(const UClass* Accessor, TArrayView<const FMVVMConstFieldVariant> Fields);
+	UE_NODISCARD MODELVIEWVIEWMODEL_API TValueOrError<FParsedBindingInfo, FText> GetBindingInfoFromFieldPath(const UClass* Accessor, TArrayView<const FMVVMConstFieldVariant> Fields);
 
 	/** */
 	UE_NODISCARD MODELVIEWVIEWMODEL_API FString ToString(TArrayView<const FMVVMFieldVariant> Fields);
 	UE_NODISCARD MODELVIEWVIEWMODEL_API FString ToString(TArrayView<const FMVVMConstFieldVariant> Fields);
+	UE_NODISCARD MODELVIEWVIEWMODEL_API FText ToText(TArrayView<const FMVVMFieldVariant> Fields);
+	UE_NODISCARD MODELVIEWVIEWMODEL_API FText ToText(TArrayView<const FMVVMConstFieldVariant> Fields);
 
 	/** */
 	UE_NODISCARD MODELVIEWVIEWMODEL_API TValueOrError<UObject*, void> EvaluateObjectProperty(const FFieldContext& InSource);
diff --git a/Engine/Plugins/Runtime/ModelViewViewModel/Source/ModelViewViewModel/Public/MVVMSubsystem.h b/Engine/Plugins/Runtime/ModelViewViewModel/Source/ModelViewViewModel/Public/MVVMSubsystem.h
--- a/Engine/Plugins/Runtime/ModelViewViewModel/Source/ModelViewViewModel/Public/MVVMSubsystem.h
+++ b/Engine/Plugins/Runtime/ModelViewViewModel/Source/ModelViewViewModel/Public/MVVMSubsystem.h
@@ -89,9 +89,9 @@ public:
 		FDirectionalBindingArgs BackwardArgs;
 	};
 
-	UE_NODISCARD TValueOrError<bool, FString> IsBindingValid(FConstDirectionalBindingArgs Args) const;
-	UE_NODISCARD TValueOrError<bool, FString> IsBindingValid(FDirectionalBindingArgs Args) const;
-	UE_NODISCARD TValueOrError<bool, FString> IsBindingValid(FBindingArgs Args) const;
+	UE_NODISCARD TValueOrError<bool, FText> IsBindingValid(FConstDirectionalBindingArgs Args) const;
+	UE_NODISCARD TValueOrError<bool, FText> IsBindingValid(FDirectionalBindingArgs Args) const;
+	UE_NODISCARD TValueOrError<bool, FText> IsBindingValid(FBindingArgs Args) const;
 
 private:
 	UPROPERTY(Transient)
diff --git a/Engine/Plugins/Runtime/ModelViewViewModel/Source/ModelViewViewModelBlueprint/Private/Bindings/MVVMCompiledBindingLibraryCompiler.cpp b/Engine/Plugins/Runtime/ModelViewViewModel/Source/ModelViewViewModelBlueprint/Private/Bindings/MVVMCompiledBindingLibraryCompiler.cpp
--- a/Engine/Plugins/Runtime/ModelViewViewModel/Source/ModelViewViewModelBlueprint/Private/Bindings/MVVMCompiledBindingLibraryCompiler.cpp
+++ b/Engine/Plugins/Runtime/ModelViewViewModel/Source/ModelViewViewModelBlueprint/Private/Bindings/MVVMCompiledBindingLibraryCompiler.cpp
@@ -8,6 +8,7 @@
 #include "Types/MVVMObjectVariant.h"
 #include <limits>
 
+#define LOCTEXT_NAMESPACE "CompiledBindingLibraryCompiler"
 
 namespace UE::MVVM::Private
 {
@@ -135,19 +136,19 @@ FCompiledBindingLibraryCompiler::FCompiledBindingLibraryCompiler()
 
 }
 
-TValueOrError<FCompiledBindingLibraryCompiler::FFieldIdHandle, FString> FCompiledBindingLibraryCompiler::AddFieldId(TSubclassOf<UObject> SourceClass, FName FieldId)
+TValueOrError<FCompiledBindingLibraryCompiler::FFieldIdHandle, FText> FCompiledBindingLibraryCompiler::AddFieldId(TSubclassOf<UObject> SourceClass, FName FieldId)
 {
 	Impl->bCompiled = false;
 
 	if (FieldId.IsNone())
 	{
-		return MakeError(TEXT("The Field/Property is not defined."));
+		return MakeError(LOCTEXT("FieldNotDefined", "The Field/Property is not defined."));
 	}
 
 	if (!SourceClass->ImplementsInterface(UNotifyFieldValueChanged::StaticClass()))
 	{
-		return MakeError(FString::Printf(TEXT("'%s' doesn't implement the NotifyFieldValueChanged interface.")
-			, *SourceClass->GetName()));
+		return MakeError(FText::Format(LOCTEXT("ClassDoesNotImplementInterface", "'{0}' doesn't implement the NotifyFieldValueChanged interface.")
+			, SourceClass->GetDisplayNameText()));
 	}
 
 	const TScriptInterface<INotifyFieldValueChanged> ScriptObject = SourceClass->GetDefaultObject();
@@ -156,9 +157,9 @@ TValueOrError<FCompiledBindingLibraryCompiler::FFieldIdHandle, FString> FCompile
 		UE::FieldNotification::FFieldId FoundFieldId = ScriptObject->GetFieldNotificationDescriptor().GetField(SourceClass, FieldId);
 		if (!FoundFieldId.IsValid())
 		{
-			return MakeError(FString::Printf(TEXT("The FieldNotify '%s' is not support by '%s'.")
-				, *FieldId.ToString()
-				, *SourceClass->GetName()));
+			return MakeError(FText::Format(LOCTEXT("FieldNotifyNotSupported", "The FieldNotify '{0}' is not supported by '{0}'.")
+				, FText::FromName(FieldId)
+				, SourceClass->GetDisplayNameText()));
 		}
 
 		int32 FoundFieldIdIndex = Impl->FieldIds.IndexOfByPredicate([FoundFieldId, SourceClass](const Private::FRawFieldId& Other)
@@ -177,11 +178,11 @@ TValueOrError<FCompiledBindingLibraryCompiler::FFieldIdHandle, FString> FCompile
 		return MakeValue(Impl->FieldIds[FoundFieldIdIndex].IdHandle);
 	}
 
-	return MakeError(TEXT("Unexpected case with AddFieldId."));
+	return MakeError(LOCTEXT("UnexpectedCaseInAddFieldId", "Unexpected case with AddFieldId."));
 }
 
 
-TValueOrError<FCompiledBindingLibraryCompiler::FFieldPathHandle, FString> FCompiledBindingLibraryCompiler::AddFieldPath(TArrayView<const FMVVMConstFieldVariant> InFieldPath, bool bInRead)
+TValueOrError<FCompiledBindingLibraryCompiler::FFieldPathHandle, FText> FCompiledBindingLibraryCompiler::AddFieldPath(TArrayView<const FMVVMConstFieldVariant> InFieldPath, bool bInRead)
 {
 	Impl->bCompiled = false;
 
@@ -189,23 +190,23 @@ TValueOrError<FCompiledBindingLibraryCompiler::FFieldPathHandle, FString> FCompi
 }
 
 
-TValueOrError<FCompiledBindingLibraryCompiler::FFieldPathHandle, FString> FCompiledBindingLibraryCompiler::AddFieldPathImpl(TArrayView<const FMVVMConstFieldVariant> InFieldPath, bool bInRead)
+TValueOrError<FCompiledBindingLibraryCompiler::FFieldPathHandle, FText> FCompiledBindingLibraryCompiler::AddFieldPathImpl(TArrayView<const FMVVMConstFieldVariant> InFieldPath, bool bInRead)
 {
 	Impl->bCompiled = false;
 
-	auto ValidateContainer = [](const FProperty* Property, bool bShouldBeInsideContainer, bool bIsObjectOrScriptStruct) -> FString
+	auto ValidateContainer = [](const FProperty* Property, bool bShouldBeInsideContainer, bool bIsObjectOrScriptStruct) -> FText
 	{	
 		const UStruct* OwnerStruct = Property->GetOwnerStruct();
 		if (OwnerStruct == nullptr)
 		{
-			return FString::Printf(TEXT("The field %s has an invalid owner struct."), *Property->GetName());
+			return FText::Format(LOCTEXT("FieldHasInvalidOwner", "The field {0} has an invalid owner struct."), Property->GetDisplayNameText());
 		}
 
 		if (bShouldBeInsideContainer)
 		{
 			if (!Cast<UScriptStruct>(OwnerStruct) && !Cast<UClass>(OwnerStruct))
 			{
-				return FString::Printf(TEXT("The field %s doesn't have a valid owner for that path."), *Property->GetName());
+				return FText::Format(LOCTEXT("FieldDoesNotHaveValidOwnerForPath", "The field {0} doesn't have a valid owner for that path."), Property->GetDisplayNameText());
 			}
 		}
 
@@ -213,20 +214,20 @@ TValueOrError<FCompiledBindingLibraryCompiler::FFieldPathHandle, FString> FCompi
 		{
 			if (const FObjectPropertyBase* ObjectProperty = CastField<const FObjectPropertyBase>(Property))
 			{
-				return FString();
+				return FText::GetEmpty();
 			}
 			else if (const FStructProperty* StructProperty = CastField<const FStructProperty>(Property))
 			{
 				if (StructProperty->HasGetter() || Property->HasMetaData(Private::NAME_BlueprintGetter))
 				{
-					return FString::Printf(TEXT("Property %s has getter accessor. Accessor not supported on FStructProperty since it would create a temporary structure and we would not able to return a valid container from that structure."), *StructProperty->GetName());
+					return FText::Format(LOCTEXT("GetterNotSupported", "Property {0} has getter accessor. Accessor not supported on FStructProperty since it would create a temporary structure and we would not able to return a valid container from that structure."), StructProperty->GetDisplayNameText());
 				}
-				return FString();
+				return FText::GetEmpty();
 			}
-			return FString::Printf(TEXT("Field can only be object properties or struct properties. %s is a %s"), *Property->GetName(), *Property->GetClass()->GetName());
+			return FText::Format(LOCTEXT("FieldCanOnlyBeObjectOrStruct", "Field can only be object properties or struct properties. {0} is a {1}"), Property->GetDisplayNameText(), Property->GetClass()->GetDisplayNameText());
 		}
 
-		return FString();
+		return FText::GetEmpty();
 	};
 
 	TArray<int32> RawFieldIndexes;
@@ -243,15 +244,15 @@ TValueOrError<FCompiledBindingLibraryCompiler::FFieldPathHandle, FString> FCompi
 			{
 				if (!BindingHelper::IsValidForDestinationBinding(FieldVariant.GetProperty()))
 				{
-					return MakeError(FString::Printf(TEXT("Property '%s' is not writable at runtime."), *FieldVariant.GetProperty()->GetName()));
+					return MakeError(FText::Format(LOCTEXT("PropertyNotWritableAtRuntime", "Property '{0}' is not writable at runtime."), FieldVariant.GetProperty()->GetDisplayNameText()));
 				}
 			}
 			else if (!BindingHelper::IsValidForSourceBinding(FieldVariant.GetProperty()))
 			{
-				return MakeError(FString::Printf(TEXT("Property '%s' is not readable at runtime."), *FieldVariant.GetProperty()->GetName()));
+				return MakeError(FText::Format(LOCTEXT("PropertyNotReadableAtRuntime", "Property '{0}' is not readable at runtime."), FieldVariant.GetProperty()->GetDisplayNameText()));
 			}
 
-			FString ValidatedStr = ValidateContainer(FieldVariant.GetProperty(), true, !bIsLast);
+			FText ValidatedStr = ValidateContainer(FieldVariant.GetProperty(), true, !bIsLast);
 			if (!ValidatedStr.IsEmpty())
 			{
 				return MakeError(ValidatedStr);
@@ -265,12 +266,12 @@ TValueOrError<FCompiledBindingLibraryCompiler::FFieldPathHandle, FString> FCompi
 			{
 				if (!BindingHelper::IsValidForDestinationBinding(FieldVariant.GetFunction()))
 				{
-					return MakeError(FString::Printf(TEXT("Function '%s' is not writable at runtime."), *FieldVariant.GetFunction()->GetName()));
+					return MakeError(FText::Format(LOCTEXT("FunctionNotWritableAtRuntime", "Function '{0}' is not writable at runtime."), FieldVariant.GetFunction()->GetDisplayNameText()));
 				}
 			}
 			else if (!BindingHelper::IsValidForSourceBinding(FieldVariant.GetFunction()))
 			{
-				return MakeError(FString::Printf(TEXT("Function '%s' is not readable at runtime."), *FieldVariant.GetFunction()->GetName()));
+				return MakeError(FText::Format(LOCTEXT("FunctionNotReadableAtRuntime", "Function '{0}' is not readable at runtime."), FieldVariant.GetFunction()->GetDisplayNameText()));
 			}
 
 			if (bIsLast && !bInRead)
@@ -288,7 +289,7 @@ TValueOrError<FCompiledBindingLibraryCompiler::FFieldPathHandle, FString> FCompi
 		}
 		else
 		{
-			return MakeError(TEXT("There is an invalid field in the field path."));
+			return MakeError(LOCTEXT("InvalidFieldInPath", "There is an invalid field in the field path."));
 		}
 	}
 
@@ -314,7 +315,7 @@ TValueOrError<FCompiledBindingLibraryCompiler::FFieldPathHandle, FString> FCompi
 }
 
 
-TValueOrError<FCompiledBindingLibraryCompiler::FFieldPathHandle, FString> FCompiledBindingLibraryCompiler::AddObjectFieldPath(TArrayView<const UE::MVVM::FMVVMConstFieldVariant> FieldPath, UClass* ExpectedType, bool bInRead)
+TValueOrError<FCompiledBindingLibraryCompiler::FFieldPathHandle, FText> FCompiledBindingLibraryCompiler::AddObjectFieldPath(TArrayView<const UE::MVVM::FMVVMConstFieldVariant> FieldPath, UClass* ExpectedType, bool bInRead)
 {
 	Impl->bCompiled = false;
 
@@ -322,7 +323,7 @@ TValueOrError<FCompiledBindingLibraryCompiler::FFieldPathHandle, FString> FCompi
 
 	if (FieldPath.Num() == 0)
 	{
-		return MakeError(FString::Printf(TEXT("The field does not return a '%s'."), *ExpectedType->GetName()));
+		return MakeError(FText::Format(LOCTEXT("FieldDoesNotReturnType", "The field does not return a '{0}'."), ExpectedType->GetDisplayNameText()));
 	}
 
 	const FObjectPropertyBase* ObjectPropertyBase = nullptr;
@@ -337,42 +338,42 @@ TValueOrError<FCompiledBindingLibraryCompiler::FFieldPathHandle, FString> FCompi
 
 	if (ObjectPropertyBase == nullptr)
 	{
-		return MakeError(FString::Printf(TEXT("The field does not return a '%s'."), *ExpectedType->GetName()));
+		return MakeError(FText::Format(LOCTEXT("FieldDoesNotReturnType", "The field does not return a '{0}'."), ExpectedType->GetDisplayNameText()));
 	}
 	if (ObjectPropertyBase->PropertyClass == nullptr || !ExpectedType->IsChildOf(ObjectPropertyBase->PropertyClass))
 	{
-		return MakeError(FString::Printf(TEXT("The field does not return a '%s'."), *ExpectedType->GetName()));
+		return MakeError(FText::Format(LOCTEXT("FieldDoesNotReturnType", "The field does not return a '{0}'."), ExpectedType->GetDisplayNameText()));
 	}
 
 	return AddFieldPathImpl(FieldPath, bInRead);
 }
 
 
-TValueOrError<FCompiledBindingLibraryCompiler::FFieldPathHandle, FString> FCompiledBindingLibraryCompiler::AddConversionFunctionFieldPath(TSubclassOf<UObject> SourceClass, const UFunction* Function)
+TValueOrError<FCompiledBindingLibraryCompiler::FFieldPathHandle, FText> FCompiledBindingLibraryCompiler::AddConversionFunctionFieldPath(TSubclassOf<UObject> SourceClass, const UFunction* Function)
 {
 	Impl->bCompiled = false;
 
 	if (SourceClass == nullptr)
 	{
-		return MakeError(TEXT("The source class is invalid."));
+		return MakeError(LOCTEXT("SourceClassInvalid", "The source class is invalid."));
 	}
 	if (Function == nullptr)
 	{
-		return MakeError(TEXT("The function path is empty."));
+		return MakeError(LOCTEXT("FunctionPathEmpty", "The function path is empty."));
 	}
 
 	const bool bIsSimpleFunction = BindingHelper::IsValidForSimpleRuntimeConversion(Function);
 	const bool bIsComplexFunction = BindingHelper::IsValidForComplexRuntimeConversion(Function);
 	if (!bIsSimpleFunction && !bIsComplexFunction)
 	{
-		return MakeError(FString::Printf(TEXT("Function %s cannot be used as a runtime conversion function."), *Function->GetName()));
+		return MakeError(FText::Format(LOCTEXT("FunctionCannotBeUsedAsConversionFunction", "Function {0} cannot be used as a runtime conversion function."), Function->GetDisplayNameText()));
 	}
 
 	if (!Function->HasAllFunctionFlags(FUNC_Static))
 	{
 		if (!SourceClass->IsChildOf(Function->GetOuterUClass()))
 		{
-			return MakeError(FString::Printf(TEXT("Function %s is going to be executed with an invalid self."), *Function->GetName()));
+			return MakeError(FText::Format(LOCTEXT("FunctionHasInvalidSelf", "Function {0} is going to be executed with an invalid self."), Function->GetDisplayNameText()));
 		}
 	}
 
@@ -398,19 +399,19 @@ TValueOrError<FCompiledBindingLibraryCompiler::FFieldPathHandle, FString> FCompi
 }
 
 
-TValueOrError<FCompiledBindingLibraryCompiler::FBindingHandle, FString> FCompiledBindingLibraryCompiler::AddBinding(FFieldPathHandle InSourceHandle, FFieldPathHandle InDestinationHandle)
+TValueOrError<FCompiledBindingLibraryCompiler::FBindingHandle, FText> FCompiledBindingLibraryCompiler::AddBinding(FFieldPathHandle InSourceHandle, FFieldPathHandle InDestinationHandle)
 {
 	return AddBinding(InSourceHandle, InDestinationHandle, FFieldPathHandle());
 }
 
 
-TValueOrError<FCompiledBindingLibraryCompiler::FBindingHandle, FString> FCompiledBindingLibraryCompiler::AddBinding(FFieldPathHandle InSourceHandle, FFieldPathHandle InDestinationHandle, FFieldPathHandle InConversionFunctionHandle)
+TValueOrError<FCompiledBindingLibraryCompiler::FBindingHandle, FText> FCompiledBindingLibraryCompiler::AddBinding(FFieldPathHandle InSourceHandle, FFieldPathHandle InDestinationHandle, FFieldPathHandle InConversionFunctionHandle)
 {
 	return AddBinding(MakeArrayView(&InSourceHandle, 1), InDestinationHandle, InConversionFunctionHandle);
 }
 
 
-TValueOrError<FCompiledBindingLibraryCompiler::FBindingHandle, FString> FCompiledBindingLibraryCompiler::AddBinding(TArrayView<const FFieldPathHandle> InSourceHandles, FFieldPathHandle InDestinationHandle, FFieldPathHandle InConversionFunctionHandle)
+TValueOrError<FCompiledBindingLibraryCompiler::FBindingHandle, FText> FCompiledBindingLibraryCompiler::AddBinding(TArrayView<const FFieldPathHandle> InSourceHandles, FFieldPathHandle InDestinationHandle, FFieldPathHandle InConversionFunctionHandle)
 {
 	Impl->bCompiled = false;
 
@@ -426,23 +427,23 @@ TValueOrError<FCompiledBindingLibraryCompiler::FBindingHandle, FString> FCompile
 			});
 		if (FoundSourceFieldPath == INDEX_NONE)
 		{
-			return MakeError(TEXT("The source handle is invalid."));
+			return MakeError(LOCTEXT("SourceHandleInvalid", "The source handle is invalid."));
 		}
 
 		Private::FRawFieldPath& SourceRawFieldPath = Impl->FieldPaths[FoundSourceFieldPath];
 		if (!SourceRawFieldPath.bIsReadable)
 		{
-			return MakeError(TEXT("The source handle was not constructed as a readable path."));
+			return MakeError(LOCTEXT("SourceHandleNotReadable", "The source handle was not constructed as a readable path."));
 		}
 		if (SourceRawFieldPath.RawFieldIndexes.Num() == 0)
 		{
-			return MakeError(TEXT("The source handle was not registered correctly."));
+			return MakeError(LOCTEXT("SourceHandleNotRegistered", "The source handle was not registered correctly."));
 		}
 
 		Private::FRawField& RawField = Impl->Fields[SourceRawFieldPath.RawFieldIndexes.Last()];
 		if (RawField.Field.IsEmpty())
 		{
-			return MakeError(TEXT("The source handle was not registered correctly."));
+			return MakeError(LOCTEXT("SourceHandleNotRegistered", "The source handle was not registered correctly."));
 		}
 
 		DirectionBindingArgs.SourceBinding = RawField.Field;
@@ -455,23 +456,23 @@ TValueOrError<FCompiledBindingLibraryCompiler::FBindingHandle, FString> FCompile
 			});
 		if (FoundDestinationFieldPath == INDEX_NONE)
 		{
-			return MakeError(TEXT("The destination handle is invalid."));
+			return MakeError(LOCTEXT("DestinationHandleInvalid", "The destination handle is invalid."));
 		}
 
 		Private::FRawFieldPath& DestinationRawFieldPath = Impl->FieldPaths[FoundDestinationFieldPath];
 		if (!DestinationRawFieldPath.bIsWritable)
 		{
-			return MakeError(TEXT("The destination handle was not constructed as a writable path."));
+			return MakeError(LOCTEXT("DestinationHandleNotWritable", "The destination handle was not constructed as a writable path."));
 		}
 		if (DestinationRawFieldPath.RawFieldIndexes.Num() == 0)
 		{
-			return MakeError(TEXT("The destination handle was not registered correctly."));
+			return MakeError(LOCTEXT("DestinationHandleNotRegistered", "The destination handle was not registered correctly."));
 		}
 
 		Private::FRawField& RawField = Impl->Fields[DestinationRawFieldPath.RawFieldIndexes.Last()];
 		if (RawField.Field.IsEmpty())
 		{
-			return MakeError(TEXT("The destination handle was not registered correctly."));
+			return MakeError(LOCTEXT("DestinationHandleNotRegistered", "The destination handle was not registered correctly."));
 		}
 
 		DirectionBindingArgs.DestinationBinding = RawField.Field;
@@ -485,25 +486,25 @@ TValueOrError<FCompiledBindingLibraryCompiler::FBindingHandle, FString> FCompile
 			});
 		if (FoundFunctionFieldPath == INDEX_NONE)
 		{
-			return MakeError(TEXT("The function handle is invalid."));
+			return MakeError(LOCTEXT("FunctionHandleInvalid", "The function handle is invalid."));
 		}
 
 		Private::FRawFieldPath& ConversionFunctinRawFieldPath = Impl->FieldPaths[FoundFunctionFieldPath];
 		if (ConversionFunctinRawFieldPath.RawFieldIndexes.Num() == 0)
 		{
-			return MakeError(TEXT("The function handle was not registered as a function."));
+			return MakeError(LOCTEXT("FunctionHandleNotRegistered", "The function handle was not registered as a function."));
 		}
 
 		Private::FRawField& RawField = Impl->Fields[ConversionFunctinRawFieldPath.RawFieldIndexes.Last()];
 		if (!RawField.Field.IsFunction())
 		{
-			return MakeError(TEXT("The function handle was not registered as a function."));
+			return MakeError(LOCTEXT("FunctionHandleNotRegistered", "The function handle was not registered as a function."));
 		}
 
 		DirectionBindingArgs.ConversionFunction = RawField.Field.GetFunction();
 	}
 
-	TValueOrError<bool, FString> IsValidBinding = GEngine->GetEngineSubsystem<UMVVMSubsystem>()->IsBindingValid(DirectionBindingArgs);
+	TValueOrError<bool, FText> IsValidBinding = GEngine->GetEngineSubsystem<UMVVMSubsystem>()->IsBindingValid(DirectionBindingArgs);
 	if (IsValidBinding.HasError())
 	{
 		return MakeError(IsValidBinding.StealError());
@@ -532,7 +533,7 @@ TValueOrError<FCompiledBindingLibraryCompiler::FBindingHandle, FString> FCompile
 }
 
 
-TValueOrError<FCompiledBindingLibraryCompiler::FCompileResult, FString> FCompiledBindingLibraryCompiler::Compile()
+TValueOrError<FCompiledBindingLibraryCompiler::FCompileResult, FText> FCompiledBindingLibraryCompiler::Compile()
 {
 	Impl->bCompiled = false;
 
@@ -647,26 +648,26 @@ TValueOrError<FCompiledBindingLibraryCompiler::FCompileResult, FString> FCompile
 
 		if (PropertyNames.Num() > std::numeric_limits<FMVVMVCompiledBinding::IndexType>::max())
 		{
-			return MakeError(FString::Printf(TEXT("There are too many properties binded to struct '%s'"), *StructCompiledFields.Key->GetName()));
+			return MakeError(FText::Format(LOCTEXT("TooManyPropertiesBound", "There are too many properties bound to struct '{0}'"), StructCompiledFields.Key->GetDisplayNameText()));
 		}
 		CompiledFields.NumberOfProperties = static_cast<int16>(PropertyNames.Num());
 
 		if (FunctionNames.Num() > std::numeric_limits<FMVVMVCompiledBinding::IndexType>::max())
 		{
-			return MakeError(FString::Printf(TEXT("There are too many functions binded to struct '%s'"), *StructCompiledFields.Key->GetName()));
+			return MakeError(FText::Format(LOCTEXT("TooManyFunctionsBound", "There are too many functions bound to struct '{0}'"), StructCompiledFields.Key->GetDisplayNameText()));
 		}
 		CompiledFields.NumberOfFunctions = static_cast<int16>(FunctionNames.Num());
 
 		if (FieldIdNames.Num() > std::numeric_limits<FMVVMVCompiledBinding::IndexType>::max())
 		{
-			return MakeError(FString::Printf(TEXT("There are too many field ids binded to struct '%s'"), *StructCompiledFields.Key->GetName()));
+			return MakeError(FText::Format(LOCTEXT("TooManyFieldIdsBound", "There are too many field IDs bound to struct '{0}'"), StructCompiledFields.Key->GetDisplayNameText()));
 		}
 		CompiledFields.NumberOfFieldIds = static_cast<int16>(FieldIdNames.Num());
 
 		int32 LibraryStartIndex = Result.Library.CompiledFieldNames.Num();
 		if (LibraryStartIndex > std::numeric_limits<FMVVMVCompiledBinding::IndexType>::max())
 		{
-			return MakeError(FString::Printf(TEXT("There are too many properties and functions binded for the library")));
+			return MakeError(LOCTEXT("TooManyPropertiesAndFunctionsBound", "There are too many properties and functions bound in the library."));
 		}
 		CompiledFields.LibraryStartIndex = static_cast<int16>(LibraryStartIndex);
 			
@@ -678,7 +679,7 @@ TValueOrError<FCompiledBindingLibraryCompiler::FCompileResult, FString> FCompile
 		FieldIdNames.Reset();
 		if (Result.Library.CompiledFieldNames.Num() > std::numeric_limits<FMVVMVCompiledBinding::IndexType>::max())
 		{
-			return MakeError(FString::Printf(TEXT("There are too many properties binded for the library")));
+			return MakeError(LOCTEXT("TooManyPropertiesBoundInLibrary", "There are too many properties bound in the library."));
 		}
 
 		Result.Library.CompiledFields.Add(CompiledFields);
@@ -772,3 +773,5 @@ TValueOrError<FCompiledBindingLibraryCompiler::FCompileResult, FString> FCompile
 }
 
 } //namespace
+
+#undef LOCTEXT_NAMESPACE
diff --git a/Engine/Plugins/Runtime/ModelViewViewModel/Source/ModelViewViewModelBlueprint/Private/Bindings/MVVMConversionFunctionHelper.cpp b/Engine/Plugins/Runtime/ModelViewViewModel/Source/ModelViewViewModelBlueprint/Private/Bindings/MVVMConversionFunctionHelper.cpp
--- a/Engine/Plugins/Runtime/ModelViewViewModel/Source/ModelViewViewModelBlueprint/Private/Bindings/MVVMConversionFunctionHelper.cpp
+++ b/Engine/Plugins/Runtime/ModelViewViewModel/Source/ModelViewViewModelBlueprint/Private/Bindings/MVVMConversionFunctionHelper.cpp
@@ -113,14 +113,26 @@ void MarkAsConversionFunction(const UK2Node_CallFunction* FunctionNode, const FM
 
 namespace Private
 {
-	FMVVMBlueprintPropertyPath GetPropertyPathForPin(const UWidgetBlueprint* WidgetBlueprint, const UEdGraphPin* StartPin)
+	FMVVMBlueprintPropertyPath GetPropertyPathForPin(const UWidgetBlueprint* WidgetBlueprint, const UEdGraphPin* StartPin, bool bSkipResolve)
 	{
+		if (StartPin->Direction != EGPD_Input || StartPin->PinName == UEdGraphSchema_K2::PN_Self)
+		{
+			return FMVVMBlueprintPropertyPath();
+		}
+		
 		UMVVMWidgetBlueprintExtension_View* ExtensionView = UMVVMWidgetBlueprintExtension_View::GetExtension<UMVVMWidgetBlueprintExtension_View>(WidgetBlueprint);
 		UMVVMBlueprintView* BlueprintView = ExtensionView->GetBlueprintView();
 
 		TArray<FMemberReference> PathParts;
 		TDeque<UEdGraphNode*> NodesToSearch;
-		NodesToSearch.PushLast(StartPin->LinkedTo[0]->GetOwningNode());
+
+		if (StartPin->LinkedTo.Num() > 0)
+		{
+			if (UEdGraphNode* Node = StartPin->LinkedTo[0]->GetOwningNode())
+			{
+				NodesToSearch.PushLast(Node);
+			}
+		}
 
 		while (NodesToSearch.Num() > 0)
 		{
@@ -129,7 +141,7 @@ namespace Private
 
 			if (UK2Node_VariableGet* Getter = Cast<UK2Node_VariableGet>(Node))
 			{
-				PathParts.Add(Getter->VariableReference);
+				PathParts.Insert(Getter->VariableReference, 0);
 			}
 
 			for (UEdGraphPin* Pin : Node->Pins)
@@ -146,25 +158,38 @@ namespace Private
 			return FMVVMBlueprintPropertyPath();
 		}
 
-		Algo::Reverse(PathParts);
-
 		FMVVMBlueprintPropertyPath ResultPath;
 
 		FMemberReference RootRef = PathParts[0];
 		PathParts.RemoveAt(0);
 
-		if (const FObjectProperty* Property = CastField<FObjectProperty>(RootRef.ResolveMember<FProperty>(WidgetBlueprint->SkeletonGeneratedClass)))
+		if (bSkipResolve)
 		{
-			if (Property->PropertyClass->ImplementsInterface(UNotifyFieldValueChanged::StaticClass()))
+			// if the generated class hasn't yet been generated we can blindly forge ahead and try to figure out if it's a widget or a viewmodel
+			if (const FMVVMBlueprintViewModelContext* ViewModel = BlueprintView->FindViewModel(RootRef.GetMemberName()))
 			{
-				if (const FMVVMBlueprintViewModelContext* ViewModel = BlueprintView->FindViewModel(Property->GetFName()))
-				{
-					ResultPath.SetViewModelId(ViewModel->GetViewModelId());
-				}
+				ResultPath.SetViewModelId(ViewModel->GetViewModelId());
 			}
-			else if (Property->PropertyClass->IsChildOf<UWidget>() || Property->PropertyClass->IsChildOf<UWidgetBlueprint>())
+			else
 			{
-				ResultPath.SetWidgetName(Property->GetFName());
+				ResultPath.SetWidgetName(RootRef.GetMemberName());
+			}
+		}
+		else
+		{
+			if (const FObjectProperty* Property = CastField<FObjectProperty>(RootRef.ResolveMember<FProperty>(WidgetBlueprint->SkeletonGeneratedClass)))
+			{
+				if (Property->PropertyClass->ImplementsInterface(UNotifyFieldValueChanged::StaticClass()))
+				{
+					if (const FMVVMBlueprintViewModelContext* ViewModel = BlueprintView->FindViewModel(Property->GetFName()))
+					{
+						ResultPath.SetViewModelId(ViewModel->GetViewModelId());
+					}
+				}
+				else if (Property->PropertyClass->IsChildOf<UWidget>() || Property->PropertyClass->IsChildOf<UWidgetBlueprint>())
+				{
+					ResultPath.SetWidgetName(Property->GetFName());
+				}
 			}
 		}
 
@@ -175,11 +200,11 @@ namespace Private
 
 		for (const FMemberReference& MemberReference : PathParts)
 		{
-			if (UFunction* Function = MemberReference.ResolveMember<UFunction>(WidgetBlueprint->SkeletonGeneratedClass))
+			if (UFunction* Function = MemberReference.ResolveMember<UFunction>(WidgetBlueprint->SkeletonGeneratedClass)) 
 			{
 				ResultPath.AppendBasePropertyPath(UE::MVVM::FMVVMConstFieldVariant(Function));
 			}
-			else if (FProperty* Property = MemberReference.ResolveMember<FProperty>(WidgetBlueprint->SkeletonGeneratedClass))
+			else if (const FProperty* Property = MemberReference.ResolveMember<FProperty>(WidgetBlueprint->SkeletonGeneratedClass))
 			{
 				ResultPath.AppendBasePropertyPath(UE::MVVM::FMVVMConstFieldVariant(Property));
 			}
@@ -189,7 +214,7 @@ namespace Private
 	}
 }
 
-FMVVMBlueprintPropertyPath GetPropertyPathForArgument(const UWidgetBlueprint* WidgetBlueprint, const UK2Node_CallFunction* FunctionNode, FName ArgumentName)
+FMVVMBlueprintPropertyPath GetPropertyPathForArgument(const UWidgetBlueprint* WidgetBlueprint, const UK2Node_CallFunction* FunctionNode, FName ArgumentName, bool bSkipResolve)
 {
 	const UEdGraphPin* ArgumentPin = FunctionNode->FindPin(ArgumentName, EGPD_Input);
 	if (ArgumentPin == nullptr || ArgumentPin->LinkedTo.Num() == 0)
@@ -197,24 +222,37 @@ FMVVMBlueprintPropertyPath GetPropertyPathForArgument(const UWidgetBlueprint* Wi
 		return FMVVMBlueprintPropertyPath();
 	}
 
-	return Private::GetPropertyPathForPin(WidgetBlueprint, ArgumentPin);
+	return Private::GetPropertyPathForPin(WidgetBlueprint, ArgumentPin, false);
 }
 
-TArray<FMVVMBlueprintPropertyPath> GetAllArgumentPathsForConversionNode(const UWidgetBlueprint* WidgetBlueprint, const UK2Node_CallFunction* FunctionNode)
+TMap<FName, FMVVMBlueprintPropertyPath> GetAllArgumentPropertyPaths(const UWidgetBlueprint* WidgetBlueprint, const UK2Node_CallFunction* FunctionNode, bool bSkipResolve)
 {
 	check(FunctionNode);
 
-	TArray<FMVVMBlueprintPropertyPath> Paths;
+	TMap<FName, FMVVMBlueprintPropertyPath> Paths;
 	for (const UEdGraphPin* Pin : FunctionNode->GetAllPins())
 	{
-		FMVVMBlueprintPropertyPath Path = Private::GetPropertyPathForPin(WidgetBlueprint, Pin);
+		FMVVMBlueprintPropertyPath Path = Private::GetPropertyPathForPin(WidgetBlueprint, Pin, bSkipResolve);
 		if (!Path.IsEmpty())
 		{
-			Paths.Add(Path);
+			Paths.Add(Pin->PinName, Path);
 		}
 	}
 
 	return Paths;
 }
 
+TMap<FName, FMVVMBlueprintPropertyPath> GetAllArgumentPropertyPaths(const UWidgetBlueprint* WidgetBlueprint, const FMVVMBlueprintViewBinding& Binding, bool bSourceToDestination, bool bSkipResolve)
+{
+	if (UEdGraph* ConversionFunctionGraph = ConversionFunctionHelper::GetGraph(WidgetBlueprint, Binding, bSourceToDestination))
+	{
+		if (UK2Node_CallFunction* ConversionNode = ConversionFunctionHelper::GetFunctionNode(ConversionFunctionGraph))
+		{
+			return GetAllArgumentPropertyPaths(WidgetBlueprint, ConversionNode, bSkipResolve);
+		}
+	}
+
+	return TMap<FName, FMVVMBlueprintPropertyPath>();
+}
+
 } //namespace UE::MVVM
\ No newline at end of file
diff --git a/Engine/Plugins/Runtime/ModelViewViewModel/Source/ModelViewViewModelBlueprint/Private/MVVMBlueprintView.cpp b/Engine/Plugins/Runtime/ModelViewViewModel/Source/ModelViewViewModelBlueprint/Private/MVVMBlueprintView.cpp
--- a/Engine/Plugins/Runtime/ModelViewViewModel/Source/ModelViewViewModelBlueprint/Private/MVVMBlueprintView.cpp
+++ b/Engine/Plugins/Runtime/ModelViewViewModel/Source/ModelViewViewModelBlueprint/Private/MVVMBlueprintView.cpp
@@ -152,6 +152,7 @@ FMVVMBlueprintViewBinding& UMVVMBlueprintView::AddBinding(const UWidget* Widget,
 FMVVMBlueprintViewBinding& UMVVMBlueprintView::AddDefaultBinding()
 {
 	FMVVMBlueprintViewBinding& NewBinding = Bindings.AddDefaulted_GetRef();
+	NewBinding.BindingId = FGuid::NewGuid();
 
 	OnBindingsUpdated.Broadcast();
 	return NewBinding;
@@ -189,6 +190,14 @@ void UMVVMBlueprintView::PostLoad()
 }
 
 #if WITH_EDITOR
+
+/*
+void UMVVMBlueprintView::PostEditUndo() 
+{
+	OnBindingsUpdated.Broadcast();
+	OnViewModelsUpdated.Broadcast();
+}*/
+
 void UMVVMBlueprintView::PostEditChangeProperty(FPropertyChangedEvent& PropertyChangedEvent)
 {
 	Super::PostEditChangeProperty(PropertyChangedEvent);
diff --git a/Engine/Plugins/Runtime/ModelViewViewModel/Source/ModelViewViewModelBlueprint/Private/MVVMBlueprintViewBinding.cpp b/Engine/Plugins/Runtime/ModelViewViewModel/Source/ModelViewViewModelBlueprint/Private/MVVMBlueprintViewBinding.cpp
--- a/Engine/Plugins/Runtime/ModelViewViewModel/Source/ModelViewViewModelBlueprint/Private/MVVMBlueprintViewBinding.cpp
+++ b/Engine/Plugins/Runtime/ModelViewViewModel/Source/ModelViewViewModelBlueprint/Private/MVVMBlueprintViewBinding.cpp
@@ -14,7 +14,7 @@
 
 FName FMVVMBlueprintViewBinding::GetFName() const
 {
-	return *BindingId.ToString();
+	return *BindingId.ToString(EGuidFormats::DigitsWithHyphensLower);
 }
 
 namespace UE::MVVM::Private
@@ -68,23 +68,28 @@ namespace UE::MVVM::Private
 
 				for (const UEdGraphPin* Pin : CallFunctionNode->Pins)
 				{
-					if (Pin->Direction == EGPD_Input)
+					if (Pin->PinName == UEdGraphSchema_K2::PN_Self || 
+						Pin->Direction != EGPD_Input)
 					{
-						if (!bFirst)
-						{
-							NameBuilder << TEXT(", ");
-						}
+						continue;
+					}
 
-						FMVVMBlueprintPropertyPath ArgumentPath = ConversionFunctionHelper::GetPropertyPathForArgument(WidgetBlueprint, CallFunctionNode, Pin->GetFName());
-						if (!ArgumentPath.IsEmpty())
-						{
-							AppendViewModelPathString(BlueprintView, ArgumentPath, NameBuilder);
-						}
-						else
-						{
-							NameBuilder << Pin->GetDefaultAsString();
-						}
+					if (!bFirst)
+					{
+						NameBuilder << TEXT(", ");
+					}
+
+					FMVVMBlueprintPropertyPath ArgumentPath = ConversionFunctionHelper::GetPropertyPathForArgument(WidgetBlueprint, CallFunctionNode, Pin->GetFName(), true);
+					if (!ArgumentPath.IsEmpty())
+					{
+						AppendViewModelPathString(BlueprintView, ArgumentPath, NameBuilder);
+					}
+					else
+					{
+						NameBuilder << Pin->GetDefaultAsString();
 					}
+
+					bFirst = false;
 				}
 
 				NameBuilder << TEXT(")");
@@ -175,7 +180,7 @@ namespace UE::MVVM::Private
 							NameBuilder << TEXT(", ");
 						}
 
-						FMVVMBlueprintPropertyPath ArgumentPath = ConversionFunctionHelper::GetPropertyPathForArgument(WidgetBlueprint, CallFunctionNode, Pin->GetFName());
+						FMVVMBlueprintPropertyPath ArgumentPath = ConversionFunctionHelper::GetPropertyPathForArgument(WidgetBlueprint, CallFunctionNode, Pin->GetFName(), true);
 						if (!ArgumentPath.IsEmpty())
 						{
 							AppendWidgetPathString(WidgetBlueprint, ArgumentPath, NameBuilder);
@@ -185,6 +190,8 @@ namespace UE::MVVM::Private
 							NameBuilder << Pin->GetDefaultAsString();
 						}
 					}
+
+					bFirst = false;
 				}
 
 				NameBuilder << TEXT(")");
diff --git a/Engine/Plugins/Runtime/ModelViewViewModel/Source/ModelViewViewModelBlueprint/Private/MVVMViewBlueprintCompiler.cpp b/Engine/Plugins/Runtime/ModelViewViewModel/Source/ModelViewViewModelBlueprint/Private/MVVMViewBlueprintCompiler.cpp
--- a/Engine/Plugins/Runtime/ModelViewViewModel/Source/ModelViewViewModelBlueprint/Private/MVVMViewBlueprintCompiler.cpp
+++ b/Engine/Plugins/Runtime/ModelViewViewModel/Source/ModelViewViewModelBlueprint/Private/MVVMViewBlueprintCompiler.cpp
@@ -1,30 +1,93 @@
 // Copyright Epic Games, Inc. All Rights Reserved.
 
 #include "MVVMViewBlueprintCompiler.h"
-#include "MVVMWidgetBlueprintExtension_View.h"
 #include "Blueprint/WidgetTree.h"
 #include "Bindings/MVVMBindingHelper.h"
 #include "Bindings/MVVMCompiledBindingLibraryCompiler.h"
+#include "Bindings/MVVMConversionFunctionHelper.h"
 #include "Bindings/MVVMFieldPathHelper.h"
 #include "EdGraphSchema_K2.h"
 #include "MVVMBlueprintView.h"
 #include "MVVMFunctionGraphHelper.h"
-#include "MVVMSubsystem.h"
-#include "MVVMViewModelBase.h"
 #include "MVVMViewModelBase.h"
+#include "MVVMWidgetBlueprintExtension_View.h"
 #include "WidgetBlueprintCompiler.h"
 #include "View/MVVMViewClass.h"
 
+#define LOCTEXT_NAMESPACE "MVVMViewBlueprintCompiler"
+
 namespace UE::MVVM::Private
 {
+FString PropertyPathToString(const UMVVMBlueprintView* BlueprintView, const FMVVMBlueprintPropertyPath& PropertyPath)
+{
+	if (PropertyPath.IsEmpty())
+	{
+		return FString();
+	}
+
+	TStringBuilder<512> Result;
+	if (PropertyPath.IsFromViewModel())
+	{
+		if (const FMVVMBlueprintViewModelContext* SourceViewModelContext = BlueprintView->FindViewModel(PropertyPath.GetViewModelId()))
+		{
+			Result << SourceViewModelContext->GetViewModelName();
+		}
+		else
+		{
+			Result << TEXT("<none>");
+		}
+	}
+	else if (PropertyPath.IsFromWidget())
+	{
+		Result << PropertyPath.GetWidgetName();
+	}
+	else
+	{
+		Result << TEXT("<none>");
+	}
+
+	FString BasePropertyPath = PropertyPath.GetBasePropertyPath();
+	if (BasePropertyPath.Len())
+	{
+		Result << TEXT('.');
+		Result << MoveTemp(BasePropertyPath);
+	}
+	return Result.ToString();
+}
+
+FText PropertyPathToText(const UMVVMBlueprintView* BlueprintView, const FMVVMBlueprintPropertyPath& PropertyPath)
+{
+	return FText::FromString(PropertyPathToString(BlueprintView, PropertyPath));
+}
 
-void FMVVMViewBlueprintCompiler::AddErrorForBinding(FMVVMBlueprintViewBinding& Binding, const UWidgetBlueprint* WidgetBlueprint, const FString& Message) const
+FText GetViewModelIdText(const FMVVMBlueprintPropertyPath& PropertyPath)
 {
-	const FString BindingName = Binding.GetDisplayNameString(WidgetBlueprint);
-	WidgetBlueprintCompilerContext.MessageLog.Error(*(BindingName + TEXT(": ") + Message));
-	Binding.Errors.Add(FText::FromString(BindingName + TEXT(": ") + Message));
+	return FText::FromString(PropertyPath.GetViewModelId().ToString(EGuidFormats::DigitsWithHyphensInBraces));
+}
+
+void FMVVMViewBlueprintCompiler::AddErrorForBinding(FMVVMBlueprintViewBinding& Binding, const FText& Message, FName ArgumentName) const
+{
+	const FText BindingName = FText::FromString(Binding.GetDisplayNameString(WidgetBlueprintCompilerContext.WidgetBlueprint()));
+
+	FText FormattedError;
+	if (!ArgumentName.IsNone())
+	{
+		FormattedError = FText::Format(LOCTEXT("BindingFormatWithArgument", "Binding '{0}': Argument '{1}' - {2}"), BindingName, FText::FromName(ArgumentName), Message);
+	}
+	else
+	{
+		FormattedError = FText::Format(LOCTEXT("BindingFormat", "Binding '{0}': {1}"), BindingName, Message);
+	}
+
+	WidgetBlueprintCompilerContext.MessageLog.Error(*FormattedError.ToString());
+	Binding.Errors.Add(FormattedError);
 }
 
+void FMVVMViewBlueprintCompiler::AddErrorForViewModel(const FMVVMBlueprintViewModelContext& ViewModel, const FText& Message) const
+{
+	const FText FormattedError = FText::Format(LOCTEXT("ViewModelFormat", "Viewodel '{0}': {1}"), ViewModel.GetDisplayName(), Message);
+	WidgetBlueprintCompilerContext.MessageLog.Error(*FormattedError.ToString());
+}
 
 void FMVVMViewBlueprintCompiler::AddExtension(UWidgetBlueprintGeneratedClass* Class, UMVVMViewClass* ViewExtension)
 {
@@ -83,7 +146,7 @@ void FMVVMViewBlueprintCompiler::CreateFunctions(UMVVMBlueprintView* BlueprintVi
 		{
 			if (!UE::MVVM::FunctionGraphHelper::GenerateViewModelSetter(WidgetBlueprintCompilerContext, SourceCreator.SetterGraph, SourceCreator.ViewModelContext.GetViewModelName()))
 			{
-				WidgetBlueprintCompilerContext.MessageLog.Error(TEXT("The setter function for '%s' could not be generated."), *SourceCreator.ViewModelContext.GetDisplayName().ToString());
+				AddErrorForViewModel(SourceCreator.ViewModelContext, LOCTEXT("SetterFunctionCouldNotBeGenerated", "The setter function could not be generated."));
 				continue;
 			}
 		}
@@ -127,7 +190,7 @@ void FMVVMViewBlueprintCompiler::CreateVariables(const FWidgetBlueprintCompilerC
 			}
 			if (!SourceContext.DisplayName.IsEmpty())
 			{
-				ViewModelProperty->SetMetaData(FBlueprintMetadata::MD_FunctionCategory, *SourceContext.DisplayName);
+				ViewModelProperty->SetMetaData(FBlueprintMetadata::MD_FunctionCategory, *SourceContext.DisplayName.ToString());
 			}
 			if (!SourceContext.CategoryName.IsEmpty())
 			{
@@ -160,8 +223,7 @@ void FMVVMViewBlueprintCompiler::CreateVariables(const FWidgetBlueprintCompilerC
 					}
 					else
 					{
-						WidgetBlueprintCompilerContext.MessageLog.Error(*FString::Printf(TEXT("The field for source '%s' exists but is not a property.")
-							, *SourceContext.DisplayName));
+						WidgetBlueprintCompilerContext.MessageLog.Error(*FText::Format(LOCTEXT("FieldIsNotProperty", "The field for source '{0}' exists but is not a property."), SourceContext.DisplayName).ToString());
 						bAreSourcesCreatorValid = false;
 						continue;
 					}
@@ -174,8 +236,7 @@ void FMVVMViewBlueprintCompiler::CreateVariables(const FWidgetBlueprintCompilerC
 		{
 			if (!BindingHelper::IsValidForSourceBinding(SourceContext.Field))
 			{
-				WidgetBlueprintCompilerContext.MessageLog.Error(*FString::Printf(TEXT("The field for source '%s' exist but will is not accessible at runtime.")
-					, *SourceContext.DisplayName));
+				WidgetBlueprintCompilerContext.MessageLog.Error(*FText::Format(LOCTEXT("FieldNotAccessibleAtRuntime", "The field for source '{0}' exists but is not accessible at runtime."), SourceContext.DisplayName).ToString());
 				bAreSourcesCreatorValid = false;
 				continue;
 			}
@@ -185,7 +246,7 @@ void FMVVMViewBlueprintCompiler::CreateVariables(const FWidgetBlueprintCompilerC
 			const bool bIsCompatible = ObjectProperty && SourceContext.Class->IsChildOf(ObjectProperty->PropertyClass);
 			if (!bIsCompatible)
 			{
-				WidgetBlueprintCompilerContext.MessageLog.Error(TEXT("There is already a property named '%s' that is not compatible with source of the same name."), *SourceContext.DisplayName);
+				WidgetBlueprintCompilerContext.MessageLog.Error(*FText::Format(LOCTEXT("PropertyExistsAndNotCompatible","There is already a property named '{0}' that is not compatible with the source of the same name."), SourceContext.DisplayName).ToString());
 				bAreSourceContextsValid = false;
 				continue;
 			}
@@ -198,7 +259,7 @@ void FMVVMViewBlueprintCompiler::CreateVariables(const FWidgetBlueprintCompilerC
 
 		if (SourceContext.Field.IsEmpty())
 		{
-			WidgetBlueprintCompilerContext.MessageLog.Error(TEXT("The variable for '%s' could not be created."), *SourceContext.DisplayName);
+			WidgetBlueprintCompilerContext.MessageLog.Error(*FText::Format(LOCTEXT("VariableCouldNotBeCreated", "The variable for '{0}' could not be created."), SourceContext.DisplayName).ToString());
 			bAreSourceContextsValid = false;
 			continue;
 		}
@@ -236,31 +297,29 @@ void FMVVMViewBlueprintCompiler::CreateSourceLists(const FWidgetBlueprintCompile
 	CompilerSourceContexts.Reset();
 	CompilerSourceCreatorContexts.Reset();
 
-	TSet<FGuid> ViewModelsGuid;
+	TSet<FGuid> ViewModelGuids;
 	TSet<FName> WidgetSources;
 	for (const FMVVMBlueprintViewModelContext& ViewModelContext : BlueprintView->GetViewModels())
 	{
 		if (!ViewModelContext.GetViewModelId().IsValid())
 		{
-			WidgetBlueprintCompilerContext.MessageLog.Error(TEXT("The viewmodel context has an invalid Guid"));
+			AddErrorForViewModel(ViewModelContext, LOCTEXT("ViewmodelInvalidGuid", "GUID is invalid."));
 			bAreSourcesCreatorValid = false;
 			continue;
 		}
 
-		if (ViewModelsGuid.Contains(ViewModelContext.GetViewModelId()))
+		if (ViewModelGuids.Contains(ViewModelContext.GetViewModelId()))
 		{
-			WidgetBlueprintCompilerContext.MessageLog.Error(*FString::Printf(TEXT("The viewmodel '%s' is present twice.")
-				, *ViewModelContext.GetViewModelId().ToString()));
+			AddErrorForViewModel(ViewModelContext, LOCTEXT("ViewmodelAlreadyAdded", "Identical viewmodel has already been added."));
 			bAreSourcesCreatorValid = false;
 			continue;
 		}
 
-		ViewModelsGuid.Add(ViewModelContext.GetViewModelId());
+		ViewModelGuids.Add(ViewModelContext.GetViewModelId());
 
 		if (ViewModelContext.GetViewModelClass() == nullptr || !ViewModelContext.IsValid())
 		{
-			WidgetBlueprintCompilerContext.MessageLog.Error(*FString::Printf(TEXT("The viewmodel '%s' has an invalid class.")
-				, *ViewModelContext.GetViewModelId().ToString()));
+			AddErrorForViewModel(ViewModelContext, LOCTEXT("ViewmodelInvalidClass", "Invalid class."));
 			bAreSourcesCreatorValid = false;
 			continue;
 		}
@@ -292,7 +351,7 @@ void FMVVMViewBlueprintCompiler::CreateSourceLists(const FWidgetBlueprintCompile
 		FCompilerSourceContext SourceVariable;
 		SourceVariable.Class = ViewModelContext.GetViewModelClass();
 		SourceVariable.PropertyName = ViewModelContext.GetViewModelName();
-		SourceVariable.DisplayName = ViewModelContext.GetDisplayName().ToString();
+		SourceVariable.DisplayName = ViewModelContext.GetDisplayName();
 		SourceVariable.CategoryName = TEXT("Viewmodel");
 		SourceVariable.bExposeOnSpawn = bCreateSetterFunction;
 		SourceVariable.BlueprintSetter = CompilerSourceCreatorContexts[FoundSourceCreatorContextIndex].SetterFunctionName;
@@ -301,32 +360,38 @@ void FMVVMViewBlueprintCompiler::CreateSourceLists(const FWidgetBlueprintCompile
 
 	bAreSourceContextsValid = bAreSourcesCreatorValid;
 
+	UWidgetBlueprintGeneratedClass* SkeletonClass = Context.GetSkeletonGeneratedClass();
+
 	// Only find the source first property and destination first property.
 	//The full path will be tested later. We want to build the list of property needed.
 	for (int32 Index = 0; Index < BlueprintView->GetNumBindings(); ++Index)
 	{
-		const FMVVMBlueprintViewBinding* BindingPtr = BlueprintView->GetBindingAt(Index);
+		FMVVMBlueprintViewBinding* BindingPtr = BlueprintView->GetBindingAt(Index);
 		if (BindingPtr == nullptr)
 		{
-			WidgetBlueprintCompilerContext.MessageLog.Error(*FString::Printf(TEXT("The binding number %d is invalid."), Index));
+			WidgetBlueprintCompilerContext.MessageLog.Error(*FText::Format(LOCTEXT("BindingInvalidIndex", "Internal error: Tried to fetch binding for invalid index {0}."), Index).ToString());
 			bAreSourceContextsValid = false;
 			continue;
 		}
-		const FMVVMBlueprintViewBinding& Binding = *BindingPtr;
+
+		FMVVMBlueprintViewBinding& Binding = *BindingPtr;
 		if (!Binding.bCompile)
 		{
 			continue;
 		}
 
 		FMVVMViewBlueprintCompiler* Self = this;
-		auto GenerateCompilerSourceContext = [Self, Index, &ViewModelsGuid, &WidgetSources](const UWidgetBlueprintGeneratedClass* Class, const UMVVMBlueprintView* BlueprintView, const FMVVMBlueprintPropertyPath& PropertyPath) -> bool
+		auto GenerateCompilerSourceContext = [Self, BlueprintView, Class = SkeletonClass, &Binding, &ViewModelGuids, &WidgetSources](const FMVVMBlueprintPropertyPath& PropertyPath, bool bViewModelPath, FName ArgumentName = FName()) -> bool
 		{
 			if (PropertyPath.IsFromWidget())
 			{
-				if (PropertyPath.GetWidgetName().IsNone())
+				if (bViewModelPath)
 				{
-					Self->WidgetBlueprintCompilerContext.MessageLog.Error(*FString::Printf(TEXT("The binding number %d has an invalid widget name."), Index));
-					return false;
+					Self->AddErrorForBinding(Binding,
+							FText::Format(LOCTEXT("ExpectedViewModelPath", "Expected a viewmodel path, but received a path from widget: {0}"), 
+								FText::FromName(PropertyPath.GetWidgetName())
+							)
+						);
 				}
 
 				if (PropertyPath.GetWidgetName() == Class->ClassGeneratedBy->GetFName())
@@ -343,7 +408,12 @@ void FMVVMViewBlueprintCompiler::CreateSourceLists(const FWidgetBlueprintCompile
 					UWidget** WidgetPtr = Self->WidgetNameToWidgetPointerMap.Find(PropertyPath.GetWidgetName());
 					if (WidgetPtr == nullptr || *WidgetPtr == nullptr)
 					{
-						Self->WidgetBlueprintCompilerContext.MessageLog.Error(*FString::Printf(TEXT("The binding number %d has an invalid path. The widget '%s' could not be found."), Index, *PropertyPath.GetWidgetName().ToString()));
+						Self->AddErrorForBinding(Binding,
+							FText::Format(LOCTEXT("InvalidWidgetFormat", "Could not find the targeted widget: {0}"), 
+								FText::FromName(PropertyPath.GetWidgetName())
+							),
+							ArgumentName
+						);
 						return false;
 					}
 
@@ -351,7 +421,7 @@ void FMVVMViewBlueprintCompiler::CreateSourceLists(const FWidgetBlueprintCompile
 					FCompilerSourceContext SourceVariable;
 					SourceVariable.Class = Widget->GetClass();
 					SourceVariable.PropertyName = PropertyPath.GetWidgetName();
-					SourceVariable.DisplayName = Widget->GetDisplayLabel();
+					SourceVariable.DisplayName = FText::FromString(Widget->GetDisplayLabel());
 					SourceVariable.CategoryName = TEXT("Widget");
 					Self->CompilerSourceContexts.Emplace(MoveTemp(SourceVariable));
 				}
@@ -361,30 +431,91 @@ void FMVVMViewBlueprintCompiler::CreateSourceLists(const FWidgetBlueprintCompile
 				const FMVVMBlueprintViewModelContext* SourceViewModelContext = BlueprintView->FindViewModel(PropertyPath.GetViewModelId());
 				if (SourceViewModelContext == nullptr)
 				{
-					Self->WidgetBlueprintCompilerContext.MessageLog.Error(*FString::Printf(TEXT("The binding number %d has an invalid viewmodel."), Index));
+					Self->AddErrorForBinding(Binding,
+						FText::Format(LOCTEXT("BindingViewModelNotFound", "Could not find viewmodel with GUID {0}."), GetViewModelIdText(PropertyPath)),
+						ArgumentName
+					);
 					return false;
 				}
 
-				bool bFound = ViewModelsGuid.Contains(SourceViewModelContext->GetViewModelId());
-				if (!bFound)
+				if (!bViewModelPath)
+				{
+					Self->AddErrorForBinding(Binding,
+							FText::Format(LOCTEXT("ExpectedWidgetPath", "Expected a widget path, but received a path from viewmodel: {0}"), 
+								SourceViewModelContext->GetDisplayName()
+							),
+							ArgumentName
+						);
+				}
+
+				if (!ViewModelGuids.Contains(SourceViewModelContext->GetViewModelId()))
 				{
-					Self->WidgetBlueprintCompilerContext.MessageLog.Error(*FString::Printf(TEXT("The binding number %d has a viewmodel that was not added."), Index));
+					Self->AddErrorForBinding(Binding,
+						FText::Format(LOCTEXT("BindingViewModelInvalid", "Viewmodel {0} {1} was invalid."), 
+							SourceViewModelContext->GetDisplayName(),
+							GetViewModelIdText(PropertyPath)
+						),
+						ArgumentName
+					);
 					return false;
 				}
 			}
 			else
 			{
-				Self->WidgetBlueprintCompilerContext.MessageLog.Error(*FString::Printf(TEXT("The binding number %d has a unknown path type."), Index));
+				Self->AddErrorForBinding(Binding,
+					bViewModelPath ? LOCTEXT("ViewModelPathNotSet", "A viewmodel path is required, but not set.") :
+							LOCTEXT("WidgetPathNotSet", "A widget path is required, but not set."),
+					ArgumentName
+				);
 				return false;
 			}
 			return true;
 		};
 
-		if (IsForwardBinding(Binding.BindingType) || IsBackwardBinding(Binding.BindingType))
+		const bool bIsForwardBinding = IsForwardBinding(Binding.BindingType);
+		const bool bIsBackwardBinding = IsBackwardBinding(Binding.BindingType);
+
+		if (bIsForwardBinding || bIsBackwardBinding)
 		{
-			bAreSourceContextsValid = GenerateCompilerSourceContext(Context.GetSkeletonGeneratedClass(), BlueprintView, Binding.ViewModelPath)
-				&& GenerateCompilerSourceContext(Context.GetSkeletonGeneratedClass(), BlueprintView, Binding.WidgetPath)
-				&& bAreSourceContextsValid;
+			TMap<FName, FMVVMBlueprintPropertyPath> ConversionFunctionArguments = ConversionFunctionHelper::GetAllArgumentPropertyPaths(WidgetBlueprintCompilerContext.WidgetBlueprint(), Binding, bIsForwardBinding, true);
+			if (ConversionFunctionArguments.Num() > 0)
+			{
+				if (Binding.BindingType == EMVVMBindingMode::TwoWay)
+				{
+					Self->AddErrorForBinding(Binding, LOCTEXT("TwoWayBindingsWithConversion", "Two-way bindings are not allowed to use conversion functions."));
+					bAreSourceContextsValid = false;
+					continue;
+				}
+
+				// generate sources for conversion function arguments
+				for (const TPair<FName, FMVVMBlueprintPropertyPath>& Arg : ConversionFunctionArguments)
+				{
+					if (bIsForwardBinding)
+					{
+						bAreSourceContextsValid &= GenerateCompilerSourceContext(Arg.Value, true, Arg.Key);
+					}
+					else
+					{
+						bAreSourceContextsValid &= GenerateCompilerSourceContext(Arg.Value, false, Arg.Key);
+					}
+				}
+
+				// generate destination source
+				if (bIsForwardBinding)
+				{
+					bAreSourceContextsValid &= GenerateCompilerSourceContext(Binding.WidgetPath, false);
+				}
+				else
+				{
+					bAreSourceContextsValid &= GenerateCompilerSourceContext(Binding.ViewModelPath, true);
+				}
+			}
+			else
+			{
+				// if we aren't using a conversion function, just validate the widget and viewmodel paths
+				bAreSourceContextsValid &= GenerateCompilerSourceContext(Binding.WidgetPath, false);
+				bAreSourceContextsValid &= GenerateCompilerSourceContext(Binding.ViewModelPath, true);
+			}
 		}
 	}
 }
@@ -429,7 +560,10 @@ void FMVVMViewBlueprintCompiler::CreateFunctionsDeclaration(const FWidgetBluepri
 
 			if (SourceCreator.SetterGraph == nullptr || SourceCreator.SetterGraph->GetFName() != FName(*SourceCreator.SetterFunctionName))
 			{
-				WidgetBlueprintCompilerContext.MessageLog.Error(*FString::Printf(TEXT("The setter name %s already exist and could not be autogenerated."), *SourceCreator.SetterFunctionName));
+				WidgetBlueprintCompilerContext.MessageLog.Error(*FText::Format(LOCTEXT("SetterNameAlreadyExists", "The setter name {0} already exists and could not be autogenerated."), 
+					FText::FromString(SourceCreator.SetterFunctionName)
+					).ToString()
+				);
 			}
 
 			UE::MVVM::FunctionGraphHelper::AddFunctionArgument(SourceCreator.SetterGraph, SourceCreator.ViewModelContext.GetViewModelClass(), "Viewmodel");
@@ -464,10 +598,10 @@ bool FMVVMViewBlueprintCompiler::Compile(UWidgetBlueprintGeneratedClass* Class,
 		return false;
 	}
 
-	TValueOrError<FCompiledBindingLibraryCompiler::FCompileResult, FString> CompileResult = BindingLibraryCompiler.Compile();
+	TValueOrError<FCompiledBindingLibraryCompiler::FCompileResult, FText> CompileResult = BindingLibraryCompiler.Compile();
 	if (CompileResult.HasError())
 	{
-		WidgetBlueprintCompilerContext.MessageLog.Error(*FString::Printf(TEXT("The binding compilation failed. %s"), *CompileResult.GetError()));
+		WidgetBlueprintCompilerContext.MessageLog.Error(*FText::Format(LOCTEXT("BindingCompilationFailed", "The binding compilation failed. {1}"), CompileResult.GetError()).ToString());
 		return false;
 	}
 	CompileSourceCreators(CompileResult.GetValue(), Class, BlueprintView, ViewExtension);
@@ -485,15 +619,13 @@ bool FMVVMViewBlueprintCompiler::Compile(UWidgetBlueprintGeneratedClass* Class,
 
 bool FMVVMViewBlueprintCompiler::PreCompileBindingSources(UWidgetBlueprintGeneratedClass* Class, UMVVMBlueprintView* BlueprintView)
 {
-	const UWidgetBlueprint* WidgetBlueprint = WidgetBlueprintCompilerContext.WidgetBlueprint();
-
 	const int32 NumBindings = BlueprintView->GetNumBindings();
 	for (int32 Index = 0; Index < NumBindings; ++Index)
 	{
 		FMVVMBlueprintViewBinding* BindingPtr = BlueprintView->GetBindingAt(Index);
 		if (BindingPtr == nullptr)
 		{
-			WidgetBlueprintCompilerContext.MessageLog.Error(*FString::Printf(TEXT("The binding at index '%d' is invalid."), Index));
+			WidgetBlueprintCompilerContext.MessageLog.Error(*FText::Format(LOCTEXT("InvalidBindingIndex", "Internal error. Invalid binding index given."), Index).ToString());
 			bIsBindingsValid = false;
 			continue;
 		}
@@ -504,78 +636,87 @@ bool FMVVMViewBlueprintCompiler::PreCompileBindingSources(UWidgetBlueprintGenera
 			continue;
 		}
 
-		if (IsForwardBinding(Binding.BindingType))
+		auto CreateSourceContextForPropertyPath = [this, &Binding, BlueprintView, Class, Index](const FMVVMBlueprintPropertyPath& Path, bool bForwardBinding, FName ArgumentName = FName()) -> bool
 		{
-			const TValueOrError<FBindingSourceContext, FString> CreatedBindingSourceContext = CreateBindingSourceContext(BlueprintView, Class, Binding.ViewModelPath);
+			const TValueOrError<FBindingSourceContext, FText> CreatedBindingSourceContext = CreateBindingSourceContext(BlueprintView, Class, Path);
 			if (CreatedBindingSourceContext.HasError())
 			{
-				AddErrorForBinding(Binding, WidgetBlueprint, FString::Printf(TEXT("The property path '%s' is invalid. %s")
-					, *PropertyPathToString(BlueprintView, Binding.ViewModelPath)
-					, *CreatedBindingSourceContext.GetError()));
-				bIsBindingsValid = false;
-				continue;
+				AddErrorForBinding(Binding, 
+					FText::Format(LOCTEXT("PropertyPathInvalidWithReason", "The property path '{0}' is invalid. {1}"),
+						PropertyPathToText(BlueprintView, Binding.ViewModelPath),
+						CreatedBindingSourceContext.GetError()
+					),
+					ArgumentName
+				);
+				return false;
 			}
 
 			FBindingSourceContext BindingSourceContext = CreatedBindingSourceContext.GetValue();
 			if (!IsPropertyPathValid(BindingSourceContext.PropertyPath))
 			{
-				AddErrorForBinding(Binding, WidgetBlueprint, FString::Printf(TEXT("The property path '%s' is invalid."), *PropertyPathToString(BlueprintView, Binding.ViewModelPath)));
-				bIsBindingsValid = false;
-				continue;
+				AddErrorForBinding(Binding, 
+					FText::Format(LOCTEXT("PropertyPathIsInvalid", "The property path '{0}' is invalid."), 
+						PropertyPathToText(BlueprintView, Binding.ViewModelPath)
+					),
+					ArgumentName
+				);
+				return false;
 			}
+
 			if (BindingSourceContext.SourceClass == nullptr)
 			{
-				AddErrorForBinding(Binding, WidgetBlueprint, TEXT("Internal error. The binding could not find its source."));
-				bIsBindingsValid = false;
-				continue;
+				AddErrorForBinding(Binding, LOCTEXT("BindingInvalidSourceClass", "Internal error. The binding could not find its source class."), ArgumentName);
+				return false;
 			}
+
 			if (!BindingSourceContext.bIsRootWidget && BindingSourceContext.CompilerSourceContextIndex == INDEX_NONE)
 			{
-				AddErrorForBinding(Binding, WidgetBlueprint, TEXT("Internal error. The binding could not find its source."));
-				bIsBindingsValid = false;
-				continue;
+				AddErrorForBinding(Binding, LOCTEXT("BindingInvalidSource", "Internal error. The binding could not find its source."), ArgumentName);
+				return false;
 			}
 
 			BindingSourceContext.BindingIndex = Index;
-			BindingSourceContext.bIsForwardBinding = true;
-			BindingSourceContexts.Add(BindingSourceContext);
-		}
+			BindingSourceContext.bIsForwardBinding = bForwardBinding;
 
-		if (IsBackwardBinding(Binding.BindingType))
+			this->BindingSourceContexts.Add(MoveTemp(BindingSourceContext));
+			return true;
+		};
+
+		auto CreateSourcesForConversionFunction = [this, &Binding, &CreateSourceContextForPropertyPath](bool bForwardBinding)
 		{
-			const TValueOrError<FBindingSourceContext, FString> CreatedBindingSourceContext = CreateBindingSourceContext(BlueprintView, Class, Binding.WidgetPath);
-			if (CreatedBindingSourceContext.HasError())
+			TMap<FName, FMVVMBlueprintPropertyPath> ArgumentPaths = ConversionFunctionHelper::GetAllArgumentPropertyPaths(WidgetBlueprintCompilerContext.WidgetBlueprint(), Binding, bForwardBinding, true);
+			for (const TPair<FName, FMVVMBlueprintPropertyPath>& Pair : ArgumentPaths)
 			{
-				AddErrorForBinding(Binding, WidgetBlueprint, FString::Printf(TEXT("The property path '%s' is invalid. %s")
-					, *PropertyPathToString(BlueprintView, Binding.ViewModelPath)
-					, *CreatedBindingSourceContext.GetError()));
-				bIsBindingsValid = false;
-				continue;
+				CreateSourceContextForPropertyPath(Pair.Value, bForwardBinding, Pair.Key);
 			}
+		};
 
-			FBindingSourceContext BindingSourceContext = CreatedBindingSourceContext.GetValue();
-			if (!IsPropertyPathValid(BindingSourceContext.PropertyPath))
-			{
-				AddErrorForBinding(Binding, WidgetBlueprint, FString::Printf(TEXT("The property path '%s' is invalid."), *PropertyPathToString(BlueprintView, Binding.WidgetPath)));
-				bIsBindingsValid = false;
-				continue;
-			}
-			if (BindingSourceContext.SourceClass == nullptr)
+		if (IsForwardBinding(Binding.BindingType))
+		{
+			CreateSourcesForConversionFunction(true);
+
+			if (!Binding.ViewModelPath.IsEmpty())
 			{
-				AddErrorForBinding(Binding, WidgetBlueprint, TEXT("Internal error. The binding could not find its source."));
-				bIsBindingsValid = false;
-				continue;
+				if (!CreateSourceContextForPropertyPath(Binding.ViewModelPath, true))
+				{
+					bIsBindingsValid = false;
+					continue;
+				}
 			}
-			if (!BindingSourceContext.bIsRootWidget && BindingSourceContext.CompilerSourceContextIndex == INDEX_NONE)
+		}
+
+		if (IsBackwardBinding(Binding.BindingType))
+		{
+			CreateSourcesForConversionFunction(false);
+
+			if (!Binding.WidgetPath.IsEmpty())
 			{
-				AddErrorForBinding(Binding, WidgetBlueprint, TEXT("Internal error. The binding could not find its source."));
-				bIsBindingsValid = false;
-				continue;
+				if (!CreateSourceContextForPropertyPath(Binding.WidgetPath, false))
+				{
+					bIsBindingsValid = false;
+					continue;
+				}
 			}
-
-			BindingSourceContext.BindingIndex = Index;
-			BindingSourceContext.bIsForwardBinding = false;
-			BindingSourceContexts.Add(BindingSourceContext);
 		}
 	}
 
@@ -601,9 +742,9 @@ bool FMVVMViewBlueprintCompiler::PreCompileSourceCreators(UWidgetBlueprintGenera
 
 			if (ViewModelContext.GetViewModelClass()->HasAllClassFlags(CLASS_Deprecated))
 			{
-				WidgetBlueprintCompilerContext.MessageLog.Warning(*FString::Printf(TEXT("The viewmodel type '%s' is deprecated and should not be used for '%s'. Please update it in the View Binding panel under Manage ViewModels.")
-					, *ViewModelContext.GetViewModelClass()->GetDisplayNameText().ToString()
-					, *ViewModelContext.GetDisplayName().ToString()));
+				AddErrorForViewModel(ViewModelContext, FText::Format(LOCTEXT("ViewModelTypeDeprecated", "Viewmodel type '{0}' is deprecated and should not be used. Please update it in the View Models panel."),
+					ViewModelContext.GetViewModelClass()->GetDisplayNameText()
+				));
 			}
 
 			if (ViewModelContext.CreationType == EMVVMBlueprintViewModelContextCreationType::Manual)
@@ -613,9 +754,9 @@ bool FMVVMViewBlueprintCompiler::PreCompileSourceCreators(UWidgetBlueprintGenera
 			{
 				if (ViewModelContext.GetViewModelClass()->HasAllClassFlags(CLASS_Abstract))
 				{
-					WidgetBlueprintCompilerContext.MessageLog.Error(*FString::Printf(TEXT("The viewmodel type '%s' is abstract and can't be created for '%s'. You can change it in the View Binding panel under Manage ViewModels.")
-						, *ViewModelContext.GetViewModelClass()->GetDisplayNameText().ToString()
-						, *ViewModelContext.GetDisplayName().ToString()));
+					AddErrorForViewModel(ViewModelContext, FText::Format(LOCTEXT("ViewModelTypeAbstract", "Viewmodel type '{0}' is abstract and can't be created. You can change it in the View Models panel."),
+						ViewModelContext.GetViewModelClass()->GetDisplayNameText()
+					));
 					bAreSourcesCreatorValid = false;
 					continue;
 				}
@@ -624,31 +765,26 @@ bool FMVVMViewBlueprintCompiler::PreCompileSourceCreators(UWidgetBlueprintGenera
 			{
 				if (ViewModelContext.ViewModelPropertyPath.IsEmpty())
 				{
-					WidgetBlueprintCompilerContext.MessageLog.Error(*FString::Printf(TEXT("The viewmodel '%s' has an invalid Getter. You can select a new one in the View Binding panel under Manage ViewModels.")
-						, *ViewModelContext.GetDisplayName().ToString()));
+					AddErrorForViewModel(ViewModelContext, LOCTEXT("ViewModelInvalidGetter", "Viewmodel has an invalid Getter. You can select a new one in the View Models panel."));
 					bAreSourcesCreatorValid = true;
 					continue;
 				}
 
 				// Generate a path to read the value at runtime
-				TValueOrError<TArray<FMVVMConstFieldVariant>, FString> GeneratedField = FieldPathHelper::GenerateFieldPathList(Class, ViewModelContext.ViewModelPropertyPath, true);
+				static const FText InvalidGetterFormat = LOCTEXT("ViewModelInvalidGetterWithReason", "Viewmodel has an invalid Getter. {0}");
+
+				TValueOrError<TArray<FMVVMConstFieldVariant>, FText> GeneratedField = FieldPathHelper::GenerateFieldPathList(Class, ViewModelContext.ViewModelPropertyPath, true);
 				if (GeneratedField.HasError())
 				{
-					WidgetBlueprintCompilerContext.MessageLog.Error(*FString::Printf(TEXT("The viewmodel '%s' (%s) has an invalid Getter. %s")
-						, *ViewModelContext.GetDisplayName().ToString()
-						, *ViewModelContext.GetViewModelClass()->GetDisplayNameText().ToString()
-						, *GeneratedField.GetError()));
+					AddErrorForViewModel(ViewModelContext, FText::Format(InvalidGetterFormat, GeneratedField.GetError()));
 					bAreSourcesCreatorValid = false;
 					continue;
 				}
 
-				TValueOrError<FCompiledBindingLibraryCompiler::FFieldPathHandle, FString> ReadFieldPathResult = BindingLibraryCompiler.AddObjectFieldPath(GeneratedField.GetValue(), ViewModelContext.GetViewModelClass(), true);
+				TValueOrError<FCompiledBindingLibraryCompiler::FFieldPathHandle, FText> ReadFieldPathResult = BindingLibraryCompiler.AddObjectFieldPath(GeneratedField.GetValue(), ViewModelContext.GetViewModelClass(), true);
 				if (ReadFieldPathResult.HasError())
 				{
-					WidgetBlueprintCompilerContext.MessageLog.Error(*FString::Printf(TEXT("The viewmodel '%s' (%s) has an invalid Getter. %s")
-						, *ViewModelContext.GetDisplayName().ToString()
-						, *ViewModelContext.GetViewModelClass()->GetDisplayNameText().ToString()
-						, *ReadFieldPathResult.GetError()));
+					AddErrorForViewModel(ViewModelContext, FText::Format(InvalidGetterFormat	, ReadFieldPathResult.GetError()));
 					bAreSourcesCreatorValid = false;
 					continue;
 				}
@@ -659,19 +795,15 @@ bool FMVVMViewBlueprintCompiler::PreCompileSourceCreators(UWidgetBlueprintGenera
 			{
 				if (ViewModelContext.GlobalViewModelIdentifier.IsNone())
 				{
-					WidgetBlueprintCompilerContext.MessageLog.Error(*FString::Printf(TEXT("The viewmodel '%s' (%s) doesn't have a valid Global identifier. You can specify a new one in the View Binding panel under Manage ViewModels.")
-						, *ViewModelContext.GetDisplayName().ToString()
-						, *ViewModelContext.GetViewModelClass()->GetDisplayNameText().ToString()));
+					AddErrorForViewModel(ViewModelContext, LOCTEXT("ViewmodelInvalidGlobalIdentifier", "Viewmodel doesn't have a valid Global identifier. You can specify a new one in the Viewmodels panel."));
 					bAreSourcesCreatorValid = false;
 					continue;
 				}
 			}
 			else
 			{
-				WidgetBlueprintCompilerContext.MessageLog.Error(*FString::Printf(TEXT("The viewmodel '%s' (%s) doesn't have a valid creation type. You can select one in the View Binding panel under Manage ViewModels.")
-					, *ViewModelContext.GetDisplayName().ToString()
-					, *ViewModelContext.GetViewModelClass()->GetDisplayNameText().ToString()));
-				bAreSourcesCreatorValid = true;
+				AddErrorForViewModel(ViewModelContext, LOCTEXT("ViewmodelInvalidCreationType", "Viewmodel doesn't have a valid creation type. You can select one in the Viewmodels panel."));
+				bAreSourcesCreatorValid = false;
 				continue;
 			}
 		}
@@ -709,8 +841,7 @@ bool FMVVMViewBlueprintCompiler::CompileSourceCreators(const FCompiledBindingLib
 				const FMVVMVCompiledFieldPath* CompiledFieldPath = CompileResult.FieldPaths.Find(SourceCreatorContext.ReadPropertyPath);
 				if (CompiledFieldPath == nullptr)
 				{
-					WidgetBlueprintCompilerContext.MessageLog.Error(*FString::Printf(TEXT("The viewmodel '%s' initialization binding was not generated.")
-						, *ViewModelContext.GetViewModelClass()->GetDisplayNameText().ToString()));
+					AddErrorForViewModel(ViewModelContext, LOCTEXT("ViewModelInvalidInitializationBindingNotGenerated", "The viewmodel initialization binding was not generated."));
 					bAreSourcesCreatorValid = false;
 					continue;
 				}
@@ -720,8 +851,7 @@ bool FMVVMViewBlueprintCompiler::CompileSourceCreators(const FCompiledBindingLib
 			{
 				if (ViewModelContext.GlobalViewModelIdentifier.IsNone())
 				{
-					WidgetBlueprintCompilerContext.MessageLog.Error(*FString::Printf(TEXT("The viewmodel '%s' doesn't have a valid Global identifier.")
-						, *ViewModelContext.GetViewModelClass()->GetDisplayNameText().ToString()));
+					AddErrorForViewModel(ViewModelContext, LOCTEXT("ViewModelInvalidGlobalIdentifier", "The viewmodel doesn't have a valid Global identifier. You can specify a new one in the Viewmodels panel."));
 					bAreSourcesCreatorValid = false;
 					continue;
 				}
@@ -731,8 +861,7 @@ bool FMVVMViewBlueprintCompiler::CompileSourceCreators(const FCompiledBindingLib
 				GlobalViewModelInstance.ContextName = ViewModelContext.GlobalViewModelIdentifier;
 				if (!GlobalViewModelInstance.IsValid())
 				{
-					WidgetBlueprintCompilerContext.MessageLog.Error(*FString::Printf(TEXT("The context for viewmodel '%s' could not be created.")
-						, *ViewModelContext.GetViewModelClass()->GetDisplayNameText().ToString()));
+					AddErrorForViewModel(ViewModelContext, LOCTEXT("ViewModelCouldNotBeCreated", "The context for viewmodel could not be created. You can change the viewmodel in the Viewmodels panel."));
 					bAreSourcesCreatorValid = false;
 					continue;
 				}
@@ -741,8 +870,7 @@ bool FMVVMViewBlueprintCompiler::CompileSourceCreators(const FCompiledBindingLib
 			}
 			else
 			{
-				WidgetBlueprintCompilerContext.MessageLog.Error(*FString::Printf(TEXT("The viewmodel '%s' doesn't have a valid creation type.")
-					, *ViewModelContext.GetViewModelClass()->GetDisplayNameText().ToString()));
+				AddErrorForViewModel(ViewModelContext, LOCTEXT("ViewModelWithoutValidCreationType", "The viewmodel doesn't have a valid creation type."));
 				bAreSourcesCreatorValid = false;
 				continue;
 			}
@@ -762,8 +890,6 @@ bool FMVVMViewBlueprintCompiler::PreCompileBindings(UWidgetBlueprintGeneratedCla
 		return false;
 	}
 
-	const UWidgetBlueprint* WidgetBlueprint = WidgetBlueprintCompilerContext.WidgetBlueprint();
-
 	for (const FBindingSourceContext& BindingSourceContext : BindingSourceContexts)
 	{
 		FMVVMBlueprintViewBinding* BindingPtr = BlueprintView->GetBindingAt(BindingSourceContext.BindingIndex);
@@ -771,75 +897,80 @@ bool FMVVMViewBlueprintCompiler::PreCompileBindings(UWidgetBlueprintGeneratedCla
 		FMVVMBlueprintViewBinding& Binding = *BindingPtr;
 
 		FMVVMViewBlueprintCompiler* Self = this;
-		auto AddBinding = [Self](UWidgetBlueprintGeneratedClass* Class, TArrayView<const UE::MVVM::FMVVMConstFieldVariant> GetterField, TArrayView<const UE::MVVM::FMVVMConstFieldVariant> SetterFields, const UFunction* ConversionFunction) -> TValueOrError<FCompilerBinding, FString>
+		auto AddBinding = [Self](UWidgetBlueprintGeneratedClass* Class, TArrayView<const UE::MVVM::FMVVMConstFieldVariant> GetterFields, TArrayView<const UE::MVVM::FMVVMConstFieldVariant> SetterFields, const UFunction* ConversionFunction) -> TValueOrError<FCompilerBinding, FText>
 		{
 			FCompilerBinding Result;
 
-			if (GetterField.Num() != 0)
+			if (GetterFields.Num() > 0)
 			{
+				static const FText CouldNotCreateSourceFieldPathFormat = LOCTEXT("CouldNotCreateSourceFieldPath", "Couldn't create the source field path '{0}'. {1}");
+
+
 				// Generate a path to read the value at runtime
-				TValueOrError<TArray<FMVVMConstFieldVariant>, FString> GeneratedField = FieldPathHelper::GenerateFieldPathList(GetterField, true);
+				TValueOrError<TArray<FMVVMConstFieldVariant>, FText> GeneratedField = FieldPathHelper::GenerateFieldPathList(GetterFields, true);
 				if (GeneratedField.HasError())
 				{
-					return MakeError(FString::Printf(TEXT("Couldn't create the source field path '%s'. %s")
-						, *::UE::MVVM::FieldPathHelper::ToString(GetterField)
-						, *GeneratedField.GetError()));
+					return MakeError(FText::Format(CouldNotCreateSourceFieldPathFormat
+						, ::UE::MVVM::FieldPathHelper::ToText(GetterFields)
+						, GeneratedField.GetError()));
 				}
 
-				TValueOrError<FCompiledBindingLibraryCompiler::FFieldPathHandle, FString> FieldPathResult = Self->BindingLibraryCompiler.AddFieldPath(GeneratedField.GetValue(), true);
+				TValueOrError<FCompiledBindingLibraryCompiler::FFieldPathHandle, FText> FieldPathResult = Self->BindingLibraryCompiler.AddFieldPath(GeneratedField.GetValue(), true);
 				if (FieldPathResult.HasError())
 				{
-					return MakeError(FString::Printf(TEXT("Couldn't create the source field path '%s'. %s")
-						, *::UE::MVVM::FieldPathHelper::ToString(GetterField)
-						, *FieldPathResult.GetError()));
+					return MakeError(FText::Format(CouldNotCreateSourceFieldPathFormat
+						, ::UE::MVVM::FieldPathHelper::ToText(GetterFields)
+						, FieldPathResult.GetError()));
 				}
 				Result.SourceRead = FieldPathResult.StealValue();
 			}
 
+			static const FText CouldNotCreateDestinationFieldPathFormat = LOCTEXT("CouldNotCreateDestinationFieldPath", "Couldn't create the destination field path '{0}'. {1}");
+
 			{
-				TValueOrError<TArray<FMVVMConstFieldVariant>, FString> GeneratedField = FieldPathHelper::GenerateFieldPathList(SetterFields, false);
+				TValueOrError<TArray<FMVVMConstFieldVariant>, FText> GeneratedField = FieldPathHelper::GenerateFieldPathList(SetterFields, false);
 				if (GeneratedField.HasError())
 				{
-					return MakeError(FString::Printf(TEXT("Couldn't create the destination field path '%s'. %s")
-						, *::UE::MVVM::FieldPathHelper::ToString(SetterFields)
-						, *GeneratedField.GetError()));
+					return MakeError(FText::Format(CouldNotCreateDestinationFieldPathFormat
+						, ::UE::MVVM::FieldPathHelper::ToText(SetterFields)
+						, GeneratedField.GetError()));
 				}
 
-				TValueOrError<FCompiledBindingLibraryCompiler::FFieldPathHandle, FString> FieldPathResult = Self->BindingLibraryCompiler.AddFieldPath(GeneratedField.GetValue(), false);
+				TValueOrError<FCompiledBindingLibraryCompiler::FFieldPathHandle, FText> FieldPathResult = Self->BindingLibraryCompiler.AddFieldPath(GeneratedField.GetValue(), false);
 				if (FieldPathResult.HasError())
 				{
-					return MakeError(FString::Printf(TEXT("Couldn't create the destination field path '%s'. %s")
-						, *::UE::MVVM::FieldPathHelper::ToString(SetterFields)
-						, *FieldPathResult.GetError()));
+					return MakeError(FText::Format(CouldNotCreateDestinationFieldPathFormat
+						, ::UE::MVVM::FieldPathHelper::ToText(SetterFields)
+						, FieldPathResult.GetError()));
 				}
 				Result.DestinationWrite = FieldPathResult.StealValue();
 			}
 
 			if (ConversionFunction != nullptr)
 			{
-				TValueOrError<FCompiledBindingLibraryCompiler::FFieldPathHandle, FString> FieldPathResult = Self->BindingLibraryCompiler.AddConversionFunctionFieldPath(Class, ConversionFunction);
+				TValueOrError<FCompiledBindingLibraryCompiler::FFieldPathHandle, FText> FieldPathResult = Self->BindingLibraryCompiler.AddConversionFunctionFieldPath(Class, ConversionFunction);
 				if (FieldPathResult.HasError())
 				{
-					return MakeError(FString::Printf(TEXT("Couldn't create the conversion function field path '%s'. %s")
-						, *ConversionFunction->GetPathName()
-						, *FieldPathResult.GetError()));
+					return MakeError(FText::Format(LOCTEXT("CouldNotCreateConversionFunctionFieldPath", "Couldn't create the conversion function field path '{0}'. {1}")
+						, FText::FromString(ConversionFunction->GetPathName())
+						, FieldPathResult.GetError()));
 				}
 				Result.ConversionFunction = FieldPathResult.StealValue();
 				Result.bIsConversionFunctionComplex = BindingHelper::IsValidForComplexRuntimeConversion(ConversionFunction);
 			}
 
 			// Sanity check.
-			if (GetterField.Num() == 0 && (ConversionFunction == nullptr || !BindingHelper::IsValidForComplexRuntimeConversion(ConversionFunction)))
+			if (GetterFields.Num() == 0 && (ConversionFunction == nullptr || !BindingHelper::IsValidForComplexRuntimeConversion(ConversionFunction)))
 			{
-				return MakeError(FString::Printf(TEXT("The binding should have a getter.")));
+				return MakeError(LOCTEXT("BindingShouldHaveGetter", "The binding should have a getter."));
 			}
-			if (GetterField.Num() != 0 && ConversionFunction == nullptr && BindingHelper::IsValidForComplexRuntimeConversion(ConversionFunction))
+			if (GetterFields.Num() != 0 && ConversionFunction == nullptr && BindingHelper::IsValidForComplexRuntimeConversion(ConversionFunction))
 			{
-				return MakeError(FString::Printf(TEXT("The binding has a complex conversion function and a getter.")));
+				return MakeError(LOCTEXT("BindingHasConversionFunctionAndGetter", "The binding has both a complex conversion function and a getter."));
 			}
 
 			// Generate the binding
-			TValueOrError<FCompiledBindingLibraryCompiler::FBindingHandle, FString> BindingResult = Self->BindingLibraryCompiler.AddBinding(Result.SourceRead, Result.DestinationWrite, Result.ConversionFunction);
+			TValueOrError<FCompiledBindingLibraryCompiler::FBindingHandle, FText> BindingResult = Self->BindingLibraryCompiler.AddBinding(Result.SourceRead, Result.DestinationWrite, Result.ConversionFunction);
 			if (BindingResult.HasError())
 			{
 				return MakeError(BindingResult.StealError());
@@ -849,7 +980,7 @@ bool FMVVMViewBlueprintCompiler::PreCompileBindings(UWidgetBlueprintGeneratedCla
 			return MakeValue(Result);
 		};
 
-		auto AddFieldId = [Self](UClass* SourceContextClass, bool bNotifyFieldValueChangedRequired, EMVVMBindingMode BindingMode, FName FieldToListenTo) -> TValueOrError<FCompiledBindingLibraryCompiler::FFieldIdHandle, FString>
+		auto AddFieldId = [Self](UClass* SourceContextClass, bool bNotifyFieldValueChangedRequired, EMVVMBindingMode BindingMode, FName FieldToListenTo) -> TValueOrError<FCompiledBindingLibraryCompiler::FFieldIdHandle, FText>
 		{
 			if (!IsOneTimeBinding(BindingMode) && bNotifyFieldValueChangedRequired)
 			{
@@ -858,10 +989,11 @@ bool FMVVMViewBlueprintCompiler::PreCompileBindings(UWidgetBlueprintGeneratedCla
 			return MakeValue(FCompiledBindingLibraryCompiler::FFieldIdHandle());
 		};
 
-		TValueOrError<FCompiledBindingLibraryCompiler::FFieldIdHandle, FString> AddFieldResult = AddFieldId(BindingSourceContext.SourceClass, true, Binding.BindingType, BindingSourceContext.FieldId.GetFieldName());
+		TValueOrError<FCompiledBindingLibraryCompiler::FFieldIdHandle, FText> AddFieldResult = AddFieldId(BindingSourceContext.SourceClass, true, Binding.BindingType, BindingSourceContext.FieldId.GetFieldName());
 		if (AddFieldResult.HasError())
 		{
-			AddErrorForBinding(Binding, WidgetBlueprint, *FString::Printf(TEXT("The binding could not create its source. %s"), *AddFieldResult.GetError()));
+			AddErrorForBinding(Binding, FText::Format(LOCTEXT("CouldNotCreateSource", "Could not create source. {0}"), 
+				AddFieldResult.GetError()));
 			bIsBindingsValid = false;
 			continue;
 		}
@@ -872,7 +1004,9 @@ bool FMVVMViewBlueprintCompiler::PreCompileBindings(UWidgetBlueprintGeneratedCla
 			SetterPath = CreateBindingDestinationPath(BlueprintView, Class, DestinationPath);
 			if (!IsPropertyPathValid(SetterPath))
 			{
-				AddErrorForBinding(Binding, WidgetBlueprint, FString::Printf(TEXT("The property path '%s' is invalid."), *PropertyPathToString(BlueprintView, DestinationPath)));
+				AddErrorForBinding(Binding, FText::Format(LOCTEXT("PropertyPathIsInvalid", "The property path '{0}' is invalid."), 
+					PropertyPathToText(BlueprintView, DestinationPath))
+				);
 				bIsBindingsValid = false;
 				continue;
 			}
@@ -888,15 +1022,19 @@ bool FMVVMViewBlueprintCompiler::PreCompileBindings(UWidgetBlueprintGeneratedCla
 		const UFunction* ConversionFunction = ConversionFunctionReference.ResolveMember<UFunction>(Class);
 		if (!ConversionFunctionWrapper.IsNone() && ConversionFunction == nullptr)
 		{
-			AddErrorForBinding(Binding, WidgetBlueprint, TEXT("A complex conversion function was created but could not be found."));
+			AddErrorForBinding(Binding, FText::Format(LOCTEXT("ConversionFunctionNotFound", "The conversion function '{0}' could not be found."), 
+				FText::FromName(ConversionFunctionReference.GetMemberName()))
+			);
 			bIsBindingsValid = false;
 			continue;
 		}
 
-		TValueOrError<FCompilerBinding, FString> AddBindingResult = AddBinding(Class, BindingSourceContext.PropertyPath, SetterPath, ConversionFunction);
+		TValueOrError<FCompilerBinding, FText> AddBindingResult = AddBinding(Class, BindingSourceContext.PropertyPath, SetterPath, ConversionFunction);
 		if (AddBindingResult.HasError())
 		{
-			AddErrorForBinding(Binding, WidgetBlueprint, FString::Printf(TEXT("The binding could not be created. %s"), *AddBindingResult.GetError()));
+			AddErrorForBinding(Binding, 
+				FText::Format(LOCTEXT("CouldNotCreateBinding", "Could not create binding. {0}"), AddBindingResult.GetError())
+			);
 			bIsBindingsValid = false;
 			continue;
 		}
@@ -923,7 +1061,7 @@ bool FMVVMViewBlueprintCompiler::CompileBindings(const FCompiledBindingLibraryCo
 
 	for (const FCompilerBinding& CompileBinding : CompilerBindings)
 	{
-		const FMVVMBlueprintViewBinding& ViewBinding = *BlueprintView->GetBindingAt(CompileBinding.BindingIndex);
+		FMVVMBlueprintViewBinding& ViewBinding = *BlueprintView->GetBindingAt(CompileBinding.BindingIndex);
 
 		FMVVMViewClass_CompiledBinding NewBinding;
 
@@ -933,8 +1071,9 @@ bool FMVVMViewBlueprintCompiler::CompileBindings(const FCompiledBindingLibraryCo
 		const FMVVMVCompiledFieldId* CompiledFieldId = CompileResult.FieldIds.Find(CompileBinding.FieldIdHandle);
 		if (CompiledFieldId == nullptr && CompileBinding.bFieldIdNeeded)
 		{
-			WidgetBlueprintCompilerContext.MessageLog.Error(*FString::Printf(TEXT("The field id for binding '%d' was not generated.")
-				, CompileBinding.BindingIndex));
+			AddErrorForBinding(ViewBinding,
+				FText::Format(LOCTEXT("FieldIdNotGenerated", "Could not generate field ID for property '{0}'."), FText::FromName(NewBinding.SourcePropertyName))
+			);
 			bIsBindingsValid = false;
 			continue;
 		}
@@ -942,8 +1081,7 @@ bool FMVVMViewBlueprintCompiler::CompileBindings(const FCompiledBindingLibraryCo
 		const FMVVMVCompiledBinding* CompiledBinding = CompileResult.Bindings.Find(CompileBinding.BindingHandle);
 		if (CompiledBinding == nullptr)
 		{
-			WidgetBlueprintCompilerContext.MessageLog.Error(*FString::Printf(TEXT("The binding '%d' was not generated.")
-				, CompileBinding.BindingIndex));
+			AddErrorForBinding(ViewBinding, LOCTEXT("CompiledBindingNotGenerated", "Could not generate compiled binding."));
 			bIsBindingsValid = false;
 			continue;
 		}
@@ -986,12 +1124,12 @@ const FMVVMViewBlueprintCompiler::FCompilerSourceCreatorContext* FMVVMViewBluepr
 }
 
 
-TValueOrError<FMVVMViewBlueprintCompiler::FBindingSourceContext, FString> FMVVMViewBlueprintCompiler::CreateBindingSourceContext(const UMVVMBlueprintView* BlueprintView, const UWidgetBlueprintGeneratedClass* Class, const FMVVMBlueprintPropertyPath& PropertyPath)
+TValueOrError<FMVVMViewBlueprintCompiler::FBindingSourceContext, FText> FMVVMViewBlueprintCompiler::CreateBindingSourceContext(const UMVVMBlueprintView* BlueprintView, const UWidgetBlueprintGeneratedClass* Class, const FMVVMBlueprintPropertyPath& PropertyPath)
 {
 	if (PropertyPath.IsEmpty())
 	{
 		ensureAlways(false);
-		return MakeError(TEXT("Empty property path found. This is legal."));
+		return MakeError(LOCTEXT("EmptyPropertyPath", "Empty property path found. This is legal."));
 	}
 
 	FBindingSourceContext Result;
@@ -1032,7 +1170,7 @@ TValueOrError<FMVVMViewBlueprintCompiler::FBindingSourceContext, FString> FMVVMV
 	}
 
 	// The path may contains another INotifyFieldValueChanged
-	TValueOrError<FieldPathHelper::FParsedBindingInfo, FString> BindingInfoResult = FieldPathHelper::GetBindingInfoFromFieldPath(Class, Result.PropertyPath);
+	TValueOrError<FieldPathHelper::FParsedBindingInfo, FText> BindingInfoResult = FieldPathHelper::GetBindingInfoFromFieldPath(Class, Result.PropertyPath);
 	if (BindingInfoResult.HasError())
 	{
 		return MakeError(BindingInfoResult.StealError());
@@ -1044,7 +1182,7 @@ TValueOrError<FMVVMViewBlueprintCompiler::FBindingSourceContext, FString> FMVVMV
 	const bool bCreateADynamicSourceEntry = BindingInfo.ViewModelIndex > 1
 		&& BindingInfo.NotifyFieldClass != nullptr;
 
-	// The INotifyFieldValueChanged/viewmodel is not the fist property. Creates the ViaBinding source entries.
+	// The INotifyFieldValueChanged/viewmodel is not the first property. Creates the ViaBinding source entries.
 	if (bCreateADynamicSourceEntry)
 	{
 		Result.SourceClass = BindingInfo.NotifyFieldClass;
@@ -1064,7 +1202,7 @@ TValueOrError<FMVVMViewBlueprintCompiler::FBindingSourceContext, FString> FMVVMV
 		// Add if the path doesn't already exists
 		if (!CompilerSourceCreatorContexts.IsValidIndex(SourceCreatorContextIndex))
 		{
-			return MakeError(TEXT("The field in the path is from a viewmodel that is not the root of the path. This is not yet supported."));
+			return MakeError(LOCTEXT("ViewModelNotAtRoot", "The field in the path is from a viewmodel that is not the root of the path. This is not yet supported."));
 		}
 
 		Result.PropertyPath.RemoveAt(0, BindingInfo.ViewModelIndex);
@@ -1077,7 +1215,7 @@ TValueOrError<FMVVMViewBlueprintCompiler::FBindingSourceContext, FString> FMVVMV
 	{
 		if (Result.SourceClass != BindingInfo.NotifyFieldClass)
 		{
-			return MakeError(TEXT("Internal Error. The class doesn't match"));
+			return MakeError(LOCTEXT("ClassDoesNotMatch", "Internal Error. The class doesn't match"));
 		}
 	}
 
@@ -1115,9 +1253,10 @@ TArray<FMVVMConstFieldVariant> FMVVMViewBlueprintCompiler::CreateBindingDestinat
 		else
 		{
 			const int32 DestinationVariableContextIndex = CompilerSourceContexts.IndexOfByPredicate([DestinationName](const FCompilerSourceContext& Other) { return Other.PropertyName == DestinationName; });
-			check(DestinationVariableContextIndex != INDEX_NONE);
-
-			return CreatePropertyPath(Class, CompilerSourceContexts[DestinationVariableContextIndex].PropertyName, PropertyPath.GetFields());
+			if (ensureAlwaysMsgf(DestinationVariableContextIndex != INDEX_NONE, TEXT("Could not find source context for destination '%s'"), *DestinationName.ToString()))
+			{
+				return CreatePropertyPath(Class, CompilerSourceContexts[DestinationVariableContextIndex].PropertyName, PropertyPath.GetFields());
+			}
 		}
 	}
 	else
@@ -1125,44 +1264,8 @@ TArray<FMVVMConstFieldVariant> FMVVMViewBlueprintCompiler::CreateBindingDestinat
 		ensureAlwaysMsgf(false, TEXT("Not supported yet."));
 		return CreatePropertyPath(Class, FName(), PropertyPath.GetFields());
 	}
-}
 
-
-FString FMVVMViewBlueprintCompiler::PropertyPathToString(const UMVVMBlueprintView* BlueprintView, const FMVVMBlueprintPropertyPath& PropertyPath) const
-{
-	if (PropertyPath.IsEmpty())
-	{
-		return FString();
-	}
-
-	TStringBuilder<512> Result;
-	if (PropertyPath.IsFromViewModel())
-	{
-		if (const FMVVMBlueprintViewModelContext* SourceViewModelContext = BlueprintView->FindViewModel(PropertyPath.GetViewModelId()))
-		{
-			Result << SourceViewModelContext->GetViewModelName();
-		}
-		else
-		{
-			Result << TEXT("<none>");
-		}
-	}
-	else if (PropertyPath.IsFromWidget())
-	{
-		Result << PropertyPath.GetWidgetName();
-	}
-	else
-	{
-		Result << TEXT("<none>");
-	}
-
-	FString BasePropertyPath = PropertyPath.GetBasePropertyPath();
-	if (BasePropertyPath.Len())
-	{
-		Result << TEXT('.');
-		Result << MoveTemp(BasePropertyPath);
-	}
-	return Result.ToString();
+	return TArray<FMVVMConstFieldVariant>();
 }
 
 TArray<FMVVMConstFieldVariant> FMVVMViewBlueprintCompiler::CreatePropertyPath(const UClass* Class, FName PropertyName, TArray<FMVVMConstFieldVariant> Properties) const
@@ -1200,3 +1303,5 @@ bool FMVVMViewBlueprintCompiler::IsPropertyPathValid(TArrayView<const FMVVMConst
 }
 
 } //namespace
+
+#undef LOCTEXT_NAMESPACE
diff --git a/Engine/Plugins/Runtime/ModelViewViewModel/Source/ModelViewViewModelBlueprint/Private/MVVMViewBlueprintCompiler.h b/Engine/Plugins/Runtime/ModelViewViewModel/Source/ModelViewViewModelBlueprint/Private/MVVMViewBlueprintCompiler.h
--- a/Engine/Plugins/Runtime/ModelViewViewModel/Source/ModelViewViewModelBlueprint/Private/MVVMViewBlueprintCompiler.h
+++ b/Engine/Plugins/Runtime/ModelViewViewModel/Source/ModelViewViewModelBlueprint/Private/MVVMViewBlueprintCompiler.h
@@ -71,12 +71,12 @@ private:
 
 	const FCompilerSourceCreatorContext* FindViewModelSource(FGuid Id) const;
 
-	void AddErrorForBinding(FMVVMBlueprintViewBinding& Binding, const UWidgetBlueprint* WidgetBlueprint, const FString& Message) const;
+	void AddErrorForBinding(FMVVMBlueprintViewBinding& Binding, const FText& Message, FName ArgumentName = FName()) const;
+	void AddErrorForViewModel(const FMVVMBlueprintViewModelContext& ViewModel, const FText& Message) const;
 
-	TValueOrError<FBindingSourceContext, FString> CreateBindingSourceContext(const UMVVMBlueprintView* BlueprintView, const UWidgetBlueprintGeneratedClass* Class, const FMVVMBlueprintPropertyPath& PropertyPath);
+	TValueOrError<FBindingSourceContext, FText> CreateBindingSourceContext(const UMVVMBlueprintView* BlueprintView, const UWidgetBlueprintGeneratedClass* Class, const FMVVMBlueprintPropertyPath& PropertyPath);
 	TArray<FMVVMConstFieldVariant> CreateBindingDestinationPath(const UMVVMBlueprintView* BlueprintView, const UWidgetBlueprintGeneratedClass* Class, const FMVVMBlueprintPropertyPath& PropertyPath) const;
 
-	FString PropertyPathToString(const UMVVMBlueprintView* BlueprintView, const FMVVMBlueprintPropertyPath& PropertyPath) const;
 	TArray<FMVVMConstFieldVariant> CreatePropertyPath(const UClass* Class, FName PropertyName, TArray<FMVVMConstFieldVariant> Properties) const;
 	bool IsPropertyPathValid(TArrayView<const FMVVMConstFieldVariant> PropertyPath) const;
 
@@ -85,7 +85,7 @@ private:
 	{
 		UClass* Class = nullptr;
 		FName PropertyName;
-		FString DisplayName;
+		FText DisplayName;
 		FString CategoryName;
 		FString BlueprintSetter;
 		FMVVMConstFieldVariant Field;
diff --git a/Engine/Plugins/Runtime/ModelViewViewModel/Source/ModelViewViewModelBlueprint/Public/Bindings/MVVMCompiledBindingLibraryCompiler.h b/Engine/Plugins/Runtime/ModelViewViewModel/Source/ModelViewViewModelBlueprint/Public/Bindings/MVVMCompiledBindingLibraryCompiler.h
--- a/Engine/Plugins/Runtime/ModelViewViewModel/Source/ModelViewViewModelBlueprint/Public/Bindings/MVVMCompiledBindingLibraryCompiler.h
+++ b/Engine/Plugins/Runtime/ModelViewViewModel/Source/ModelViewViewModelBlueprint/Public/Bindings/MVVMCompiledBindingLibraryCompiler.h
@@ -156,25 +156,25 @@ public:
 
 public:
 	/** */
-	TValueOrError<FFieldIdHandle, FString> AddFieldId(TSubclassOf<UObject> SourceClass, FName FieldName);
+	TValueOrError<FFieldIdHandle, FText> AddFieldId(TSubclassOf<UObject> SourceClass, FName FieldName);
 	
 	/** */
-	TValueOrError<FFieldPathHandle, FString> AddFieldPath(TArrayView<const UE::MVVM::FMVVMConstFieldVariant> FieldPath, bool bRead);
+	TValueOrError<FFieldPathHandle, FText> AddFieldPath(TArrayView<const UE::MVVM::FMVVMConstFieldVariant> FieldPath, bool bRead);
 
 	/** */
-	TValueOrError<FFieldPathHandle, FString> AddObjectFieldPath(TArrayView<const UE::MVVM::FMVVMConstFieldVariant> FieldPath, UClass* ExpectedType, bool bRead);
+	TValueOrError<FFieldPathHandle, FText> AddObjectFieldPath(TArrayView<const UE::MVVM::FMVVMConstFieldVariant> FieldPath, UClass* ExpectedType, bool bRead);
 
 	/** */
-	TValueOrError<FFieldPathHandle, FString> AddConversionFunctionFieldPath(TSubclassOf<UObject> SourceClass, const UFunction* Function);
+	TValueOrError<FFieldPathHandle, FText> AddConversionFunctionFieldPath(TSubclassOf<UObject> SourceClass, const UFunction* Function);
 
 	/** */
-	TValueOrError<FBindingHandle, FString> AddBinding(FFieldPathHandle Source, FFieldPathHandle Destination);
+	TValueOrError<FBindingHandle, FText> AddBinding(FFieldPathHandle Source, FFieldPathHandle Destination);
 
 	/** */
-	TValueOrError<FBindingHandle, FString> AddBinding(FFieldPathHandle Source, FFieldPathHandle Destination, FFieldPathHandle ConversionFunction);
+	TValueOrError<FBindingHandle, FText> AddBinding(FFieldPathHandle Source, FFieldPathHandle Destination, FFieldPathHandle ConversionFunction);
 
 	/** */
-	TValueOrError<FBindingHandle, FString> AddBinding(TArrayView<const FFieldPathHandle> Sources, FFieldPathHandle Destination, FFieldPathHandle ConversionFunction);
+	TValueOrError<FBindingHandle, FText> AddBinding(TArrayView<const FFieldPathHandle> Sources, FFieldPathHandle Destination, FFieldPathHandle ConversionFunction);
 
 	struct FCompileResult
 	{
@@ -185,11 +185,11 @@ public:
 	};
 
 	/** */
-	TValueOrError<FCompileResult, FString> Compile();
+	TValueOrError<FCompileResult, FText> Compile();
 
 private:
 	/** */
-	TValueOrError<FFieldPathHandle, FString> AddFieldPathImpl(TArrayView<const UE::MVVM::FMVVMConstFieldVariant> FieldPath, bool bRead);
+	TValueOrError<FFieldPathHandle, FText> AddFieldPathImpl(TArrayView<const UE::MVVM::FMVVMConstFieldVariant> FieldPath, bool bRead);
 
 	TPimplPtr<Private::FCompiledBindingLibraryCompilerImpl> Impl;
 };
diff --git a/Engine/Plugins/Runtime/ModelViewViewModel/Source/ModelViewViewModelBlueprint/Public/Bindings/MVVMConversionFunctionHelper.h b/Engine/Plugins/Runtime/ModelViewViewModel/Source/ModelViewViewModelBlueprint/Public/Bindings/MVVMConversionFunctionHelper.h
--- a/Engine/Plugins/Runtime/ModelViewViewModel/Source/ModelViewViewModelBlueprint/Public/Bindings/MVVMConversionFunctionHelper.h
+++ b/Engine/Plugins/Runtime/ModelViewViewModel/Source/ModelViewViewModelBlueprint/Public/Bindings/MVVMConversionFunctionHelper.h
@@ -14,11 +14,14 @@ struct FMVVMBlueprintViewBinding;
 
 namespace UE::MVVM::ConversionFunctionHelper
 {
+	/** Find all BlueprintPropertyPath used by the given binding. */
+	MODELVIEWVIEWMODELBLUEPRINT_API TMap<FName, FMVVMBlueprintPropertyPath> GetAllArgumentPropertyPaths(const UWidgetBlueprint* WidgetBlueprint, const FMVVMBlueprintViewBinding& Binding, bool bSourceToDestination, bool bSkipResolve);
+
 	/** Find all BlueprintPropertyPath used by the given conversion function node. */
-	MODELVIEWVIEWMODELBLUEPRINT_API TArray<FMVVMBlueprintPropertyPath> GetAllArgumentPropertyPaths(const UWidgetBlueprint* WidgetBlueprint, const UK2Node_CallFunction* FunctionNode);
+	MODELVIEWVIEWMODELBLUEPRINT_API TMap<FName, FMVVMBlueprintPropertyPath> GetAllArgumentPropertyPaths(const UWidgetBlueprint* WidgetBlueprint, const UK2Node_CallFunction* FunctionNode, bool bSkipResolve);
 	
 	/** Find the property path of a given argument in the conversion function. */
-	MODELVIEWVIEWMODELBLUEPRINT_API FMVVMBlueprintPropertyPath GetPropertyPathForArgument(const UWidgetBlueprint* WidgetBlueprint, const UK2Node_CallFunction* Function, FName ArgumentName);
+	MODELVIEWVIEWMODELBLUEPRINT_API FMVVMBlueprintPropertyPath GetPropertyPathForArgument(const UWidgetBlueprint* WidgetBlueprint, const UK2Node_CallFunction* Function, FName ArgumentName, bool bSkipResolve);
 
 	/** Create the name of the conversion function wrapper this binding should have. */
 	MODELVIEWVIEWMODELBLUEPRINT_API FName CreateWrapperName(const FMVVMBlueprintViewBinding& Binding, bool bSourceToDestination);
diff --git a/Engine/Plugins/Runtime/ModelViewViewModel/Source/ModelViewViewModelBlueprint/Public/MVVMBlueprintView.h b/Engine/Plugins/Runtime/ModelViewViewModel/Source/ModelViewViewModelBlueprint/Public/MVVMBlueprintView.h
--- a/Engine/Plugins/Runtime/ModelViewViewModel/Source/ModelViewViewModelBlueprint/Public/MVVMBlueprintView.h
+++ b/Engine/Plugins/Runtime/ModelViewViewModel/Source/ModelViewViewModelBlueprint/Public/MVVMBlueprintView.h
@@ -64,6 +64,8 @@ public:
 	}
 
 #if WITH_EDITOR
+//	virtual void PostEditUndo() override;
+
 	virtual void PostEditChangeProperty(FPropertyChangedEvent& PropertyChangedEvent) override;
 	virtual void PostEditChangeChainProperty(FPropertyChangedChainEvent& PropertyChainEvent) override;
 
diff --git a/Engine/Plugins/Runtime/ModelViewViewModel/Source/ModelViewViewModelBlueprint/Public/MVVMPropertyPath.h b/Engine/Plugins/Runtime/ModelViewViewModel/Source/ModelViewViewModelBlueprint/Public/MVVMPropertyPath.h
--- a/Engine/Plugins/Runtime/ModelViewViewModel/Source/ModelViewViewModelBlueprint/Public/MVVMPropertyPath.h
+++ b/Engine/Plugins/Runtime/ModelViewViewModel/Source/ModelViewViewModelBlueprint/Public/MVVMPropertyPath.h
@@ -329,7 +329,9 @@ public:
 
 	bool operator==(const FMVVMBlueprintPropertyPath& Other) const
 	{
-		return Paths == Other.Paths;
+		return WidgetName == Other.WidgetName && 
+			ContextId == Other.ContextId && 
+			Paths == Other.Paths;
 	}
 
 	bool operator!=(const FMVVMBlueprintPropertyPath& Other) const
diff --git a/Engine/Plugins/Runtime/ModelViewViewModel/Source/ModelViewViewModelEditor/Private/MVVMEditorSubsystem.cpp b/Engine/Plugins/Runtime/ModelViewViewModel/Source/ModelViewViewModelEditor/Private/MVVMEditorSubsystem.cpp
--- a/Engine/Plugins/Runtime/ModelViewViewModel/Source/ModelViewViewModelEditor/Private/MVVMEditorSubsystem.cpp
+++ b/Engine/Plugins/Runtime/ModelViewViewModel/Source/ModelViewViewModelEditor/Private/MVVMEditorSubsystem.cpp
@@ -22,6 +22,7 @@
 #include "MVVMWidgetBlueprintExtension_View.h"
 #include "ScopedTransaction.h"
 #include "Types/MVVMBindingSource.h"
+#include "UObject/MetaData.h"
 #include "WidgetBlueprint.h"
 
 #include UE_INLINE_GENERATED_CPP_BY_NAME(MVVMEditorSubsystem)
@@ -156,11 +157,7 @@ namespace UE::MVVM::Private
 	{
 		if (UEdGraph* Graph = FindExistingConversionFunctionWrapper(WidgetBlueprint, WrapperName))
 		{
-			Graph->Rename(nullptr, GetTransientPackage(), REN_DoNotDirty | REN_ForceNoResetLoaders | REN_DontCreateRedirectors);
-			Graph->SetFlags(RF_Transient);
-			Graph->ClearFlags(RF_Public | RF_Standalone | RF_ArchetypeObject);
-			FLinkerLoad::InvalidateExport(Graph);
-			WidgetBlueprint->FunctionGraphs.RemoveSingle(Graph);
+			FBlueprintEditorUtils::RemoveGraph(WidgetBlueprint, Graph);
 		}
 	}
 
@@ -329,7 +326,7 @@ UEdGraph* UMVVMEditorSubsystem::CreateConversionFunctionWrapperGraph(UWidgetBlue
 {
 	if (UMVVMBlueprintView* View = GetView(WidgetBlueprint))
 	{
-		TValueOrError<TArray<const FProperty*>, FString> ArgumentsResult = UE::MVVM::BindingHelper::TryGetArgumentsForConversionFunction(ConversionFunction);
+		TValueOrError<TArray<const FProperty*>, FText> ArgumentsResult = UE::MVVM::BindingHelper::TryGetArgumentsForConversionFunction(ConversionFunction);
 		if (ArgumentsResult.HasValue())
 		{
 			TArray<const FProperty*> Arguments = ArgumentsResult.StealValue();
@@ -350,7 +347,7 @@ UEdGraph* UMVVMEditorSubsystem::CreateConversionFunctionWrapperGraph(UWidgetBlue
 				FGraphNodeCreator<UK2Node_CallFunction> CallFunctionCreator(*NewGraph);
 				UK2Node_CallFunction* CallFunctionNode = CallFunctionCreator.CreateNode();
 				CallFunctionNode->SetFromFunction(ConversionFunction);
-				CallFunctionNode->NodePosX = (FunctionEntry->NodePosX + FunctionResult->NodePosX) / 2;
+				CallFunctionNode->NodePosX = 0;
 				CallFunctionCreator.Finalize();
 
 				UE::MVVM::ConversionFunctionHelper::MarkAsConversionFunction(CallFunctionNode, Binding);
@@ -561,7 +558,7 @@ void UMVVMEditorSubsystem::SetCompileForBinding(UWidgetBlueprint* WidgetBlueprin
 
 bool UMVVMEditorSubsystem::IsValidConversionFunction(const UFunction* Function, const FMVVMBlueprintPropertyPath& Source, const FMVVMBlueprintPropertyPath& Destination) const
 {
-	TValueOrError<const FProperty*, FString> ReturnResult = UE::MVVM::BindingHelper::TryGetReturnTypeForConversionFunction(Function);
+	TValueOrError<const FProperty*, FText> ReturnResult = UE::MVVM::BindingHelper::TryGetReturnTypeForConversionFunction(Function);
 	if (ReturnResult.HasError())
 	{
 		return false;
@@ -573,7 +570,7 @@ bool UMVVMEditorSubsystem::IsValidConversionFunction(const UFunction* Function,
 		return false;
 	}
 
-	TValueOrError<TArray<const FProperty*>, FString> ArgumentsResult = UE::MVVM::BindingHelper::TryGetArgumentsForConversionFunction(Function);
+	TValueOrError<TArray<const FProperty*>, FText> ArgumentsResult = UE::MVVM::BindingHelper::TryGetArgumentsForConversionFunction(Function);
 	if (ArgumentsResult.HasError())
 	{
 		return false;
@@ -631,7 +628,7 @@ bool UMVVMEditorSubsystem::IsValidConversionFunction(const UFunction* Function,
 
 bool UMVVMEditorSubsystem::IsSimpleConversionFunction(const UFunction* Function) const
 {
-	TValueOrError<const FProperty*, FString> ReturnResult = UE::MVVM::BindingHelper::TryGetReturnTypeForConversionFunction(Function);
+	TValueOrError<const FProperty*, FText> ReturnResult = UE::MVVM::BindingHelper::TryGetReturnTypeForConversionFunction(Function);
 	if (ReturnResult.HasError())
 	{
 		return false;
@@ -643,7 +640,7 @@ bool UMVVMEditorSubsystem::IsSimpleConversionFunction(const UFunction* Function)
 		return false;
 	}
 
-	TValueOrError<TArray<const FProperty*>, FString> ArgumentsResult = UE::MVVM::BindingHelper::TryGetArgumentsForConversionFunction(Function);
+	TValueOrError<TArray<const FProperty*>, FText> ArgumentsResult = UE::MVVM::BindingHelper::TryGetArgumentsForConversionFunction(Function);
 	if (ArgumentsResult.HasError())
 	{
 		return false;
@@ -715,7 +712,7 @@ FMVVMBlueprintPropertyPath UMVVMEditorSubsystem::GetPathForConversionFunctionArg
 		return FMVVMBlueprintPropertyPath();
 	}
 
-	return UE::MVVM::ConversionFunctionHelper::GetPropertyPathForArgument(WidgetBlueprint, FunctionNode, ArgumentName);
+	return UE::MVVM::ConversionFunctionHelper::GetPropertyPathForArgument(WidgetBlueprint, FunctionNode, ArgumentName, false);
 }
 
 void UMVVMEditorSubsystem::SetPathForConversionFunctionArgument(UWidgetBlueprint* WidgetBlueprint, FMVVMBlueprintViewBinding& Binding, FName ArgumentName, const FMVVMBlueprintPropertyPath& Path, bool bSourceToDestination) const
@@ -801,8 +798,8 @@ void UMVVMEditorSubsystem::SetPathForConversionFunctionArgument(UWidgetBlueprint
 
 	TArray<UE::MVVM::FMVVMConstFieldVariant> Fields = Path.GetFields();
 
-	float PosX = ConversionNode->NodePosX - (300 * Fields.Num());
-	const float PosY = ConversionNode->NodePosY + ArgumentIndex * 50;
+	float PosX = ConversionNode->NodePosX - 300 * (Fields.Num() + 1);
+	const float PosY = ConversionNode->NodePosY + ArgumentIndex * 100;
 
 	UEdGraphPin* PreviousExecPin = UE::MVVM::Private::FindFunctionEntry(ConversionFunctionGraph)->FindPin(UEdGraphSchema_K2::PN_Then);
 	UEdGraphPin* PreviousDataPin = nullptr;
diff --git a/Engine/Plugins/Runtime/ModelViewViewModel/Source/ModelViewViewModelEditor/Private/Tests/MVVMBindingHelperTest.cpp b/Engine/Plugins/Runtime/ModelViewViewModel/Source/ModelViewViewModelEditor/Private/Tests/MVVMBindingHelperTest.cpp
--- a/Engine/Plugins/Runtime/ModelViewViewModel/Source/ModelViewViewModelEditor/Private/Tests/MVVMBindingHelperTest.cpp
+++ b/Engine/Plugins/Runtime/ModelViewViewModel/Source/ModelViewViewModelEditor/Private/Tests/MVVMBindingHelperTest.cpp
@@ -101,7 +101,7 @@ bool FMVVMBindingHelperTest::RunTest(const FString& Parameters)
 		}
 
 		{
-			TValueOrError<const FProperty*, FString> Result = ::UE::MVVM::BindingHelper::TryGetPropertyTypeForSourceBinding(Function);
+			TValueOrError<const FProperty*, FText> Result = ::UE::MVVM::BindingHelper::TryGetPropertyTypeForSourceBinding(Function);
 			if ((PropertyTypeForSourceBinding != nullptr) != Result.HasValue())
 			{
 				AddError(FString::Printf(TEXT("TryGetPropertyTypeForSourceBinding failed for '%s'."), *FunctionName.ToString()));
@@ -112,7 +112,7 @@ bool FMVVMBindingHelperTest::RunTest(const FString& Parameters)
 			}
 		}
 		{
-			TValueOrError<const FProperty*, FString> Result = ::UE::MVVM::BindingHelper::TryGetPropertyTypeForDestinationBinding(Function);
+			TValueOrError<const FProperty*, FText> Result = ::UE::MVVM::BindingHelper::TryGetPropertyTypeForDestinationBinding(Function);
 			if ((PropertyTypeForDestinationBinding != nullptr) != Result.HasValue())
 			{
 				AddError(FString::Printf(TEXT("TryGetPropertyTypeForDestinationBinding failed for '%s'."), *FunctionName.ToString()));
diff --git a/Engine/Plugins/Runtime/ModelViewViewModel/Source/ModelViewViewModelEditor/Private/Widgets/SMVVMViewBindingListView.cpp b/Engine/Plugins/Runtime/ModelViewViewModel/Source/ModelViewViewModelEditor/Private/Widgets/SMVVMViewBindingListView.cpp
--- a/Engine/Plugins/Runtime/ModelViewViewModel/Source/ModelViewViewModelEditor/Private/Widgets/SMVVMViewBindingListView.cpp
+++ b/Engine/Plugins/Runtime/ModelViewViewModel/Source/ModelViewViewModelEditor/Private/Widgets/SMVVMViewBindingListView.cpp
@@ -1074,7 +1074,7 @@ void SBindingsList::Refresh()
 			const UFunction* Function = EditorSubsystem->GetConversionFunction(MVVMExtensionPtr->GetWidgetBlueprint(), Binding, UE::MVVM::IsForwardBinding(Binding.BindingType));
 			if (Function != nullptr)
 			{
-				TValueOrError<TArray<const FProperty*>, FString> ArgumentsResult = BindingHelper::TryGetArgumentsForConversionFunction(Function);
+				TValueOrError<TArray<const FProperty*>, FText> ArgumentsResult = BindingHelper::TryGetArgumentsForConversionFunction(Function);
 				if (ArgumentsResult.HasValue())
 				{
 					for (const FProperty* Argument : ArgumentsResult.GetValue())
diff --git a/Engine/Source/Editor/UMGEditor/Public/WidgetBlueprintCompiler.h b/Engine/Source/Editor/UMGEditor/Public/WidgetBlueprintCompiler.h
--- a/Engine/Source/Editor/UMGEditor/Public/WidgetBlueprintCompiler.h
+++ b/Engine/Source/Editor/UMGEditor/Public/WidgetBlueprintCompiler.h
@@ -11,6 +11,10 @@
 #include "ComponentReregisterContext.h"
 #include "Components/WidgetComponent.h"
 
+class FProperty;
+class UWidget;
+class UWidgetAnimation;
+class UWidgetGraphSchema;
 
 //////////////////////////////////////////////////////////////////////////
 // FWidgetBlueprintCompiler 
@@ -112,13 +116,13 @@ protected:
 
 	TArray<UWidgetAnimation*> OldWidgetAnimations;
 
-	class UWidgetGraphSchema* WidgetSchema;
+	UWidgetGraphSchema* WidgetSchema;
 
 	// Map of properties created for widgets; to aid in debug data generation
-	TMap<class UWidget*, class FProperty*> WidgetToMemberVariableMap;
+	TMap<UWidget*, FProperty*> WidgetToMemberVariableMap;
 
 	// Map of properties created for widget animations; to aid in debug data generation
-	TMap<class UWidgetAnimation*, class FProperty*> WidgetAnimToMemberVariableMap;
+	TMap<UWidgetAnimation*, FProperty*> WidgetAnimToMemberVariableMap;
 
 	///----------------------------------------------------------------
 };
