Fixed LOCTEXT mismatch in MVVMFieldPathHelper.

#rnx

diff --git a/Engine/Plugins/Runtime/ModelViewViewModel/Source/ModelViewViewModel/Private/Bindings/MVVMFieldPathHelper.cpp b/Engine/Plugins/Runtime/ModelViewViewModel/Source/ModelViewViewModel/Private/Bindings/MVVMFieldPathHelper.cpp
--- a/Engine/Plugins/Runtime/ModelViewViewModel/Source/ModelViewViewModel/Private/Bindings/MVVMFieldPathHelper.cpp
+++ b/Engine/Plugins/Runtime/ModelViewViewModel/Source/ModelViewViewModel/Private/Bindings/MVVMFieldPathHelper.cpp
@@ -120,7 +120,7 @@ TValueOrError<TArray<FMVVMConstFieldVariant>, FText> GenerateFieldPathList(const
 		FMVVMConstFieldVariant Field = BindingHelper::FindFieldByName(CurrentContainer, FMVVMBindingName(FName(FoundIndex, InFieldPath.GetData())));
 		if (Field.IsEmpty())
 		{
-			return MakeError(FText::Format(LOCTEXT("FieldDoesNotExistInStruct", "The field '%s' does not exist in the struct '%s'."), 
+			return MakeError(FText::Format(LOCTEXT("FieldDoesNotExistInStruct", "The field '{0}' does not exist in the struct '{1}'."), 
 				FText::FromName(FName(FoundIndex, InFieldPath.GetData())), 
 				CurrentContainer ? FText::FromString(CurrentContainer->GetName()) : LOCTEXT("None", "<none>"))
 			);
