MVVM: Set the category on the generated widget property with the cagetory sent by the user.

diff --git a/Engine/Plugins/Runtime/ModelViewViewModel/Source/ModelViewViewModelBlueprint/Private/MVVMViewBlueprintCompiler.cpp b/Engine/Plugins/Runtime/ModelViewViewModel/Source/ModelViewViewModelBlueprint/Private/MVVMViewBlueprintCompiler.cpp
--- a/Engine/Plugins/Runtime/ModelViewViewModel/Source/ModelViewViewModelBlueprint/Private/MVVMViewBlueprintCompiler.cpp
+++ b/Engine/Plugins/Runtime/ModelViewViewModel/Source/ModelViewViewModelBlueprint/Private/MVVMViewBlueprintCompiler.cpp
@@ -176,33 +176,33 @@ void FMVVMViewBlueprintCompiler::CreateVariables(const FWidgetBlueprintCompilerC
 	auto CreateVariable = [&Context](const FCompilerSourceContext& SourceContext) -> FProperty*
 	{
 		FEdGraphPinType ViewModelPinType(UEdGraphSchema_K2::PC_Object, NAME_None, SourceContext.Class, EPinContainerType::None, false, FEdGraphTerminalType());
-		FProperty* ViewModelProperty = Context.CreateVariable(SourceContext.PropertyName, ViewModelPinType);
-		if (ViewModelProperty != nullptr)
+		FProperty* NewProperty = Context.CreateVariable(SourceContext.PropertyName, ViewModelPinType);
+		if (NewProperty != nullptr)
 		{
-			ViewModelProperty->SetPropertyFlags(CPF_BlueprintVisible | CPF_BlueprintReadOnly | CPF_RepSkip
+			NewProperty->SetPropertyFlags(CPF_BlueprintVisible | CPF_BlueprintReadOnly | CPF_RepSkip
 				| CPF_Transient | CPF_DuplicateTransient);
-			ViewModelProperty->SetPropertyFlags(SourceContext.bExposeOnSpawn ? CPF_ExposeOnSpawn : CPF_DisableEditOnInstance);
+			NewProperty->SetPropertyFlags(SourceContext.bExposeOnSpawn ? CPF_ExposeOnSpawn : CPF_DisableEditOnInstance);
 
 #if WITH_EDITOR
 			if (!SourceContext.BlueprintSetter.IsEmpty())
 			{
-				ViewModelProperty->SetMetaData(FBlueprintMetadata::MD_PropertySetFunction, *SourceContext.BlueprintSetter);
+				NewProperty->SetMetaData(FBlueprintMetadata::MD_PropertySetFunction, *SourceContext.BlueprintSetter);
 			}
 			if (!SourceContext.DisplayName.IsEmpty())
 			{
-				ViewModelProperty->SetMetaData(FBlueprintMetadata::MD_FunctionCategory, *SourceContext.DisplayName.ToString());
+				NewProperty->SetMetaData(FBlueprintMetadata::MD_DisplayName, *SourceContext.DisplayName.ToString());
 			}
 			if (!SourceContext.CategoryName.IsEmpty())
 			{
-				ViewModelProperty->SetMetaData(FBlueprintMetadata::MD_FunctionCategory, *SourceContext.CategoryName);
+				NewProperty->SetMetaData(FBlueprintMetadata::MD_FunctionCategory, *SourceContext.CategoryName);
 			}
 			if (SourceContext.bExposeOnSpawn)
 			{
-				ViewModelProperty->SetMetaData(FBlueprintMetadata::MD_ExposeOnSpawn, TEXT("true"));
+				NewProperty->SetMetaData(FBlueprintMetadata::MD_ExposeOnSpawn, TEXT("true"));
 			}
 #endif
 		}
-		return ViewModelProperty;
+		return NewProperty;
 	};
 
 	for (FCompilerSourceContext& SourceContext : CompilerSourceContexts)
@@ -361,6 +361,7 @@ void FMVVMViewBlueprintCompiler::CreateSourceLists(const FWidgetBlueprintCompile
 	bAreSourceContextsValid = bAreSourcesCreatorValid;
 
 	UWidgetBlueprintGeneratedClass* SkeletonClass = Context.GetSkeletonGeneratedClass();
+	const FName DefaultWidgetCategory = Context.GetWidgetBlueprint()->GetFName();
 
 	// Only find the source first property and destination first property.
 	//The full path will be tested later. We want to build the list of property needed.
@@ -381,7 +382,7 @@ void FMVVMViewBlueprintCompiler::CreateSourceLists(const FWidgetBlueprintCompile
 		}
 
 		FMVVMViewBlueprintCompiler* Self = this;
-		auto GenerateCompilerSourceContext = [Self, BlueprintView, Class = SkeletonClass, &Binding, &ViewModelGuids, &WidgetSources](const FMVVMBlueprintPropertyPath& PropertyPath, bool bViewModelPath, FName ArgumentName = FName()) -> bool
+		auto GenerateCompilerSourceContext = [Self, BlueprintView, DefaultWidgetCategory, Class = SkeletonClass, &Binding, &ViewModelGuids, &WidgetSources](const FMVVMBlueprintPropertyPath& PropertyPath, bool bViewModelPath, FName ArgumentName = FName()) -> bool
 		{
 			if (PropertyPath.IsFromWidget())
 			{
@@ -422,7 +423,7 @@ void FMVVMViewBlueprintCompiler::CreateSourceLists(const FWidgetBlueprintCompile
 					SourceVariable.Class = Widget->GetClass();
 					SourceVariable.PropertyName = PropertyPath.GetWidgetName();
 					SourceVariable.DisplayName = FText::FromString(Widget->GetDisplayLabel());
-					SourceVariable.CategoryName = TEXT("Widget");
+					SourceVariable.CategoryName = !Widget->GetCategoryName().IsEmpty() ? Widget->GetCategoryName() : DefaultWidgetCategory.ToString();
 					Self->CompilerSourceContexts.Emplace(MoveTemp(SourceVariable));
 				}
 			}
