MVVM: ViewModel Panel display all the "notify" properties.

diff --git a/Engine/Plugins/Runtime/ModelViewViewModel/Source/ModelViewViewModelEditor/Private/Widgets/SMVVMViewModelPanel.cpp b/Engine/Plugins/Runtime/ModelViewViewModel/Source/ModelViewViewModelEditor/Private/Widgets/SMVVMViewModelPanel.cpp
--- a/Engine/Plugins/Runtime/ModelViewViewModel/Source/ModelViewViewModelEditor/Private/Widgets/SMVVMViewModelPanel.cpp
+++ b/Engine/Plugins/Runtime/ModelViewViewModel/Source/ModelViewViewModelEditor/Private/Widgets/SMVVMViewModelPanel.cpp
@@ -50,7 +50,7 @@ void SMVVMViewModelPanel::Construct(const FArguments& InArgs, TSharedPtr<FWidget
 
 	WeakBlueprintEditor = WidgetBlueprintEditor;
 	WeakBlueprintView = CurrentBlueprintView;
-	FieldIterator = MakeUnique<FFieldIterator_Bindable>(WidgetBlueprint, EFieldVisibility::All);
+	FieldIterator = MakeUnique<FFieldIterator_Bindable>(WidgetBlueprint, EFieldVisibility::Notify);
 
 	if (CurrentBlueprintView)
 	{
