MVVM: Prevent access to protected properties and functions.

diff --git a/Engine/Plugins/Runtime/ModelViewViewModel/Source/ModelViewViewModel/Private/MVVMSubsystem.cpp b/Engine/Plugins/Runtime/ModelViewViewModel/Source/ModelViewViewModel/Private/MVVMSubsystem.cpp
--- a/Engine/Plugins/Runtime/ModelViewViewModel/Source/ModelViewViewModel/Private/MVVMSubsystem.cpp
+++ b/Engine/Plugins/Runtime/ModelViewViewModel/Source/ModelViewViewModel/Private/MVVMSubsystem.cpp
@@ -81,6 +81,11 @@ namespace UE::MVVM::Private
 		{
 			return FMVVMAvailableBinding();
 		}
+		static FName NAME_BlueprintProtected = "BlueprintProtected";
+		if (Property->GetBoolMetaData(NAME_BlueprintProtected) && !bCanAccessProtected)
+		{
+			return FMVVMAvailableBinding();
+		}
 #endif
 
 		const bool bIsReadable = BindingHelper::IsValidForSourceBinding(Property);
