MVVM: Display the correct name in the detail view for the source and destination property. Use the correct property, as filter, for the mode direction.

diff --git a/Engine/Plugins/Runtime/ModelViewViewModel/Source/ModelViewViewModel/Private/MVVMSubsystem.cpp b/Engine/Plugins/Runtime/ModelViewViewModel/Source/ModelViewViewModel/Private/MVVMSubsystem.cpp
--- a/Engine/Plugins/Runtime/ModelViewViewModel/Source/ModelViewViewModel/Private/MVVMSubsystem.cpp
+++ b/Engine/Plugins/Runtime/ModelViewViewModel/Source/ModelViewViewModel/Private/MVVMSubsystem.cpp
@@ -77,12 +77,12 @@ namespace UE::MVVM::Private
 		//Only property defined in BP can be BP visible and really private.
 #if WITH_EDITOR
 		static FName NAME_BlueprintPrivate = "BlueprintPrivate";
-		if (Property->GetBoolMetaData(NAME_BlueprintPrivate) && !bCanAccessPrivate)
+		if (Property->HasMetaData(NAME_BlueprintPrivate) && !bCanAccessPrivate)
 		{
 			return FMVVMAvailableBinding();
 		}
 		static FName NAME_BlueprintProtected = "BlueprintProtected";
-		if (Property->GetBoolMetaData(NAME_BlueprintProtected) && !bCanAccessProtected)
+		if (Property->HasMetaData(NAME_BlueprintProtected) && !bCanAccessProtected)
 		{
 			return FMVVMAvailableBinding();
 		}
diff --git a/Engine/Plugins/Runtime/ModelViewViewModel/Source/ModelViewViewModelEditor/Private/Customizations/MVVMPropertyPathCustomization.cpp b/Engine/Plugins/Runtime/ModelViewViewModel/Source/ModelViewViewModelEditor/Private/Customizations/MVVMPropertyPathCustomization.cpp
--- a/Engine/Plugins/Runtime/ModelViewViewModel/Source/ModelViewViewModelEditor/Private/Customizations/MVVMPropertyPathCustomization.cpp
+++ b/Engine/Plugins/Runtime/ModelViewViewModel/Source/ModelViewViewModelEditor/Private/Customizations/MVVMPropertyPathCustomization.cpp
@@ -45,13 +45,17 @@ void FPropertyPathCustomization::CustomizeHeader(TSharedRef<IPropertyHandle> InP
 	PropertyHandle = InPropertyHandle;
 
 	const FName PropertyName = PropertyHandle->GetProperty()->GetFName();
+	FText PropertyDisplayName;
 	if (PropertyName == GET_MEMBER_NAME_CHECKED(FMVVMBlueprintViewBinding, WidgetPath))
 	{
 		bIsWidget = true;
+		PropertyDisplayName = LOCTEXT("PropertyDisplay_Widget", "Destination");
 	}
 	else if (PropertyName == GET_MEMBER_NAME_CHECKED(FMVVMBlueprintViewBinding, ViewModelPath))
 	{
 		bIsWidget = false;
+
+		PropertyDisplayName = LOCTEXT("PropertyDisplay_Source", "Source");
 	}
 	else
 	{
@@ -96,7 +100,7 @@ void FPropertyPathCustomization::CustomizeHeader(TSharedRef<IPropertyHandle> InP
 		.NameWidget
 		.VAlign(VAlign_Center)
 		[
-			PropertyHandle->CreatePropertyNameWidget()
+			PropertyHandle->CreatePropertyNameWidget(PropertyDisplayName)
 		]
 		.ValueWidget
 		.VAlign(VAlign_Center)
diff --git a/Engine/Plugins/Runtime/ModelViewViewModel/Source/ModelViewViewModelEditor/Private/Widgets/SMVVMViewBindingListView.cpp b/Engine/Plugins/Runtime/ModelViewViewModel/Source/ModelViewViewModelEditor/Private/Widgets/SMVVMViewBindingListView.cpp
--- a/Engine/Plugins/Runtime/ModelViewViewModel/Source/ModelViewViewModelEditor/Private/Widgets/SMVVMViewBindingListView.cpp
+++ b/Engine/Plugins/Runtime/ModelViewViewModel/Source/ModelViewViewModelEditor/Private/Widgets/SMVVMViewBindingListView.cpp
@@ -542,7 +542,7 @@ private:
 					const UFunction* Function = LastField.GetFunction();
 					if (Function != nullptr)
 					{
-						return BindingHelper::GetFirstArgumentProperty(Function);
+						return BindingHelper::GetReturnProperty(Function);
 					}
 				}
 			}
