MVVM: Saving binding id.

diff --git a/Engine/Plugins/Runtime/ModelViewViewModel/Source/ModelViewViewModelBlueprint/Private/MVVMBlueprintView.cpp b/Engine/Plugins/Runtime/ModelViewViewModel/Source/ModelViewViewModelBlueprint/Private/MVVMBlueprintView.cpp
--- a/Engine/Plugins/Runtime/ModelViewViewModel/Source/ModelViewViewModelBlueprint/Private/MVVMBlueprintView.cpp
+++ b/Engine/Plugins/Runtime/ModelViewViewModel/Source/ModelViewViewModelBlueprint/Private/MVVMBlueprintView.cpp
@@ -176,6 +176,16 @@ const FMVVMBlueprintViewBinding* UMVVMBlueprintView::GetBindingAt(int32 Index) c
 	return nullptr;
 }
 
+FMVVMBlueprintViewBinding* UMVVMBlueprintView::GetBinding(FGuid Id)
+{
+	return Bindings.FindByPredicate([Id](const FMVVMBlueprintViewBinding& Binding){ return Id == Binding.BindingId; });
+}
+
+const FMVVMBlueprintViewBinding* UMVVMBlueprintView::GetBinding(FGuid Id) const
+{
+	return Bindings.FindByPredicate([Id](const FMVVMBlueprintViewBinding& Binding) { return Id == Binding.BindingId; });
+}
+
 void UMVVMBlueprintView::PostLoad()
 {
 	Super::PostLoad();
diff --git a/Engine/Plugins/Runtime/ModelViewViewModel/Source/ModelViewViewModelBlueprint/Public/MVVMBlueprintView.h b/Engine/Plugins/Runtime/ModelViewViewModel/Source/ModelViewViewModelBlueprint/Public/MVVMBlueprintView.h
--- a/Engine/Plugins/Runtime/ModelViewViewModel/Source/ModelViewViewModelBlueprint/Public/MVVMBlueprintView.h
+++ b/Engine/Plugins/Runtime/ModelViewViewModel/Source/ModelViewViewModelBlueprint/Public/MVVMBlueprintView.h
@@ -52,6 +52,8 @@ public:
 
 	FMVVMBlueprintViewBinding* GetBindingAt(int32 Index);
 	const FMVVMBlueprintViewBinding* GetBindingAt(int32 Index) const;
+	FMVVMBlueprintViewBinding* GetBinding(FGuid Id);
+	const FMVVMBlueprintViewBinding* GetBinding(FGuid Id) const;
 
 	TArrayView<FMVVMBlueprintViewBinding> GetBindings()
 	{
diff --git a/Engine/Plugins/Runtime/ModelViewViewModel/Source/ModelViewViewModelBlueprint/Public/MVVMBlueprintViewBinding.h b/Engine/Plugins/Runtime/ModelViewViewModel/Source/ModelViewViewModelBlueprint/Public/MVVMBlueprintViewBinding.h
--- a/Engine/Plugins/Runtime/ModelViewViewModel/Source/ModelViewViewModelBlueprint/Public/MVVMBlueprintViewBinding.h
+++ b/Engine/Plugins/Runtime/ModelViewViewModel/Source/ModelViewViewModelBlueprint/Public/MVVMBlueprintViewBinding.h
@@ -63,6 +63,7 @@ struct MODELVIEWVIEWMODELBLUEPRINT_API FMVVMBlueprintViewBinding
 	TArray<FText> Errors;
 
 	/** The unique ID of this binding. */
+	UPROPERTY(VisibleAnywhere, Category = "MVVM")
 	FGuid BindingId;
 
 	/** Whether the binding is enabled or disabled by default. The instance may enable the binding at runtime. */
