MVVM: Fixed plugin description to describe what MVVM means.

#rnx

diff --git a/Engine/Plugins/Runtime/ModelViewViewModel/ModelViewViewModel.uplugin b/Engine/Plugins/Runtime/ModelViewViewModel/ModelViewViewModel.uplugin
--- a/Engine/Plugins/Runtime/ModelViewViewModel/ModelViewViewModel.uplugin
+++ b/Engine/Plugins/Runtime/ModelViewViewModel/ModelViewViewModel.uplugin
@@ -6,7 +6,7 @@
 	"VersionName": "1.0",
 	"CreatedBy": "Epic Games, Inc.",
 	"CreatedByURL": "https://epicgames.com",
-	"Description": "A plugin to support the Viewmodel pattern in UMG.",
+	"Description": "A plugin to support the Model-View-Viewmodel pattern in UMG.",
 	"Category": "UI",
 	"EnabledByDefault": false,
 	"CanContainContent": false,
