Updating MVVMViewModelCollection to support finding a View Model by class regardless of name.

diff --git a/Engine/Plugins/Runtime/ModelViewViewModel/Source/ModelViewViewModel/Private/Types/MVVMViewModelCollection.cpp b/Engine/Plugins/Runtime/ModelViewViewModel/Source/ModelViewViewModel/Private/Types/MVVMViewModelCollection.cpp
--- a/Engine/Plugins/Runtime/ModelViewViewModel/Source/ModelViewViewModel/Private/Types/MVVMViewModelCollection.cpp
+++ b/Engine/Plugins/Runtime/ModelViewViewModel/Source/ModelViewViewModel/Private/Types/MVVMViewModelCollection.cpp
@@ -5,13 +5,17 @@
 
 #include UE_INLINE_GENERATED_CPP_BY_NAME(MVVMViewModelCollection)
 
-
 UMVVMViewModelBase* FMVVMViewModelCollection::FindViewModelInstance(FMVVMViewModelContext Context) const
 {
 	FMVVMViewModelContextInstance* FoundInstance = ViewModelInstances.FindByPredicate([Context](const FMVVMViewModelContextInstance& Other) { return Other.GetContext() == Context; });
 	return FoundInstance ? FoundInstance->GetViewModel() : nullptr;
 }
 
+UMVVMViewModelBase* FMVVMViewModelCollection::FindFirstViewModelInstanceOfType(const TSubclassOf<UMVVMViewModelBase>& ViewModelClass) const
+{
+	FMVVMViewModelContextInstance* FoundInstance = ViewModelInstances.FindByPredicate([ViewModelClass](const FMVVMViewModelContextInstance& Other) { return Other.GetContext().IsCompatibleWith(ViewModelClass); });
+	return FoundInstance ? FoundInstance->GetViewModel() : nullptr;
+}
 
 bool FMVVMViewModelCollection::AddInstance(FMVVMViewModelContext Context, UMVVMViewModelBase* ViewModel)
 {
diff --git a/Engine/Plugins/Runtime/ModelViewViewModel/Source/ModelViewViewModel/Private/Types/MVVMViewModelContext.cpp b/Engine/Plugins/Runtime/ModelViewViewModel/Source/ModelViewViewModel/Private/Types/MVVMViewModelContext.cpp
--- a/Engine/Plugins/Runtime/ModelViewViewModel/Source/ModelViewViewModel/Private/Types/MVVMViewModelContext.cpp
+++ b/Engine/Plugins/Runtime/ModelViewViewModel/Source/ModelViewViewModel/Private/Types/MVVMViewModelContext.cpp
@@ -20,10 +20,14 @@ bool FMVVMViewModelContext::operator== (const FMVVMViewModelContext& Other) cons
 
 bool FMVVMViewModelContext::IsCompatibleWith(const FMVVMViewModelContext& Other) const
 {
-	return Other.ContextName == ContextName
-		&& Other.ContextClass.Get() != nullptr
+	return Other.ContextName == ContextName && IsCompatibleWith(Other.ContextClass);
+}
+
+bool FMVVMViewModelContext::IsCompatibleWith(const TSubclassOf<UMVVMViewModelBase>& OtherClass) const
+{
+	return OtherClass != nullptr
 		&& ContextClass.Get() != nullptr
-		&& ContextClass.Get()->IsChildOf(Other.ContextClass.Get());
+		&& ContextClass.Get()->IsChildOf(OtherClass);
 }
 
 
diff --git a/Engine/Plugins/Runtime/ModelViewViewModel/Source/ModelViewViewModel/Public/Types/MVVMViewModelCollection.h b/Engine/Plugins/Runtime/ModelViewViewModel/Source/ModelViewViewModel/Public/Types/MVVMViewModelCollection.h
--- a/Engine/Plugins/Runtime/ModelViewViewModel/Source/ModelViewViewModel/Public/Types/MVVMViewModelCollection.h
+++ b/Engine/Plugins/Runtime/ModelViewViewModel/Source/ModelViewViewModel/Public/Types/MVVMViewModelCollection.h
@@ -19,6 +19,7 @@ struct MODELVIEWVIEWMODEL_API FMVVMViewModelCollection
 
 public:
 	UMVVMViewModelBase* FindViewModelInstance(FMVVMViewModelContext Context) const;
+	UMVVMViewModelBase* FindFirstViewModelInstanceOfType(const TSubclassOf<UMVVMViewModelBase>& ViewModelClass) const;
 
 	bool AddInstance(FMVVMViewModelContext Context, UMVVMViewModelBase* ViewModel);
 	bool RemoveInstance(FMVVMViewModelContext Context);
@@ -51,6 +52,16 @@ public:
 	{
 		return ViewModelCollection.FindViewModelInstance(Context);
 	}
+
+	/**
+	 * Finds a View Model of the given type.
+	 * If the collection contains multiple instances of the same type then this will return the first one found.
+	 */
+	UFUNCTION(BlueprintCallable, Category = "MVVM")
+	UMVVMViewModelBase* FindFirstViewModelInstanceOfType(const TSubclassOf<UMVVMViewModelBase>& ViewModelClass) const
+	{
+		return ViewModelCollection.FindFirstViewModelInstanceOfType(ViewModelClass);
+	}
 	
 	UFUNCTION(BlueprintCallable, Category = "MVVM")
 	bool AddViewModelInstance(FMVVMViewModelContext Context, UMVVMViewModelBase* ViewModel)
diff --git a/Engine/Plugins/Runtime/ModelViewViewModel/Source/ModelViewViewModel/Public/Types/MVVMViewModelContext.h b/Engine/Plugins/Runtime/ModelViewViewModel/Source/ModelViewViewModel/Public/Types/MVVMViewModelContext.h
--- a/Engine/Plugins/Runtime/ModelViewViewModel/Source/ModelViewViewModel/Public/Types/MVVMViewModelContext.h
+++ b/Engine/Plugins/Runtime/ModelViewViewModel/Source/ModelViewViewModel/Public/Types/MVVMViewModelContext.h
@@ -28,5 +28,6 @@ public:
 	bool IsValid() const;
 	bool operator== (const FMVVMViewModelContext& Other) const;
 	bool IsCompatibleWith(const FMVVMViewModelContext& Other) const;
+	bool IsCompatibleWith(const TSubclassOf<UMVVMViewModelBase>& OtherClass) const;
 	bool IsCompatibleWith(const UMVVMViewModelBase* Other) const;
 };
