MVVM: Field selector now show a tooltip of the full path, and functions show the function signature.

It is possible to show only the last field in a path when space is limited.

#rnx

diff --git a/Engine/Plugins/Runtime/ModelViewViewModel/Source/ModelViewViewModelEditor/Private/Widgets/SMVVMFieldEntry.cpp b/Engine/Plugins/Runtime/ModelViewViewModel/Source/ModelViewViewModelEditor/Private/Widgets/SMVVMFieldEntry.cpp
--- a/Engine/Plugins/Runtime/ModelViewViewModel/Source/ModelViewViewModelEditor/Private/Widgets/SMVVMFieldEntry.cpp
+++ b/Engine/Plugins/Runtime/ModelViewViewModel/Source/ModelViewViewModelEditor/Private/Widgets/SMVVMFieldEntry.cpp
@@ -1,6 +1,8 @@
 // Copyright Epic Games, Inc. All Rights Reserved.
 
 #include "SMVVMFieldEntry.h"
+
+#include "Bindings/MVVMBindingHelper.h"
 #include "SMVVMFieldIcon.h"
 #include "Widgets/SBoxPanel.h"
 #include "Widgets/Images/SImage.h"
@@ -29,20 +31,89 @@ FText GetFieldDisplayName(const FMVVMConstFieldVariant& Field)
 	return LOCTEXT("None", "<None>");
 }
 
-FText GetFieldToolTip(const FMVVMConstFieldVariant& Field)
+FString GetFunctionSignature(const UFunction* Function)
 {
-	if (!Field.IsEmpty())
+	if (Function == nullptr)
+	{
+		return FString();
+	}
+
+	TStringBuilder<256> Signature;
+
+	const FProperty* ReturnProperty = BindingHelper::GetReturnProperty(Function);
+	if (ReturnProperty != nullptr)
+	{
+		Signature << ReturnProperty->GetCPPType();
+	}
+	else
+	{
+		Signature << TEXT("void");
+	}
+
+	Signature << TEXT(" ");
+	Signature << Function->GetName();
+	Signature << TEXT("(");
+
+	TArray<const FProperty*> Arguments = BindingHelper::GetAllArgumentProperties(Function);
+	Signature << FString::JoinBy(Arguments, TEXT(", "), 
+		[](const FProperty* Property)
+		{
+			return Property->GetCPPType();
+		}
+	);
+
+	Signature << TEXT(")");
+	return Signature.ToString();
+}
+
+FText GetPathToolTip(TConstArrayView<FMVVMConstFieldVariant> Path)
+{
+	if (Path.IsEmpty())
+	{
+		return FText::GetEmpty();
+	}
+
+	FMVVMConstFieldVariant LastField = Path.Last();
+	if (!LastField.IsEmpty())
 	{
-		if (Field.IsFunction() && Field.GetFunction() != nullptr)
+		FText LastToolTip, LastType;
+		if (LastField.IsFunction() && LastField.GetFunction() != nullptr)
+		{
+			LastToolTip = LastField.GetFunction()->GetToolTipText();
+			LastType = FText::FromString(GetFunctionSignature(LastField.GetFunction()));
+		}
+		else if (LastField.IsProperty() && LastField.GetProperty() != nullptr)
+		{
+			LastToolTip = LastField.GetProperty()->GetToolTipText();
+			LastType = FText::Format(LOCTEXT("PropertyFormat", "Property: {0} {1}"),
+				FText::FromString(LastField.GetProperty()->GetCPPType()),
+				FText::FromString(LastField.GetProperty()->GetName())
+			);
+		}
+
+		// only show the joined name if we've got more than 1 part to the path
+		if (Path.Num() > 1)
 		{
-			return Field.GetFunction()->GetToolTipText();
+			TArray<FText> DisplayNames;
+			for (const FMVVMConstFieldVariant& Field : Path)
+			{
+				DisplayNames.Add(GetFieldDisplayName(Field));
+			}
+
+			const FText JoinedDisplayName = FText::Join(FText::FromString(TEXT(".")), DisplayNames);
+
+			return FText::Join(FText::FromString(TEXT("\n")),
+				JoinedDisplayName,
+				LastToolTip,
+				LastType
+			);
 		}
-		if (Field.IsProperty() && Field.GetProperty() != nullptr)
+		else
 		{
-			return FText::Join(FText::FromString(TEXT("\n")), 
-				Field.GetProperty()->GetToolTipText(), 
-					FText::FromString(Field.GetProperty()->GetCPPType())
-				);
+			return FText::Join(FText::FromString(TEXT("\n")),
+				LastToolTip,
+				LastType
+			);
 		}
 	}
 
@@ -69,7 +140,14 @@ void SFieldEntry::Refresh()
 	FieldBox->ClearChildren();
 
 	TArray<FMVVMConstFieldVariant> Fields = Field.GetFields();
-	for (int32 Index = 0; Index < Fields.Num(); ++Index)
+
+	int32 Index = 0;
+	if (bShowOnlyLast && Fields.Num() > 0)
+	{
+		Index = Fields.Num() - 1;
+	}
+
+	for (; Index < Fields.Num(); ++Index)
 	{
 		FieldBox->AddSlot()
 		.HAlign(HAlign_Left)
@@ -107,10 +185,7 @@ void SFieldEntry::Refresh()
 		}
 	}
 
-	if (Fields.Num() > 0)
-	{
-		SetToolTipText(Private::GetFieldToolTip(Fields.Last()));
-	}
+	SetToolTipText(Private::GetPathToolTip(Fields));
 }
 
 void SFieldEntry::SetField(const FMVVMBlueprintPropertyPath& InField)
diff --git a/Engine/Plugins/Runtime/ModelViewViewModel/Source/ModelViewViewModelEditor/Private/Widgets/SMVVMFieldEntry.h b/Engine/Plugins/Runtime/ModelViewViewModel/Source/ModelViewViewModelEditor/Private/Widgets/SMVVMFieldEntry.h
--- a/Engine/Plugins/Runtime/ModelViewViewModel/Source/ModelViewViewModelEditor/Private/Widgets/SMVVMFieldEntry.h
+++ b/Engine/Plugins/Runtime/ModelViewViewModel/Source/ModelViewViewModelEditor/Private/Widgets/SMVVMFieldEntry.h
@@ -29,6 +29,7 @@ public:
 	{}
 		SLATE_STYLE_ARGUMENT(FTextBlockStyle, TextStyle)
 		SLATE_ARGUMENT(FMVVMBlueprintPropertyPath, Field)
+		SLATE_ARGUMENT(bool, ShowOnlyLast)
 	SLATE_END_ARGS()
 
 	void Construct(const FArguments& InArgs);
@@ -39,6 +40,7 @@ private:
 	const FTextBlockStyle* TextStyle = nullptr;
 	FMVVMBlueprintPropertyPath Field;
 	TSharedPtr<SHorizontalBox> FieldBox;
+	bool bShowOnlyLast = false;
 };
 
 } // namespace UE::MVVM
diff --git a/Engine/Plugins/Runtime/ModelViewViewModel/Source/ModelViewViewModelEditor/Private/Widgets/SMVVMFieldSelector.cpp b/Engine/Plugins/Runtime/ModelViewViewModel/Source/ModelViewViewModelEditor/Private/Widgets/SMVVMFieldSelector.cpp
--- a/Engine/Plugins/Runtime/ModelViewViewModel/Source/ModelViewViewModelEditor/Private/Widgets/SMVVMFieldSelector.cpp
+++ b/Engine/Plugins/Runtime/ModelViewViewModel/Source/ModelViewViewModelEditor/Private/Widgets/SMVVMFieldSelector.cpp
@@ -95,7 +95,14 @@ void SFieldSelector::Construct(const FArguments& InArgs, const UWidgetBlueprint*
 			+ SOverlay::Slot()
 			[
 				SNew(SHorizontalBox)
-				.Visibility_Lambda([this]() { return CachedSelectedConversionFunction != nullptr ? EVisibility::Visible : EVisibility::Collapsed; })
+				.Visibility_Lambda([this]() 
+				{
+					return CachedSelectedConversionFunction != nullptr ? EVisibility::Visible : EVisibility::Collapsed; 
+				})
+				.ToolTipText_Lambda([this]() 
+				{
+					return CachedSelectedConversionFunction != nullptr ? CachedSelectedConversionFunction->GetToolTipText() : FText::GetEmpty();
+				})
 				+ SHorizontalBox::Slot()
 				.HAlign(HAlign_Left)
 				.VAlign(VAlign_Center)
@@ -111,7 +118,10 @@ void SFieldSelector::Construct(const FArguments& InArgs, const UWidgetBlueprint*
 				.VAlign(VAlign_Center)
 				[
 					SNew(STextBlock)
-					.Text_Lambda([this]() { return CachedSelectedConversionFunction != nullptr ? CachedSelectedConversionFunction->GetDisplayNameText() : FText::GetEmpty(); })
+					.Text_Lambda([this]() 
+					{ 
+						return CachedSelectedConversionFunction != nullptr ? CachedSelectedConversionFunction->GetDisplayNameText() : FText::GetEmpty(); 
+					})
 				]
 			]
 
@@ -122,7 +132,10 @@ void SFieldSelector::Construct(const FArguments& InArgs, const UWidgetBlueprint*
 				.Padding(FMargin(8, 0, 8, 0))
 				[
 					SNew(STextBlock)
-					.Visibility_Lambda([this]() { return CachedSelectedField.IsEmpty() && CachedSelectedConversionFunction == nullptr ? EVisibility::Visible : EVisibility::Collapsed; })
+					.Visibility_Lambda([this]() 
+					{ 
+						return CachedSelectedField.IsEmpty() && CachedSelectedConversionFunction == nullptr ? EVisibility::Visible : EVisibility::Collapsed;
+					})
 					.TextStyle(FAppStyle::Get(), "HintText")
 					.Text(LOCTEXT("None", "No field selected"))
 				]
@@ -238,7 +251,10 @@ bool SFieldSelector::IsClearEnabled() const
 
 FReply SFieldSelector::OnClearClicked()
 {
-	SetPropertySelection(FMVVMBlueprintPropertyPath());
+	if (FixedSource.IsSet())
+	{
+		SetPropertySelection(FMVVMBlueprintPropertyPath());
+	}
 	SetConversionFunctionSelection(nullptr);
 	return FReply::Handled();
 }
