MVVM: Split MVVMViewModel's delegate into it's own struct to be reused by other plugins.

diff --git a/Engine/Plugins/Runtime/ModelViewViewModel/Source/ModelViewViewModel/ModelViewViewModel.Build.cs b/Engine/Plugins/Runtime/ModelViewViewModel/Source/ModelViewViewModel/ModelViewViewModel.Build.cs
--- a/Engine/Plugins/Runtime/ModelViewViewModel/Source/ModelViewViewModel/ModelViewViewModel.Build.cs
+++ b/Engine/Plugins/Runtime/ModelViewViewModel/Source/ModelViewViewModel/ModelViewViewModel.Build.cs
@@ -19,7 +19,6 @@ public class ModelViewViewModel : ModuleRules
 		PrivateDependencyModuleNames.AddRange(
 			new string[]
 			{
-				"NetCore",
 				"SlateCore",
 				"Slate",
 				"UMG",
diff --git a/Engine/Plugins/Runtime/ModelViewViewModel/Source/ModelViewViewModel/Private/MVVMViewModelBase.cpp b/Engine/Plugins/Runtime/ModelViewViewModel/Source/ModelViewViewModel/Private/MVVMViewModelBase.cpp
--- a/Engine/Plugins/Runtime/ModelViewViewModel/Source/ModelViewViewModel/Private/MVVMViewModelBase.cpp
+++ b/Engine/Plugins/Runtime/ModelViewViewModel/Source/ModelViewViewModel/Private/MVVMViewModelBase.cpp
@@ -11,17 +11,7 @@
 
 FDelegateHandle UMVVMViewModelBase::AddFieldValueChangedDelegate(UE::FieldNotification::FFieldId InFieldId, FFieldValueChangedDelegate InNewDelegate)
 {
-	FDelegateHandle Result;
-	if (InFieldId.IsValid())
-	{
-		Result = Delegates.Add(this, InFieldId, MoveTemp(InNewDelegate));
-		if (Result.IsValid())
-		{
-			EnabledFieldNotifications.PadToNum(InFieldId.GetIndex() + 1, false);
-			EnabledFieldNotifications[InFieldId.GetIndex()] = true;
-		}
-	}
-	return Result;
+	return NotificationDelegates.AddFieldValueChangedDelegate(this, InFieldId, MoveTemp(InNewDelegate));
 }
 
 
@@ -30,29 +20,15 @@ void UMVVMViewModelBase::K2_AddFieldValueChangedDelegate(FFieldNotificationId In
 	if (InFieldId.IsValid())
 	{
 		const UE::FieldNotification::FFieldId FieldId = GetFieldNotificationDescriptor().GetField(GetClass(), InFieldId.FieldName);
-		if (ensureMsgf(FieldId.IsValid(), TEXT("The field should be compiled correctly.")))
-		{
-			FDelegateHandle Result = Delegates.Add(this, FieldId, InDelegate);
-			if (Result.IsValid())
-			{
-				EnabledFieldNotifications.PadToNum(FieldId.GetIndex() + 1, false);
-				EnabledFieldNotifications[FieldId.GetIndex()] = true;
-			}
-		}
+		ensureMsgf(FieldId.IsValid(), TEXT("The field should be compiled correctly."));
+		NotificationDelegates.AddFieldValueChangedDelegate(this, FieldId, InDelegate);
 	}
 }
 
 
 bool UMVVMViewModelBase::RemoveFieldValueChangedDelegate(UE::FieldNotification::FFieldId InFieldId, FDelegateHandle InHandle)
 {
-	bool bResult = false;
-	if (InFieldId.IsValid() && InHandle.IsValid()  && EnabledFieldNotifications.IsValidIndex(InFieldId.GetIndex()) && EnabledFieldNotifications[InFieldId.GetIndex()])
-	{
-		UE::FieldNotification::FFieldMulticastDelegate::FRemoveFromResult RemoveResult = Delegates.RemoveFrom(this, InFieldId, InHandle);
-		bResult = RemoveResult.bRemoved;
-		EnabledFieldNotifications[InFieldId.GetIndex()] = RemoveResult.bHasOtherBoundDelegates;
-	}
-	return bResult;
+	return NotificationDelegates.RemoveFieldValueChangedDelegate(this, InFieldId, InHandle);
 }
 
 
@@ -61,41 +37,21 @@ void UMVVMViewModelBase::K2_RemoveFieldValueChangedDelegate(FFieldNotificationId
 	if (InFieldId.IsValid())
 	{
 		const UE::FieldNotification::FFieldId FieldId = GetFieldNotificationDescriptor().GetField(GetClass(), InFieldId.FieldName);
-		if (ensureMsgf(FieldId.IsValid(), TEXT("The field should be compiled correctly.")))
-		{
-			if (EnabledFieldNotifications.IsValidIndex(FieldId.GetIndex()) && EnabledFieldNotifications[FieldId.GetIndex()])
-			{
-				UE::FieldNotification::FFieldMulticastDelegate::FRemoveFromResult RemoveResult = Delegates.RemoveFrom(this, FieldId, InDelegate);
-				EnabledFieldNotifications[FieldId.GetIndex()] = RemoveResult.bHasOtherBoundDelegates;
-			}
-		}
+		ensureMsgf(FieldId.IsValid(), TEXT("The field should be compiled correctly."));
+		NotificationDelegates.RemoveFieldValueChangedDelegate(this, FieldId, MoveTemp(InDelegate));
 	}
 }
 
 
 int32 UMVVMViewModelBase::RemoveAllFieldValueChangedDelegates(const void* InUserObject)
 {
-	int32 bResult = 0;
-	if (InUserObject)
-	{
-		UE::FieldNotification::FFieldMulticastDelegate::FRemoveAllResult RemoveResult = Delegates.RemoveAll(this, InUserObject);
-		bResult = RemoveResult.RemoveCount;
-		EnabledFieldNotifications = RemoveResult.HasFields;
-	}
-	return bResult;
+	return NotificationDelegates.RemoveAllFieldValueChangedDelegates(this, InUserObject);
 }
 
 
 int32 UMVVMViewModelBase::RemoveAllFieldValueChangedDelegates(UE::FieldNotification::FFieldId InFieldId, const void* InUserObject)
 {
-	int32 bResult = 0;
-	if (InUserObject)
-	{
-		UE::FieldNotification::FFieldMulticastDelegate::FRemoveAllResult RemoveResult = Delegates.RemoveAll(this, InFieldId, InUserObject);
-		bResult = RemoveResult.RemoveCount;
-		EnabledFieldNotifications = RemoveResult.HasFields;
-	}
-	return bResult;
+	return NotificationDelegates.RemoveAllFieldValueChangedDelegates(this, InFieldId, InUserObject);
 }
 
 
@@ -115,43 +71,16 @@ void UMVVMViewModelBase::FFieldNotificationClassDescriptor::ForEachField(const U
 }
 
 
-void UMVVMViewModelBase::MarkRepDirty(int32 InRepIndex)
-{
-#if WITH_PUSH_MODEL
-	check(InRepIndex >= 0 && InRepIndex < GetClass()->ClassReps.Num());
-	if (IS_PUSH_MODEL_ENABLED())
-	{
-		MARK_PROPERTY_DIRTY_UNSAFE(this, InRepIndex);
-	}
-#endif
-}
-
-
-void UMVVMViewModelBase::MarkRepDirty(FProperty* InProperty)
-{
-#if WITH_PUSH_MODEL
-	check(InProperty);
-	if (IS_PUSH_MODEL_ENABLED() && InProperty->RepIndex >= 0)
-	{
-		MARK_PROPERTY_DIRTY(this, InProperty);
-	}
-#endif
-}
-
-
 void UMVVMViewModelBase::BroadcastFieldValueChanged(UE::FieldNotification::FFieldId InFieldId)
 {
-	if (InFieldId.IsValid() && EnabledFieldNotifications.IsValidIndex(InFieldId.GetIndex()) && EnabledFieldNotifications[InFieldId.GetIndex()])
-	{
-		Delegates.Broadcast(this, InFieldId);
-	}
+	NotificationDelegates.BroadcastFieldValueChanged(this, InFieldId);
 }
 
 
 void UMVVMViewModelBase::K2_BroadcastFieldValueChanged(FFieldNotificationId ViewModelPropertyName)
 {
 	UE::FieldNotification::FFieldId FieldId = GetFieldNotificationDescriptor().GetField(GetClass(), ViewModelPropertyName.GetFieldName());
-	BroadcastFieldValueChanged(FieldId);
+	NotificationDelegates.BroadcastFieldValueChanged(this, FieldId);
 }
 
 
@@ -198,7 +127,6 @@ DEFINE_FUNCTION(UMVVMViewModelBase::execK2_SetPropertyValue)
 			{
 				// Set the value then notify that the value changed.
 				TargetProperty->SetValue_InContainer(ViewModelContext, SourceValuePtr);
-				ViewModelContext->MarkRepDirty(TargetProperty);
 				ViewModelContext->BroadcastFieldValueChanged(FieldId);
 			}
 		}
diff --git a/Engine/Plugins/Runtime/ModelViewViewModel/Source/ModelViewViewModel/Private/View/MVVMView.cpp b/Engine/Plugins/Runtime/ModelViewViewModel/Source/ModelViewViewModel/Private/View/MVVMView.cpp
--- a/Engine/Plugins/Runtime/ModelViewViewModel/Source/ModelViewViewModel/Private/View/MVVMView.cpp
+++ b/Engine/Plugins/Runtime/ModelViewViewModel/Source/ModelViewViewModel/Private/View/MVVMView.cpp
@@ -70,7 +70,7 @@ void UMVVMView::Destruct()
 }
 
 
-bool UMVVMView::SetViewModel(FName ViewModelName, UMVVMViewModelBase* NewValue)
+bool UMVVMView::SetViewModel(FName ViewModelName, TScriptInterface<INotifyFieldValueChanged> NewValue)
 {
 	if (!ViewModelName.IsNone() && ClassExtension != nullptr)
 	{
@@ -89,7 +89,7 @@ bool UMVVMView::SetViewModel(FName ViewModelName, UMVVMViewModelBase* NewValue)
 				if (Item.GetSourcePropertyName() == ViewModelName)
 				{
 					bFound = true;
-					if (NewValue && !NewValue->GetClass()->IsChildOf(Item.GetSourceClass()))
+					if (NewValue.GetObject() && !NewValue.GetObject()->GetClass()->IsChildOf(Item.GetSourceClass()))
 					{
 						UE::MVVM::FMessageLog Log(GetUserWidget());
 						Log.Error(LOCTEXT("SetViewModelInvalidValueType", "The new viewmodel is not of the expected type."));
@@ -109,7 +109,7 @@ bool UMVVMView::SetViewModel(FName ViewModelName, UMVVMViewModelBase* NewValue)
 
 
 		UObject* PreviousValue = FoundObjectProperty->GetObjectPropertyValue_InContainer(GetUserWidget());
-		if (PreviousValue != NewValue)
+		if (PreviousValue != NewValue.GetObject())
 		{
 			FMemMark Mark(FMemStack::Get());
 			TArray<int32, TMemStackAllocator<>> BindingToReenabled;
@@ -131,9 +131,9 @@ bool UMVVMView::SetViewModel(FName ViewModelName, UMVVMViewModelBase* NewValue)
 				}
 			}
 
-			FoundObjectProperty->SetObjectPropertyValue_InContainer(GetUserWidget(), NewValue);
+			FoundObjectProperty->SetObjectPropertyValue_InContainer(GetUserWidget(), NewValue.GetObject());
 
-			if (NewValue)
+			if (NewValue.GetObject())
 			{
 				// Register back any binding that was previously enabled
 				if (PreviousValue)
diff --git a/Engine/Plugins/Runtime/ModelViewViewModel/Source/ModelViewViewModel/Private/ViewModel/MVVMFieldNotificationDelegates.cpp b/Engine/Plugins/Runtime/ModelViewViewModel/Source/ModelViewViewModel/Private/ViewModel/MVVMFieldNotificationDelegates.cpp
new file mode 100644
--- /dev/null
+++ b/Engine/Plugins/Runtime/ModelViewViewModel/Source/ModelViewViewModel/Private/ViewModel/MVVMFieldNotificationDelegates.cpp
@@ -0,0 +1,104 @@
+// Copyright Epic Games, Inc. All Rights Reserved.
+
+#include "ViewModel/MVVMFieldNotificationDelegates.h"
+
+#include "Bindings/MVVMBindingHelper.h"
+
+#define LOCTEXT_NAMESPACE "MVVMFieldNotificationDelegates"
+
+namespace UE::MVVM {
+
+FDelegateHandle FMVVMFieldNotificationDelegates::AddFieldValueChangedDelegate(UObject* InOwner, UE::FieldNotification::FFieldId InFieldId, INotifyFieldValueChanged::FFieldValueChangedDelegate InNewDelegate)
+{
+	FDelegateHandle Result;
+	if (InFieldId.IsValid())
+	{
+		Result = Delegates.Add(InOwner, InFieldId, MoveTemp(InNewDelegate));
+		if (Result.IsValid())
+		{
+			EnabledFieldNotifications.PadToNum(InFieldId.GetIndex() + 1, false);
+			EnabledFieldNotifications[InFieldId.GetIndex()] = true;
+		}
+	}
+	return Result;
+}
+
+
+void FMVVMFieldNotificationDelegates::AddFieldValueChangedDelegate(UObject* InOwner, const UE::FieldNotification::FFieldId InFieldId, const FFieldValueChangedDynamicDelegate& InDelegate)
+{
+	if (InFieldId.IsValid())
+	{
+		FDelegateHandle Result = Delegates.Add(InOwner, InFieldId, InDelegate);
+		if (Result.IsValid())
+		{
+			EnabledFieldNotifications.PadToNum(InFieldId.GetIndex() + 1, false);
+			EnabledFieldNotifications[InFieldId.GetIndex()] = true;
+		}
+	}
+}
+
+
+bool FMVVMFieldNotificationDelegates::RemoveFieldValueChangedDelegate(UObject* InOwner, UE::FieldNotification::FFieldId InFieldId, FDelegateHandle InHandle)
+{
+	bool bResult = false;
+	if (InFieldId.IsValid() && InHandle.IsValid() && EnabledFieldNotifications.IsValidIndex(InFieldId.GetIndex()) && EnabledFieldNotifications[InFieldId.GetIndex()])
+	{
+		UE::FieldNotification::FFieldMulticastDelegate::FRemoveFromResult RemoveResult = Delegates.RemoveFrom(InOwner, InFieldId, InHandle);
+		bResult = RemoveResult.bRemoved;
+		EnabledFieldNotifications[InFieldId.GetIndex()] = RemoveResult.bHasOtherBoundDelegates;
+	}
+	return bResult;
+}
+
+
+bool FMVVMFieldNotificationDelegates::RemoveFieldValueChangedDelegate(UObject* InOwner, const UE::FieldNotification::FFieldId InFieldId, const FFieldValueChangedDynamicDelegate& InDelegate)
+{
+	bool bResult = false;
+	if (InFieldId.IsValid() && EnabledFieldNotifications.IsValidIndex(InFieldId.GetIndex()) && EnabledFieldNotifications[InFieldId.GetIndex()])
+	{
+		UE::FieldNotification::FFieldMulticastDelegate::FRemoveFromResult RemoveResult = Delegates.RemoveFrom(InOwner, InFieldId, InDelegate);
+		bResult = RemoveResult.bRemoved;
+		EnabledFieldNotifications[InFieldId.GetIndex()] = RemoveResult.bHasOtherBoundDelegates;
+	}
+	return bResult;
+}
+
+
+int32 FMVVMFieldNotificationDelegates::RemoveAllFieldValueChangedDelegates(UObject* InOwner, const void* InUserObject)
+{
+	int32 bResult = 0;
+	if (InUserObject)
+	{
+		UE::FieldNotification::FFieldMulticastDelegate::FRemoveAllResult RemoveResult = Delegates.RemoveAll(InOwner, InUserObject);
+		bResult = RemoveResult.RemoveCount;
+		EnabledFieldNotifications = RemoveResult.HasFields;
+	}
+	return bResult;
+}
+
+
+int32 FMVVMFieldNotificationDelegates::RemoveAllFieldValueChangedDelegates(UObject* InOwner, UE::FieldNotification::FFieldId InFieldId, const void* InUserObject)
+{
+	int32 bResult = 0;
+	if (InFieldId.IsValid() && InUserObject)
+	{
+		UE::FieldNotification::FFieldMulticastDelegate::FRemoveAllResult RemoveResult = Delegates.RemoveAll(InOwner, InFieldId, InUserObject);
+		bResult = RemoveResult.RemoveCount;
+		EnabledFieldNotifications = RemoveResult.HasFields;
+	}
+	return bResult;
+}
+
+
+void FMVVMFieldNotificationDelegates::BroadcastFieldValueChanged(UObject* InOwner, UE::FieldNotification::FFieldId InFieldId)
+{
+	if (InFieldId.IsValid() && EnabledFieldNotifications.IsValidIndex(InFieldId.GetIndex()) && EnabledFieldNotifications[InFieldId.GetIndex()])
+	{
+		Delegates.Broadcast(InOwner, InFieldId);
+	}
+}
+
+} //namespace
+
+#undef LOCTEXT_NAMESPACE
+
diff --git a/Engine/Plugins/Runtime/ModelViewViewModel/Source/ModelViewViewModel/Public/MVVMViewModelBase.h b/Engine/Plugins/Runtime/ModelViewViewModel/Source/ModelViewViewModel/Public/MVVMViewModelBase.h
--- a/Engine/Plugins/Runtime/ModelViewViewModel/Source/ModelViewViewModel/Public/MVVMViewModelBase.h
+++ b/Engine/Plugins/Runtime/ModelViewViewModel/Source/ModelViewViewModel/Public/MVVMViewModelBase.h
@@ -9,28 +9,19 @@
 #include "FieldNotification/FieldNotificationDeclaration.h"
 #include "FieldNotification/FieldMulticastDelegate.h"
 #include "FieldNotification/IFieldValueChanged.h"
-#include "Net/Core/PushModel/PushModel.h"
 #include "Types/MVVMAvailableBinding.h"
 #include "Types/MVVMBindingName.h"
+#include "ViewModel/MVVMFieldNotificationDelegates.h"
 
 #include "MVVMViewModelBase.generated.h"
 
-/** After a field value changed. Broadcast the event (doesn't execute the replication code). */
+/** After a field value changed. Broadcast the event. */
 #define UE_MVVM_BROADCAST_FIELD_VALUE_CHANGED(MemberName) \
 	BroadcastFieldValueChanged(ThisClass::FFieldNotificationClassDescriptor::MemberName)
 
-/** After a field value changed. Replicate and broadcast the event. */
-#define UE_MVVM_BROADCAST_FIELD_VALUE_CHANGED_WITH_REP(MemberName ) \
-	MARK_PROPERTY_DIRTY_FROM_NAME(ThisClass, this, MemberName); \
-	BroadcastFieldValueChanged(ThisClass::FFieldNotificationClassDescriptor::MemberName)
-
 /** If the property value changed then set the new value and notify. */
 #define UE_MVVM_SET_PROPERTY_VALUE(MemberName, NewValue) \
 	SetPropertyValue(MemberName, NewValue, ThisClass::FFieldNotificationClassDescriptor::MemberName)
-	
-/** If the property value changed then set the new value, replicate and notify. */
-#define UE_MVVM_SET_PROPERTY_VALUE_WITH_REP(MemberName, NewValue) \
-	SetPropertyValue(MemberName, NewValue, ThisClass::FFieldNotificationClassDescriptor::MemberName, ThisClass::ENetFields_Private::MemberName)
 
 
 /** Base class for MVVM viewmodel. */
@@ -61,13 +52,7 @@ public:
 	void K2_RemoveFieldValueChangedDelegate(FFieldNotificationId FieldId, FFieldValueChangedDynamicDelegate Delegate);
 
 protected:
-	/** Execute the replication code. */
-	void MarkRepDirty(FProperty* InProperty);
-private:
-	void MarkRepDirty(int32 InRepIndex);
-
-protected:
-	UFUNCTION(BlueprintCallable, Category="FieldNotify", meta=(DisplayName="Broadcast Field Value Changed", ScriptName="BroadcastFieldValueChanged", BlueprintInternalUseOnly="true"))
+	UFUNCTION(BlueprintCallable, Category="FieldNotify", meta=(DisplayName="Broadcast Field Value Changed", ScriptName="BroadcastFieldValueChanged"))
 	void K2_BroadcastFieldValueChanged(FFieldNotificationId FieldId);
 
 	/** Broadcast the event (doesn't execute the replication code). */
@@ -91,25 +76,9 @@ protected:
 		return true;
 	}
 
-	/** Set the new value and notify if the property value changed. */
-	template<typename T, typename U, typename NetFieldsEnum>
-	bool SetPropertyValue(T& Value, const U& NewValue, UE::FieldNotification::FFieldId FieldId, NetFieldsEnum InRepIndex)
-	{
-		if (Value == NewValue)
-		{
-			return false;
-		}
-
-		Value = NewValue;
-		MarkRepDirty((int32)InRepIndex);
-		BroadcastFieldValueChanged(FieldId);
-		return true;
-	}
-
 private:
 	DECLARE_FUNCTION(execK2_SetPropertyValue);
 
 private:
-	UE::FieldNotification::FFieldMulticastDelegate Delegates;
-	TBitArray<> EnabledFieldNotifications;
+	UE::MVVM::FMVVMFieldNotificationDelegates NotificationDelegates;
 };
diff --git a/Engine/Plugins/Runtime/ModelViewViewModel/Source/ModelViewViewModel/Public/View/MVVMView.h b/Engine/Plugins/Runtime/ModelViewViewModel/Source/ModelViewViewModel/Public/View/MVVMView.h
--- a/Engine/Plugins/Runtime/ModelViewViewModel/Source/ModelViewViewModel/Public/View/MVVMView.h
+++ b/Engine/Plugins/Runtime/ModelViewViewModel/Source/ModelViewViewModel/Public/View/MVVMView.h
@@ -47,7 +47,7 @@ public:
 	//bool IsLibraryBindingEnabled(FGuid ViewModelId, FMVVMBindingName BindingName) const;
 
 	UFUNCTION(BlueprintCallable, Category = "MVVM")
-	bool SetViewModel(FName ViewModelName, UMVVMViewModelBase* ViewModel);
+	bool SetViewModel(FName ViewModelName, TScriptInterface<INotifyFieldValueChanged> ViewModel);
 
 private:
 	void HandledLibraryBindingValueChanged(UObject* InViewModel, UE::FieldNotification::FFieldId InFieldId, int32 InCompiledBindingIndex) const;
diff --git a/Engine/Plugins/Runtime/ModelViewViewModel/Source/ModelViewViewModel/Public/ViewModel/MVVMFieldNotificationDelegates.h b/Engine/Plugins/Runtime/ModelViewViewModel/Source/ModelViewViewModel/Public/ViewModel/MVVMFieldNotificationDelegates.h
new file mode 100644
--- /dev/null
+++ b/Engine/Plugins/Runtime/ModelViewViewModel/Source/ModelViewViewModel/Public/ViewModel/MVVMFieldNotificationDelegates.h
@@ -0,0 +1,31 @@
+// Copyright Epic Games, Inc. All Rights Reserved.
+
+#pragma once
+
+#include "CoreMinimal.h"
+#include "FieldNotification/FieldId.h"
+#include "FieldNotification/FieldMulticastDelegate.h"
+#include "FieldNotification/IFieldValueChanged.h"
+#include "Net/Core/PushModel/PushModel.h"
+
+namespace UE::MVVM
+{
+
+/** Basic implementation of the INotifyFieldValueChanged implementation. */
+class MODELVIEWVIEWMODEL_API FMVVMFieldNotificationDelegates
+{
+public:
+	FDelegateHandle AddFieldValueChangedDelegate(UObject* Owner, UE::FieldNotification::FFieldId InFieldId, INotifyFieldValueChanged::FFieldValueChangedDelegate InNewDelegate);
+	void AddFieldValueChangedDelegate(UObject* Owner, UE::FieldNotification::FFieldId FieldId, const FFieldValueChangedDynamicDelegate& Delegate);
+	bool RemoveFieldValueChangedDelegate(UObject* Owner, UE::FieldNotification::FFieldId InFieldId, FDelegateHandle InHandle);
+	bool RemoveFieldValueChangedDelegate(UObject* Owner, UE::FieldNotification::FFieldId FieldId, const FFieldValueChangedDynamicDelegate& Delegate);
+	int32 RemoveAllFieldValueChangedDelegates(UObject* Owner, const void* InUserObject);
+	int32 RemoveAllFieldValueChangedDelegates(UObject* Owner, UE::FieldNotification::FFieldId InFieldId, const void* InUserObject);
+	void BroadcastFieldValueChanged(UObject* Owner, UE::FieldNotification::FFieldId InFieldId);
+
+private:
+	UE::FieldNotification::FFieldMulticastDelegate Delegates;
+	TBitArray<> EnabledFieldNotifications;
+};
+
+} //namespace
\ No newline at end of file
