MVVM - Added missing transactions.

diff --git a/Engine/Plugins/Runtime/ModelViewViewModel/Source/ModelViewViewModelEditor/Private/MVVMEditorSubsystem.cpp b/Engine/Plugins/Runtime/ModelViewViewModel/Source/ModelViewViewModelEditor/Private/MVVMEditorSubsystem.cpp
--- a/Engine/Plugins/Runtime/ModelViewViewModel/Source/ModelViewViewModelEditor/Private/MVVMEditorSubsystem.cpp
+++ b/Engine/Plugins/Runtime/ModelViewViewModel/Source/ModelViewViewModelEditor/Private/MVVMEditorSubsystem.cpp
@@ -263,7 +263,7 @@ bool UMVVMEditorSubsystem::RenameViewModel(UWidgetBlueprint* WidgetBlueprint, FN
 		return false;
 	}
 
-	const FScopedTransaction Transaction(LOCTEXT("RenameViewModel", "Rename viewmodel"));
+	const FScopedTransaction Transaction(LOCTEXT("RenameViewModel", "Rename Viewmodel"));
 	View->Modify();
 	return View->RenameViewModel(ViewModel, NewViewModel);
 }
@@ -532,6 +532,8 @@ void UMVVMEditorSubsystem::SetEnabledForBinding(UWidgetBlueprint* WidgetBlueprin
 	{
 		if (UMVVMBlueprintView* View = GetView(WidgetBlueprint))
 		{
+			FScopedTransaction Transaction(LOCTEXT("SetBindingEnabled", "Set Binding Enabled"));
+
 			UE::MVVM::Private::OnBindingPreEditChange(View, GET_MEMBER_NAME_CHECKED(FMVVMBlueprintViewBinding, bEnabled));
 
 			Binding.bEnabled = bEnabled;
@@ -547,6 +549,8 @@ void UMVVMEditorSubsystem::SetCompileForBinding(UWidgetBlueprint* WidgetBlueprin
 	{
 		if (UMVVMBlueprintView* View = GetView(WidgetBlueprint))
 		{
+			FScopedTransaction Transaction(LOCTEXT("SetBindingCompiled", "Set Binding Compiled"));
+
 			UE::MVVM::Private::OnBindingPreEditChange(View, GET_MEMBER_NAME_CHECKED(FMVVMBlueprintViewBinding, bCompile));
 
 			Binding.bCompile = bCompile;
