MVVM - Widgets without an explicit display name override are now displayed properly in the compile error log.

#rnx

diff --git a/Engine/Plugins/Runtime/ModelViewViewModel/Source/ModelViewViewModelBlueprint/Private/MVVMViewBlueprintCompiler.cpp b/Engine/Plugins/Runtime/ModelViewViewModel/Source/ModelViewViewModelBlueprint/Private/MVVMViewBlueprintCompiler.cpp
--- a/Engine/Plugins/Runtime/ModelViewViewModel/Source/ModelViewViewModelBlueprint/Private/MVVMViewBlueprintCompiler.cpp
+++ b/Engine/Plugins/Runtime/ModelViewViewModel/Source/ModelViewViewModelBlueprint/Private/MVVMViewBlueprintCompiler.cpp
@@ -422,7 +422,7 @@ void FMVVMViewBlueprintCompiler::CreateSourceLists(const FWidgetBlueprintCompile
 					FCompilerSourceContext SourceVariable;
 					SourceVariable.Class = Widget->GetClass();
 					SourceVariable.PropertyName = PropertyPath.GetWidgetName();
-					SourceVariable.DisplayName = FText::FromString(Widget->GetDisplayLabel());
+					SourceVariable.DisplayName = Widget->GetLabelText();
 					SourceVariable.CategoryName = !Widget->GetCategoryName().IsEmpty() ? Widget->GetCategoryName() : DefaultWidgetCategory.ToString();
 					Self->CompilerSourceContexts.Emplace(MoveTemp(SourceVariable));
 				}
