MVVM: Execute Library Binding

diff --git a/Engine/Plugins/Runtime/ModelViewViewModel/Source/ModelViewViewModel/Private/ModelViewViewModelModule.cpp b/Engine/Plugins/Runtime/ModelViewViewModel/Source/ModelViewViewModel/Private/ModelViewViewModelModule.cpp
--- a/Engine/Plugins/Runtime/ModelViewViewModel/Source/ModelViewViewModel/Private/ModelViewViewModelModule.cpp
+++ b/Engine/Plugins/Runtime/ModelViewViewModel/Source/ModelViewViewModel/Private/ModelViewViewModelModule.cpp
@@ -1,9 +1,38 @@
 // Copyright Epic Games, Inc. All Rights Reserved.
 
 #include "ModelViewViewModelModule.h"
+#include "HAL/IConsoleManager.h"
+#include "Types/MVVMExecutionMode.h"
 
 #define LOCTEXT_NAMESPACE "ModelViewViewModelModule"
 
+void FModelViewViewModelModule::StartupModule()
+{
+	if (IConsoleVariable* CVarDefaultExecutionMode = IConsoleManager::Get().FindConsoleVariable(TEXT("MVVM.DefaultExecutionMode")))
+	{
+		CVarDefaultExecutionMode->OnChangedDelegate().AddRaw(this, &FModelViewViewModelModule::HandleDefaultExecutionModeChanged);
+	}
+}
+
+void FModelViewViewModelModule::ShutdownModule()
+{
+	if (!IsEngineExitRequested())
+	{
+		if (IConsoleVariable* CVarDefaultExecutionMode = IConsoleManager::Get().FindConsoleVariable(TEXT("MVVM.DefaultExecutionMode")))
+		{
+			CVarDefaultExecutionMode->OnChangedDelegate().RemoveAll(this);
+		}
+	}
+}
+
+void FModelViewViewModelModule::HandleDefaultExecutionModeChanged(IConsoleVariable* Variable)
+{
+	EMVVMExecutionMode Value = (EMVVMExecutionMode)Variable->GetInt();
+	if (Value != EMVVMExecutionMode::Delayed && Value != EMVVMExecutionMode::Immediate && Value != EMVVMExecutionMode::Tick)
+	{
+		Variable->Set((int32)EMVVMExecutionMode::Immediate, (EConsoleVariableFlags)(Variable->GetFlags() & ECVF_SetByMask));
+	}
+}
 
 IMPLEMENT_MODULE(FModelViewViewModelModule, ModelViewViewModel);
 
diff --git a/Engine/Plugins/Runtime/ModelViewViewModel/Source/ModelViewViewModel/Private/ModelViewViewModelModule.h b/Engine/Plugins/Runtime/ModelViewViewModel/Source/ModelViewViewModel/Private/ModelViewViewModelModule.h
--- a/Engine/Plugins/Runtime/ModelViewViewModel/Source/ModelViewViewModel/Private/ModelViewViewModelModule.h
+++ b/Engine/Plugins/Runtime/ModelViewViewModel/Source/ModelViewViewModel/Private/ModelViewViewModelModule.h
@@ -6,6 +6,8 @@
 #include "Modules/ModuleInterface.h"
 #include "Modules/ModuleManager.h"
 
+class IConsoleVariable;
+
 /**
  *
  */
@@ -15,7 +17,10 @@ public:
 	FModelViewViewModelModule() = default;
 
 	//~ Begin IModuleInterface interface
-	// virtual void StartupModule() override;
-	// virtual void ShutdownModule() override;
+	virtual void StartupModule() override;
+	virtual void ShutdownModule() override;
 	//~ End IModuleInterface interface
+
+private:
+	void HandleDefaultExecutionModeChanged(IConsoleVariable* Variable);
 };
diff --git a/Engine/Plugins/Runtime/ModelViewViewModel/Source/ModelViewViewModel/Private/View/MVVMView.cpp b/Engine/Plugins/Runtime/ModelViewViewModel/Source/ModelViewViewModel/Private/View/MVVMView.cpp
--- a/Engine/Plugins/Runtime/ModelViewViewModel/Source/ModelViewViewModel/Private/View/MVVMView.cpp
+++ b/Engine/Plugins/Runtime/ModelViewViewModel/Source/ModelViewViewModel/Private/View/MVVMView.cpp
@@ -2,12 +2,13 @@
 
 #include "View/MVVMView.h"
 #include "View/MVVMViewClass.h"
+#include "View/MVVMViewWorldSubsystem.h"
 
 #include "Blueprint/UserWidget.h"
 #include "Blueprint/WidgetTree.h"
+#include "Engine/World.h"
 #include "MVVMMessageLog.h"
 #include "MVVMViewModelBase.h"
-#include "MVVMSubsystem.h"
 
 #include UE_INLINE_GENERATED_CPP_BY_NAME(MVVMView)
 
@@ -36,6 +37,8 @@ void UMVVMView::Construct()
 		Item.CreateInstance(ClassExtension, this, GetUserWidget());
 	}
 
+	bHasEveryTickBinding = false;
+
 	const TArrayView<const FMVVMViewClass_CompiledBinding> CompiledBindings = ClassExtension->GetCompiledBindings();
 	for (int32 Index = 0; Index < CompiledBindings.Num(); ++Index)
 	{
@@ -47,6 +50,11 @@ void UMVVMView::Construct()
 	}
 
 	bConstructed = true;
+
+	if (bHasEveryTickBinding)
+	{
+		GetUserWidget()->GetWorld()->GetSubsystem<UMVVMViewWorldSubsystem>()->AddViewWithEveryTickBinding(this);
+	}
 }
 
 
@@ -67,6 +75,12 @@ void UMVVMView::Destruct()
 	}
 	AllSources.Reset();
 	EnabledLibraryBindings.Reset();
+
+	if (bHasEveryTickBinding)
+	{
+		GetUserWidget()->GetWorld()->GetSubsystem<UMVVMViewWorldSubsystem>()->RemoveViewWithEveryTickBinding(this);
+	}
+	bHasEveryTickBinding = false;
 }
 
 
@@ -133,6 +147,8 @@ bool UMVVMView::SetViewModel(FName ViewModelName, TScriptInterface<INotifyFieldV
 
 			FoundObjectProperty->SetObjectPropertyValue_InContainer(GetUserWidget(), NewValue.GetObject());
 
+			bool bPreviousEveryTickBinding = bHasEveryTickBinding;
+			bHasEveryTickBinding = false;
 			if (NewValue.GetObject())
 			{
 				// Register back any binding that was previously enabled
@@ -157,6 +173,18 @@ bool UMVVMView::SetViewModel(FName ViewModelName, TScriptInterface<INotifyFieldV
 					}
 				}
 			}
+
+			if (bPreviousEveryTickBinding != bHasEveryTickBinding)
+			{
+				if (bHasEveryTickBinding)
+				{
+					GetUserWidget()->GetWorld()->GetSubsystem<UMVVMViewWorldSubsystem>()->AddViewWithEveryTickBinding(this);
+				}
+				else
+				{
+					GetUserWidget()->GetWorld()->GetSubsystem<UMVVMViewWorldSubsystem>()->RemoveViewWithEveryTickBinding(this);
+				}
+			}
 		}
 		return true;
 	}
@@ -166,7 +194,7 @@ bool UMVVMView::SetViewModel(FName ViewModelName, TScriptInterface<INotifyFieldV
 
 namespace UE::MVVM::Private
 {
-	using FRecursiveDetctionElement = TTuple<UObject*, int32>;
+	using FRecursiveDetctionElement = TTuple<const UObject*, int32>;
 	TArray<FRecursiveDetctionElement> RecursiveDetector;
 }
 
@@ -180,25 +208,96 @@ void UMVVMView::HandledLibraryBindingValueChanged(UObject* InViewModelOrWidget,
 		checkf(ClassExtension->GetCompiledBindings().IsValidIndex(InCompiledBindingIndex), TEXT("The binding at index '%d' does not exist. The binding was probably not cleared on destroyed."), InCompiledBindingIndex);
 		const FMVVMViewClass_CompiledBinding& Binding = ClassExtension->GetCompiledBinding(InCompiledBindingIndex);
 
-		// Test for recursivity
-		if (UE::MVVM::Private::RecursiveDetector.FindByPredicate([InViewModelOrWidget, InCompiledBindingIndex](const UE::MVVM::Private::FRecursiveDetctionElement& Element)
+		EMVVMExecutionMode ExecutionMode = Binding.GetExecuteMode();
+		if (ExecutionMode == EMVVMExecutionMode::Delayed)
+		{
+			GetUserWidget()->GetWorld()->GetSubsystem<UMVVMViewWorldSubsystem>()->AddDelayedBinding(this, InCompiledBindingIndex);
+		}
+		else if (ExecutionMode != EMVVMExecutionMode::Tick)
+		{
+			// Test for recursivity
+			const UMVVMView* Self = this;
+			if (UE::MVVM::Private::RecursiveDetector.FindByPredicate([Self, InCompiledBindingIndex](const UE::MVVM::Private::FRecursiveDetctionElement& Element)
+				{
+					return Element.Get<0>() == Self && Element.Get<1>() == InCompiledBindingIndex;
+				}) != nullptr)
+			{
+				ensureAlwaysMsgf(false, TEXT("Recursive binding detected"));
+				//Todo add more infos. Callstack maybe? Log the chain?
+				UE::MVVM::FMessageLog Log(GetUserWidget());
+				Log.Error(LOCTEXT("RecursionDetected", "A recursive binding was detected (ie. A->B->C->A->B->C) at runtime."));
+				return;
+			}
+
 			{
-				return Element.Get<0>() == InViewModelOrWidget && Element.Get<1>() == InCompiledBindingIndex;
-			}) != nullptr)
+				UE::MVVM::Private::RecursiveDetector.Emplace(this, InCompiledBindingIndex);
+
+				ExecuteLibraryBinding(Binding, InViewModelOrWidget);
+
+				UE::MVVM::Private::RecursiveDetector.Pop();
+			}
+		}
+		else
 		{
-			ensureAlwaysMsgf(false, TEXT("Recursive binding detected"));
-			//Todo add more infos. Callstack maybe? Log the chain?
-			UE::MVVM::FMessageLog Log(GetUserWidget());
-			Log.Error(LOCTEXT("RecursionDetected", "A recusion binding was detected (ie. A->B->C->A->B->C) at runtime."));
-			return;
+			ensureMsgf(false, TEXT("We should not have registered the binding since it will always be executed."));
 		}
+	}
+}
+
 
+void UMVVMView::ExecuteDelayedBinding(const FMVVMViewDelayedBinding& DelayedBinding) const
+{
+	if (ensure(ClassExtension))
+	{
+		if (ensure(ClassExtension->GetCompiledBindings().IsValidIndex(DelayedBinding.GetCompiledBindingIndex())))
 		{
-			UE::MVVM::Private::RecursiveDetector.Emplace(InViewModelOrWidget, InCompiledBindingIndex);
+			if (IsLibraryBindingEnabled(DelayedBinding.GetCompiledBindingIndex()))
+			{
+				const FMVVMViewClass_CompiledBinding& Binding = ClassExtension->GetCompiledBinding(DelayedBinding.GetCompiledBindingIndex());
 
-			ExecuteLibraryBinding(Binding, InViewModelOrWidget);
+				// Test for recursivity
+				int32 CompiledBindingIndex = DelayedBinding.GetCompiledBindingIndex();
+				const UMVVMView* Self = this;
+				if (UE::MVVM::Private::RecursiveDetector.FindByPredicate([Self, CompiledBindingIndex](const UE::MVVM::Private::FRecursiveDetctionElement& Element)
+					{
+						return Element.Get<0>() == Self && Element.Get<1>() == CompiledBindingIndex;
+					}) != nullptr)
+				{
+					ensureAlwaysMsgf(false, TEXT("Recursive binding detected"));
+					//Todo add more infos. Callstack maybe? Log the chain?
+					UE::MVVM::FMessageLog Log(GetUserWidget());
+					Log.Error(LOCTEXT("RecursionDetected", "A recursive binding was detected (ie. A->B->C->A->B->C) at runtime."));
+					return;
+				}
+
+				{
+					UE::MVVM::Private::RecursiveDetector.Emplace(this, CompiledBindingIndex);
+
+					ExecuteLibraryBinding(Binding);
 
-			UE::MVVM::Private::RecursiveDetector.Pop();
+					UE::MVVM::Private::RecursiveDetector.Pop();
+				}
+			}
+		}
+	}
+}
+
+
+void UMVVMView::ExecuteEveryTickBindings() const
+{
+	ensure(bHasEveryTickBinding);
+
+	if (ClassExtension)
+	{
+		const TArrayView<const FMVVMViewClass_CompiledBinding> CompiledBindings = ClassExtension->GetCompiledBindings();
+
+		for (int32 Index = 0; Index < CompiledBindings.Num(); ++Index)
+		{
+			const FMVVMViewClass_CompiledBinding& Binding = CompiledBindings[Index];
+			if (Binding.GetExecuteMode() == EMVVMExecutionMode::Tick && IsLibraryBindingEnabled(Index))
+			{
+				ExecuteLibraryBinding(Binding);
+			}
 		}
 	}
 }
@@ -293,8 +392,9 @@ void UMVVMView::EnableLibraryBinding(const FMVVMViewClass_CompiledBinding& Bindi
 	EnabledLibraryBindings.PadToNum(BindingIndex + 1, false);
 	check(EnabledLibraryBindings[BindingIndex] == false);
 
+	EMVVMExecutionMode ExecutionMode = Binding.GetExecuteMode();
 	bool bRegistered = true;
-	if (!Binding.IsOneTime())
+	if (!Binding.IsOneTime() && ExecutionMode != EMVVMExecutionMode::Tick)
 	{
 		bRegistered = RegisterLibraryBinding(Binding, BindingIndex);
 	}
@@ -311,6 +411,8 @@ void UMVVMView::EnableLibraryBinding(const FMVVMViewClass_CompiledBinding& Bindi
 	{
 		ExecuteLibraryBinding(Binding);
 	}
+
+	bHasEveryTickBinding = bHasEveryTickBinding || ExecutionMode == EMVVMExecutionMode::Tick;
 }
 
 
@@ -318,8 +420,12 @@ void UMVVMView::DisableLibraryBinding(const FMVVMViewClass_CompiledBinding& Bind
 {
 	check(IsLibraryBindingEnabled(BindingIndex));
 
+	EMVVMExecutionMode ExecutionMode = Binding.GetExecuteMode();
 	EnabledLibraryBindings[BindingIndex] = false;
-	UnregisterLibraryBinding(Binding);
+	if (!Binding.IsOneTime() && ExecutionMode != EMVVMExecutionMode::Tick)
+	{
+		UnregisterLibraryBinding(Binding);
+	}
 
 	if (bLogBinding)
 	{
diff --git a/Engine/Plugins/Runtime/ModelViewViewModel/Source/ModelViewViewModel/Private/View/MVVMViewClass.cpp b/Engine/Plugins/Runtime/ModelViewViewModel/Source/ModelViewViewModel/Private/View/MVVMViewClass.cpp
--- a/Engine/Plugins/Runtime/ModelViewViewModel/Source/ModelViewViewModel/Private/View/MVVMViewClass.cpp
+++ b/Engine/Plugins/Runtime/ModelViewViewModel/Source/ModelViewViewModel/Private/View/MVVMViewClass.cpp
@@ -156,6 +156,21 @@ UObject* FMVVMViewClass_SourceCreator::CreateInstance(const UMVVMViewClass* InVi
 ///////////////////////////////////////////////////////////////////////
 // 
 ///////////////////////////////////////////////////////////////////////
+namespace UE::MVVM::Private
+{
+	int32 GDefaultEvaluationMode = (int32)EMVVMExecutionMode::Immediate;
+	static FAutoConsoleVariableRef CVarDefaultEvaluationMode(
+		TEXT("MVVM.DefaultExecutionMode"),
+		GDefaultEvaluationMode,
+		TEXT("The default evaluation mode of a MVVM binding.")
+	);
+}
+
+EMVVMExecutionMode FMVVMViewClass_CompiledBinding::GetExecuteMode() const
+{
+	EMVVMExecutionMode DefaultMode = (EMVVMExecutionMode)UE::MVVM::Private::GDefaultEvaluationMode;
+	return (Flags & EBindingFlags::OverrideExecuteMode) == 0 ? DefaultMode : ExecutionMode;
+}
 
 FString FMVVMViewClass_CompiledBinding::ToString() const
 {
diff --git a/Engine/Plugins/Runtime/ModelViewViewModel/Source/ModelViewViewModel/Private/View/MVVMViewWorldSubsystem.cpp b/Engine/Plugins/Runtime/ModelViewViewModel/Source/ModelViewViewModel/Private/View/MVVMViewWorldSubsystem.cpp
new file mode 100644
--- /dev/null
+++ b/Engine/Plugins/Runtime/ModelViewViewModel/Source/ModelViewViewModel/Private/View/MVVMViewWorldSubsystem.cpp
@@ -0,0 +1,95 @@
+// Copyright Epic Games, Inc. All Rights Reserved.
+
+#include "View/MVVMViewWorldSubsystem.h"
+
+#include "View/MVVMView.h"
+
+void UMVVMViewWorldSubsystem::Tick(float DeltaTime)
+{
+	if (EveryTickBindings.Num() > 0)
+	{
+		FMemMark Mark(FMemStack::Get());
+		TArray<const UMVVMView*, TMemStackAllocator<>> ToTick;
+		ToTick.Reserve(EveryTickBindings.Num());
+		for (TWeakObjectPtr<const UMVVMView> View : EveryTickBindings)
+		{
+			if (const UMVVMView* ViewPtr = View.Get())
+			{
+				ToTick.Add(ViewPtr);
+			}
+		}
+
+		for (const UMVVMView* ViewPtr : ToTick)
+		{
+			ViewPtr->ExecuteEveryTickBindings();
+		}
+	}
+
+	if (DelayedBindings.Num() > 0)
+	{
+		FDelayedMap AllDelayedBindingsExecutedThisFrame;
+		AllDelayedBindingsExecutedThisFrame.Reserve(DelayedBindings.Num());
+
+		FDelayedMap DelayedBindingsWhileTicking = MoveTemp(DelayedBindings);
+		DelayedBindings = FDelayedMap();
+
+		do 
+		{
+			for (const auto& DelayedBindingsPair : DelayedBindingsWhileTicking)
+			{
+				if (const UMVVMView* View = DelayedBindingsPair.Key.Get())
+				{
+					ensure(DelayedBindingsPair.Value.Num() > 0);
+					FDelayedBindingList& AllBindingList = AllDelayedBindingsExecutedThisFrame.FindOrAdd(DelayedBindingsPair.Key);
+					for (const FMVVMViewDelayedBinding& DelayedBinding : DelayedBindingsPair.Value)
+					{
+						View->ExecuteDelayedBinding(DelayedBinding);
+						AllBindingList.AddUnique(DelayedBinding);
+					}
+				}
+			}
+
+			DelayedBindingsWhileTicking.Reset();
+
+			// Test new binding added while executing the current binding list
+			for (const auto& DelayedBindingPair : DelayedBindings)
+			{
+				if (DelayedBindingPair.Key.Get())
+				{
+					if (FDelayedBindingList* AllDelayedBindingsListPtr = AllDelayedBindingsExecutedThisFrame.Find(DelayedBindingPair.Key))
+					{
+						for (const FMVVMViewDelayedBinding& NewCompiledBinding : DelayedBindingPair.Value)
+						{
+							if (!AllDelayedBindingsListPtr->Find(NewCompiledBinding))
+							{
+								// It was not executed. Add it to be executed this frame.
+								DelayedBindingsWhileTicking.FindOrAdd(DelayedBindingPair.Key).AddUnique(NewCompiledBinding);
+							}
+						}
+					}
+				}
+			}
+		} while (DelayedBindingsWhileTicking.Num() > 0);
+	}
+}
+
+TStatId UMVVMViewWorldSubsystem::GetStatId() const
+{
+	RETURN_QUICK_DECLARE_CYCLE_STAT(UMVVMViewWorldSubsystem, STATGROUP_Tickables);
+}
+
+void UMVVMViewWorldSubsystem::AddViewWithEveryTickBinding(const UMVVMView* InView)
+{
+	check(!EveryTickBindings.Contains(InView));
+	EveryTickBindings.Add(InView);
+}
+
+void UMVVMViewWorldSubsystem::RemoveViewWithEveryTickBinding(const UMVVMView* InView)
+{
+	EveryTickBindings.RemoveSingleSwap(InView);
+}
+
+void UMVVMViewWorldSubsystem::AddDelayedBinding(const UMVVMView* View, FMVVMViewDelayedBinding InCompiledBinding)
+{
+	DelayedBindings.FindOrAdd(View).AddUnique(InCompiledBinding);
+}
\ No newline at end of file
diff --git a/Engine/Plugins/Runtime/ModelViewViewModel/Source/ModelViewViewModel/Private/View/MVVMViewWorldSubsystem.h b/Engine/Plugins/Runtime/ModelViewViewModel/Source/ModelViewViewModel/Private/View/MVVMViewWorldSubsystem.h
new file mode 100644
--- /dev/null
+++ b/Engine/Plugins/Runtime/ModelViewViewModel/Source/ModelViewViewModel/Private/View/MVVMViewWorldSubsystem.h
@@ -0,0 +1,60 @@
+// Copyright Epic Games, Inc. All Rights Reserved.
+
+#pragma once
+
+#include "Subsystems/WorldSubsystem.h"
+
+#include "MVVMViewWorldSubsystem.generated.h"
+
+class UMVVMView;
+
+/** */
+struct FMVVMViewDelayedBinding
+{
+	FMVVMViewDelayedBinding(int32 InCompiledBindingIndex)
+		: CompiledBindingIndex(InCompiledBindingIndex)
+	{
+	}
+
+	int32 GetCompiledBindingIndex() const
+	{
+		return CompiledBindingIndex;
+	}
+
+	bool operator==(const FMVVMViewDelayedBinding& Other) const
+	{
+		return CompiledBindingIndex == Other.CompiledBindingIndex;
+	}
+
+	bool operator!=(const FMVVMViewDelayedBinding& Other) const
+	{
+		return CompiledBindingIndex != Other.CompiledBindingIndex;
+	}
+
+private:
+	int32 CompiledBindingIndex = INDEX_NONE;
+};
+
+/** */
+UCLASS()
+class UMVVMViewWorldSubsystem : public UTickableWorldSubsystem
+{
+	GENERATED_BODY()
+
+public:
+	//~ Begin UEngineSubsystem interface
+	virtual void Tick(float DeltaTime) override;
+	virtual TStatId GetStatId() const override;
+	//~ End UEngineSubsystem interface
+
+	void AddViewWithEveryTickBinding(const UMVVMView* View);
+	void RemoveViewWithEveryTickBinding(const UMVVMView* View);
+
+	void AddDelayedBinding(const UMVVMView* View, FMVVMViewDelayedBinding CompiledBinding);
+
+private:
+	using FDelayedBindingList = TArray<FMVVMViewDelayedBinding, TInlineAllocator<8>>;
+	using FDelayedMap = TMap<TWeakObjectPtr<const UMVVMView>, FDelayedBindingList>;
+	FDelayedMap DelayedBindings;
+	TArray<TWeakObjectPtr<const UMVVMView>> EveryTickBindings;
+};
diff --git a/Engine/Plugins/Runtime/ModelViewViewModel/Source/ModelViewViewModel/Public/Types/MVVMBindingMode.h b/Engine/Plugins/Runtime/ModelViewViewModel/Source/ModelViewViewModel/Public/Types/MVVMBindingMode.h
--- a/Engine/Plugins/Runtime/ModelViewViewModel/Source/ModelViewViewModel/Public/Types/MVVMBindingMode.h
+++ b/Engine/Plugins/Runtime/ModelViewViewModel/Source/ModelViewViewModel/Public/Types/MVVMBindingMode.h
@@ -18,19 +18,6 @@ enum class EMVVMBindingMode : uint8
 };
 
 
-/** */
-UENUM()
-enum class EMVVMViewBindingUpdateMode : uint8
-{
-	/** Execute the binding as soon as the source value changes. */
-	Immediate,
-	///** Execute the binding at the end of the frame before drawing when the source value changes. */
-	//Delayed,
-	///** Always execute the binding at the end of the frame before drawing. */
-	//Debug,
-};
-
-
 namespace UE::MVVM
 {
 	UE_NODISCARD MODELVIEWVIEWMODEL_API bool IsForwardBinding(EMVVMBindingMode Mode);
diff --git a/Engine/Plugins/Runtime/ModelViewViewModel/Source/ModelViewViewModel/Public/Types/MVVMExecutionMode.h b/Engine/Plugins/Runtime/ModelViewViewModel/Source/ModelViewViewModel/Public/Types/MVVMExecutionMode.h
new file mode 100644
--- /dev/null
+++ b/Engine/Plugins/Runtime/ModelViewViewModel/Source/ModelViewViewModel/Public/Types/MVVMExecutionMode.h
@@ -0,0 +1,19 @@
+// Copyright Epic Games, Inc. All Rights Reserved.
+
+#pragma once
+
+#include "CoreMinimal.h"
+#include "MVVMExecutionMode.generated.h"
+
+
+/** */
+UENUM()
+enum class EMVVMExecutionMode : uint8
+{
+	/** Execute the binding as soon as the source value changes. */
+	Immediate,
+	/** Execute the binding at the end of the frame before drawing when the source value changes. */
+	Delayed,
+	/** Always execute the binding at the end of the frame. */
+	Tick,
+};
\ No newline at end of file
diff --git a/Engine/Plugins/Runtime/ModelViewViewModel/Source/ModelViewViewModel/Public/View/MVVMView.h b/Engine/Plugins/Runtime/ModelViewViewModel/Source/ModelViewViewModel/Public/View/MVVMView.h
--- a/Engine/Plugins/Runtime/ModelViewViewModel/Source/ModelViewViewModel/Public/View/MVVMView.h
+++ b/Engine/Plugins/Runtime/ModelViewViewModel/Source/ModelViewViewModel/Public/View/MVVMView.h
@@ -15,10 +15,10 @@
 
 class UMVVMViewClass;
 struct FMVVMViewClass_CompiledBinding;
+struct FMVVMViewDelayedBinding;
 class UMVVMViewModelBase;
 class UWidget;
 
-
 /**
  * Instance UMVVMClassExtension_View for the UUserWdiget
  */
@@ -33,12 +33,11 @@ public:
 	//~ Begin UUserWidgetExtension implementation
 	virtual void Construct() override;
 	virtual void Destruct() override;
-
-// todo needed when we will support different type of update
-	//virtual bool RequiresTick() const override;
-	//virtual void Tick(const FGeometry& MyGeometry, float InDeltaTime) override;
 	//~ End UUserWidgetExtension implementation
 
+	void ExecuteDelayedBinding(const FMVVMViewDelayedBinding& DelayedBinding) const;
+	void ExecuteEveryTickBindings() const;
+
 // todo a way to identify a binding from outside. maybe a unique name in the editor?
 	//UFUNCTION(BlueprintCallable, Category = "MVVM")
 	//void SetLibraryBindingEnabled(FGuid ViewModelId, FMVVMBindingName BindingName, bool bEnable);
@@ -89,4 +88,7 @@ private:
 
 	/** Is the Construct method was called. */
 	bool bConstructed = false;
+
+	/** Is the Construct method was called. */
+	bool bHasEveryTickBinding = false;
 };
diff --git a/Engine/Plugins/Runtime/ModelViewViewModel/Source/ModelViewViewModel/Public/View/MVVMViewClass.h b/Engine/Plugins/Runtime/ModelViewViewModel/Source/ModelViewViewModel/Public/View/MVVMViewClass.h
--- a/Engine/Plugins/Runtime/ModelViewViewModel/Source/ModelViewViewModel/Public/View/MVVMViewClass.h
+++ b/Engine/Plugins/Runtime/ModelViewViewModel/Source/ModelViewViewModel/Public/View/MVVMViewClass.h
@@ -12,6 +12,7 @@
 #include "FieldNotification/FieldId.h"
 #include "Types/MVVMBindingMode.h"
 #include "Types/MVVMBindingName.h"
+#include "Types/MVVMExecutionMode.h"
 #include "Types/MVVMViewModelContext.h"
 #include "View/MVVMView.h"
 
@@ -146,6 +147,9 @@ public:
 	{
 		return (Flags & EBindingFlags::ConversionFunctionIsComplex) != 0;
 	}
+	
+	/** How the binding should be executed. */
+	EMVVMExecutionMode GetExecuteMode() const;
 
 	/** @return a human readable version of the binding that can be use for debugging purposes. */
 	FString ToString() const;
@@ -160,9 +164,8 @@ private:
 	UPROPERTY()
 	FMVVMVCompiledBinding Binding;
 
-	/** How the binding should be executed. */
 	UPROPERTY()
-	EMVVMViewBindingUpdateMode UpdateMode = EMVVMViewBindingUpdateMode::Immediate;
+	EMVVMExecutionMode ExecutionMode = EMVVMExecutionMode::Immediate;
 
 	enum EBindingFlags
 	{
@@ -173,8 +176,8 @@ private:
 		EnabledByDefault = 1 << 3,
 		ViewModelOptional = 1 << 4,	// The source (viewmodel) can be nullptr and the binding could failed and should not log a warning.
 		ConversionFunctionIsComplex = 1 << 5,	// The conversion function is complex, there is no input. The inputs are calculated in the BP function.
-		Unused01 = 1 << 6,
-		Unused02 = 1 << 7,
+		// In development, (when the Blueprint maybe not be compiled with the latest data), the ExecutionMode may not reflect the default project setting value.
+		OverrideExecuteMode = 1 << 6,
 	};
 
 	UPROPERTY()
diff --git a/Engine/Plugins/Runtime/ModelViewViewModel/Source/ModelViewViewModelBlueprint/Private/MVVMViewBlueprintCompiler.cpp b/Engine/Plugins/Runtime/ModelViewViewModel/Source/ModelViewViewModelBlueprint/Private/MVVMViewBlueprintCompiler.cpp
--- a/Engine/Plugins/Runtime/ModelViewViewModel/Source/ModelViewViewModelBlueprint/Private/MVVMViewBlueprintCompiler.cpp
+++ b/Engine/Plugins/Runtime/ModelViewViewModel/Source/ModelViewViewModelBlueprint/Private/MVVMViewBlueprintCompiler.cpp
@@ -1060,6 +1060,13 @@ bool FMVVMViewBlueprintCompiler::CompileBindings(const FCompiledBindingLibraryCo
 		return false;
 	}
 
+	IConsoleVariable* CVarDefaultExecutionMode = IConsoleManager::Get().FindConsoleVariable(TEXT("MVVM.DefaultExecutionMode"));
+	if (!CVarDefaultExecutionMode)
+	{
+		WidgetBlueprintCompilerContext.MessageLog.Error(*LOCTEXT("CantFindDefaultExecutioMode", "The default execution mode cannot be found.").ToString());
+		return false;
+	}
+
 	for (const FCompilerBinding& CompileBinding : CompilerBindings)
 	{
 		FMVVMBlueprintViewBinding& ViewBinding = *BlueprintView->GetBindingAt(CompileBinding.BindingIndex);
@@ -1099,7 +1106,7 @@ bool FMVVMViewBlueprintCompiler::CompileBindings(const FCompiledBindingLibraryCo
 
 		NewBinding.FieldId = CompiledFieldId  ? *CompiledFieldId : FMVVMVCompiledFieldId();
 		NewBinding.Binding = *CompiledBinding;
-		NewBinding.UpdateMode = ViewBinding.UpdateMode;
+		NewBinding.ExecutionMode = ViewBinding.bOverrideExecutionMode ? ViewBinding.OverrideExecutionMode : (EMVVMExecutionMode)CVarDefaultExecutionMode->GetInt();
 
 		NewBinding.Flags = 0;
 		NewBinding.Flags |= (ViewBinding.bEnabled) ? FMVVMViewClass_CompiledBinding::EBindingFlags::EnabledByDefault : 0;
@@ -1108,6 +1115,7 @@ bool FMVVMViewBlueprintCompiler::CompileBindings(const FCompiledBindingLibraryCo
 		NewBinding.Flags |= (IsOneTimeBinding(ViewBinding.BindingType)) ? FMVVMViewClass_CompiledBinding::EBindingFlags::OneTime : 0;
 		NewBinding.Flags |= (bIsOptional) ? FMVVMViewClass_CompiledBinding::EBindingFlags::ViewModelOptional : 0;
 		NewBinding.Flags |= (CompileBinding.bIsConversionFunctionComplex) ? FMVVMViewClass_CompiledBinding::EBindingFlags::ConversionFunctionIsComplex : 0;
+		NewBinding.Flags |= (ViewBinding.bOverrideExecutionMode) ? FMVVMViewClass_CompiledBinding::EBindingFlags::OverrideExecuteMode : 0;
 
 		ViewExtension->CompiledBindings.Emplace(MoveTemp(NewBinding));
 	}
diff --git a/Engine/Plugins/Runtime/ModelViewViewModel/Source/ModelViewViewModelBlueprint/Public/MVVMBlueprintViewBinding.h b/Engine/Plugins/Runtime/ModelViewViewModel/Source/ModelViewViewModelBlueprint/Public/MVVMBlueprintViewBinding.h
--- a/Engine/Plugins/Runtime/ModelViewViewModel/Source/ModelViewViewModelBlueprint/Public/MVVMBlueprintViewBinding.h
+++ b/Engine/Plugins/Runtime/ModelViewViewModel/Source/ModelViewViewModelBlueprint/Public/MVVMBlueprintViewBinding.h
@@ -4,6 +4,7 @@
 
 #include "MVVMPropertyPath.h"
 #include "Types/MVVMBindingMode.h"
+#include "Types/MVVMExecutionMode.h"
 
 #include "MVVMBlueprintViewBinding.generated.h"
 
@@ -51,8 +52,11 @@ struct MODELVIEWVIEWMODELBLUEPRINT_API FMVVMBlueprintViewBinding
 	UPROPERTY(EditAnywhere, Category = "MVVM")
 	EMVVMBindingMode BindingType = EMVVMBindingMode::OneWayToDestination;
 
-	UPROPERTY(EditAnywhere, Category = "MVVM")
-	EMVVMViewBindingUpdateMode UpdateMode = EMVVMViewBindingUpdateMode::Immediate;
+	UPROPERTY()
+	bool bOverrideExecutionMode = false;
+
+	UPROPERTY(EditAnywhere, Category = "MVVM", meta=(EditCondition="bOverrideExecutionMode"))
+	EMVVMExecutionMode OverrideExecutionMode = EMVVMExecutionMode::Immediate;
 
 	/** */
 	UPROPERTY(EditAnywhere, Category = "MVVM")
diff --git a/Engine/Plugins/Runtime/ModelViewViewModel/Source/ModelViewViewModelEditor/Private/MVVMEditorSubsystem.cpp b/Engine/Plugins/Runtime/ModelViewViewModel/Source/ModelViewViewModelEditor/Private/MVVMEditorSubsystem.cpp
--- a/Engine/Plugins/Runtime/ModelViewViewModel/Source/ModelViewViewModelEditor/Private/MVVMEditorSubsystem.cpp
+++ b/Engine/Plugins/Runtime/ModelViewViewModel/Source/ModelViewViewModelEditor/Private/MVVMEditorSubsystem.cpp
@@ -492,19 +492,37 @@ void UMVVMEditorSubsystem::SetViewModelPropertyForBinding(UWidgetBlueprint* Widg
 	}
 }
 
-void UMVVMEditorSubsystem::SetUpdateModeForBinding(UWidgetBlueprint* WidgetBlueprint, FMVVMBlueprintViewBinding& Binding, EMVVMViewBindingUpdateMode Mode)
+void UMVVMEditorSubsystem::OverrideExecutionModeForBinding(UWidgetBlueprint* WidgetBlueprint, FMVVMBlueprintViewBinding& Binding, EMVVMExecutionMode Mode)
 {
-	if (Binding.UpdateMode != Mode)
+	if (!Binding.bOverrideExecutionMode || Binding.OverrideExecutionMode != Mode)
 	{
 		if (UMVVMBlueprintView* View = GetView(WidgetBlueprint))
 		{
-			FScopedTransaction Transaction(LOCTEXT("SetUpdateMode", "Set Update Mode"));
+			FScopedTransaction Transaction(LOCTEXT("SetExecutionMode", "Set Execution Mode"));
 
-			UE::MVVM::Private::OnBindingPreEditChange(View, GET_MEMBER_NAME_CHECKED(FMVVMBlueprintViewBinding, UpdateMode));
+			UE::MVVM::Private::OnBindingPreEditChange(View, GET_MEMBER_NAME_CHECKED(FMVVMBlueprintViewBinding, OverrideExecutionMode));
 
-			Binding.UpdateMode = Mode;
+			Binding.bOverrideExecutionMode = true;
+			Binding.OverrideExecutionMode = Mode;
 
-			UE::MVVM::Private::OnBindingPreEditChange(View, GET_MEMBER_NAME_CHECKED(FMVVMBlueprintViewBinding, UpdateMode));
+			UE::MVVM::Private::OnBindingPostEditChange(View, GET_MEMBER_NAME_CHECKED(FMVVMBlueprintViewBinding, OverrideExecutionMode));
+		}
+	}
+}
+
+void UMVVMEditorSubsystem::ResetExecutionModeForBinding(UWidgetBlueprint* WidgetBlueprint, FMVVMBlueprintViewBinding& Binding)
+{
+	if (Binding.bOverrideExecutionMode)
+	{
+		if (UMVVMBlueprintView* View = GetView(WidgetBlueprint))
+		{
+			FScopedTransaction Transaction(LOCTEXT("ResetExecutionMode", "Reset Execution Mode"));
+
+			UE::MVVM::Private::OnBindingPreEditChange(View, GET_MEMBER_NAME_CHECKED(FMVVMBlueprintViewBinding, OverrideExecutionMode));
+
+			Binding.bOverrideExecutionMode = false;
+
+			UE::MVVM::Private::OnBindingPostEditChange(View, GET_MEMBER_NAME_CHECKED(FMVVMBlueprintViewBinding, OverrideExecutionMode));
 		}
 	}
 }
diff --git a/Engine/Plugins/Runtime/ModelViewViewModel/Source/ModelViewViewModelEditor/Private/Widgets/SMVVMViewBindingListView.cpp b/Engine/Plugins/Runtime/ModelViewViewModel/Source/ModelViewViewModelEditor/Private/Widgets/SMVVMViewBindingListView.cpp
--- a/Engine/Plugins/Runtime/ModelViewViewModel/Source/ModelViewViewModelEditor/Private/Widgets/SMVVMViewBindingListView.cpp
+++ b/Engine/Plugins/Runtime/ModelViewViewModel/Source/ModelViewViewModelEditor/Private/Widgets/SMVVMViewBindingListView.cpp
@@ -271,6 +271,7 @@ public:
 		WidgetBlueprintWeak = InWidgetBlueprint;
 
 		OnBlueprintChangedHandle = InWidgetBlueprint->OnChanged().AddSP(this, &SBindingRow::HandleBlueprintChanged);
+		CVarDefaultExecutionMode = IConsoleManager::Get().FindConsoleVariable(TEXT("MVVM.DefaultExecutionMode"));
 		
 		FMVVMBlueprintViewBinding* ViewBinding = GetThisViewBinding();
 		FBindingSource WidgetSource = FBindingSource::CreateForWidget(InWidgetBlueprint, ViewBinding->WidgetPath.GetWidgetName());
@@ -384,10 +385,22 @@ public:
 				.HAlign(HAlign_Right)
 				.AutoWidth()
 				[
-					SNew(SEnumComboBox, StaticEnum<EMVVMViewBindingUpdateMode>())
+					SNew(SCheckBox)
+					.IsChecked(this, &SBindingRow::IsExecutionModeOverrideChecked)
+					.OnCheckStateChanged(this, &SBindingRow::OnExecutionModeOverrideChanged)
+				]
+
+				+ SHorizontalBox::Slot()
+				.Padding(2, 1)
+				.VAlign(VAlign_Fill)
+				.HAlign(HAlign_Right)
+				.AutoWidth()
+				[
+					SNew(SEnumComboBox, StaticEnum<EMVVMExecutionMode>())
 					.ContentPadding(FMargin(4, 0))
-					.OnEnumSelectionChanged(this, &SBindingRow::OnUpdateModeSelectionChanged)
-					.CurrentValue(this, &SBindingRow::GetUpdateModeValue)
+					.OnEnumSelectionChanged(this, &SBindingRow::OnExecutionModeSelectionChanged)
+					.CurrentValue(this, &SBindingRow::GetExecutioModeValue)
+					.IsEnabled(this, &SBindingRow::IsExecutionModeOverridden)
 				]
 			],
 			OwnerTableView
@@ -654,22 +667,62 @@ private:
 		}
 	}
 
-	void OnUpdateModeSelectionChanged(int32 Value, ESelectInfo::Type)
+	ECheckBoxState IsExecutionModeOverrideChecked() const
+	{
+		if (FMVVMBlueprintViewBinding* ViewBinding = GetThisViewBinding())
+		{
+			return ViewBinding->bOverrideExecutionMode ? ECheckBoxState::Checked : ECheckBoxState::Unchecked;
+		}
+		return ECheckBoxState::Undetermined;
+	}
+
+	void OnExecutionModeOverrideChanged(ECheckBoxState NewState)
+	{
+		if (FMVVMBlueprintViewBinding* ViewBinding = GetThisViewBinding())
+		{
+			if (NewState == ECheckBoxState::Checked)
+			{
+				UMVVMEditorSubsystem* Subsystem = GEditor->GetEditorSubsystem<UMVVMEditorSubsystem>();
+				EMVVMExecutionMode ExecutionMode = CVarDefaultExecutionMode ? (EMVVMExecutionMode)CVarDefaultExecutionMode->GetInt() : EMVVMExecutionMode::Immediate;
+				Subsystem->OverrideExecutionModeForBinding(WidgetBlueprintWeak.Get(), *ViewBinding, ExecutionMode);
+			}
+			else
+			{
+				UMVVMEditorSubsystem* Subsystem = GEditor->GetEditorSubsystem<UMVVMEditorSubsystem>();
+				Subsystem->ResetExecutionModeForBinding(WidgetBlueprintWeak.Get(), *ViewBinding);
+			}
+		}
+	}
+
+	void OnExecutionModeSelectionChanged(int32 Value, ESelectInfo::Type)
 	{
 		if (FMVVMBlueprintViewBinding* ViewBinding = GetThisViewBinding())
 		{
 			UMVVMEditorSubsystem* Subsystem = GEditor->GetEditorSubsystem<UMVVMEditorSubsystem>();
-			Subsystem->SetUpdateModeForBinding(WidgetBlueprintWeak.Get(), *ViewBinding, (EMVVMViewBindingUpdateMode) Value);
+			Subsystem->OverrideExecutionModeForBinding(WidgetBlueprintWeak.Get(), *ViewBinding, static_cast<EMVVMExecutionMode>(Value));
 		}			
 	}
 
-	int32 GetUpdateModeValue() const
+	int32 GetExecutioModeValue() const
 	{
+		EMVVMExecutionMode ExecutionMode = CVarDefaultExecutionMode ? (EMVVMExecutionMode)CVarDefaultExecutionMode->GetInt() : EMVVMExecutionMode::Immediate;
 		if (FMVVMBlueprintViewBinding* ViewBinding = GetThisViewBinding())
 		{
-			return (int32) ViewBinding->UpdateMode;
+			if (ViewBinding->bOverrideExecutionMode)
+			{
+				ExecutionMode = ViewBinding->OverrideExecutionMode;
+			}
 		}
-		return (int32) EMVVMViewBindingUpdateMode::Immediate;
+		return static_cast<int32>(ExecutionMode);
+	}
+	
+	bool IsExecutionModeOverridden() const
+	{
+		if (FMVVMBlueprintViewBinding* ViewBinding = GetThisViewBinding())
+		{
+			return ViewBinding->bOverrideExecutionMode;
+		}
+		return false;
 	}
 
 	void OnIsBindingEnableChanged(ECheckBoxState NewState)
@@ -839,6 +892,7 @@ private:
 	TSharedPtr<SMenuAnchor> WidgetMenuAnchor;
 	TSharedPtr<SSourceBindingList> WidgetBindingList;
 	TSharedPtr<SFieldEntry> SelectedWidget;
+	IConsoleVariable* CVarDefaultExecutionMode = nullptr;
 };
 
 class SFunctionParameterRow : public STableRow<TSharedPtr<FBindingEntry>>
diff --git a/Engine/Plugins/Runtime/ModelViewViewModel/Source/ModelViewViewModelEditor/Public/MVVMEditorSubsystem.h b/Engine/Plugins/Runtime/ModelViewViewModel/Source/ModelViewViewModelEditor/Public/MVVMEditorSubsystem.h
--- a/Engine/Plugins/Runtime/ModelViewViewModel/Source/ModelViewViewModelEditor/Public/MVVMEditorSubsystem.h
+++ b/Engine/Plugins/Runtime/ModelViewViewModel/Source/ModelViewViewModelEditor/Public/MVVMEditorSubsystem.h
@@ -8,6 +8,7 @@
 #include "Templates/SubclassOf.h"
 #include "Types/MVVMBindingMode.h"
 #include "Types/MVVMBindingSource.h"
+#include "Types/MVVMExecutionMode.h"
 #include "Types/MVVMFieldVariant.h"
 
 #include "MVVMEditorSubsystem.generated.h"
@@ -54,7 +55,8 @@ public:
 	void SetDestinationToSourceConversionFunction(UWidgetBlueprint* WidgetBlueprint, FMVVMBlueprintViewBinding& Binding, const UFunction* ConversionFunction);
 	void SetWidgetPropertyForBinding(UWidgetBlueprint* WidgetBlueprint, FMVVMBlueprintViewBinding& Binding, FMVVMBlueprintPropertyPath Field);
 	void SetViewModelPropertyForBinding(UWidgetBlueprint* WidgetBlueprint, FMVVMBlueprintViewBinding& Binding, FMVVMBlueprintPropertyPath Field);
-	void SetUpdateModeForBinding(UWidgetBlueprint* WidgetBlueprint, FMVVMBlueprintViewBinding& Binding, EMVVMViewBindingUpdateMode Mode);
+	void OverrideExecutionModeForBinding(UWidgetBlueprint* WidgetBlueprint, FMVVMBlueprintViewBinding& Binding, EMVVMExecutionMode Mode);
+	void ResetExecutionModeForBinding(UWidgetBlueprint* WidgetBlueprint, FMVVMBlueprintViewBinding& Binding);
 	void SetBindingTypeForBinding(UWidgetBlueprint* WidgetBlueprint, FMVVMBlueprintViewBinding& Binding, EMVVMBindingMode Type);
 	void SetEnabledForBinding(UWidgetBlueprint* WidgetBlueprint, FMVVMBlueprintViewBinding& Binding, bool bEnabled);
 	void SetCompileForBinding(UWidgetBlueprint* WidgetBlueprint, FMVVMBlueprintViewBinding& Binding, bool bCompile);
