MVVM: Fix the binding picker to select property when we are in read mode.

diff --git a/Engine/Plugins/Runtime/ModelViewViewModel/Config/BaseModelViewViewModel.ini b/Engine/Plugins/Runtime/ModelViewViewModel/Config/BaseModelViewViewModel.ini
--- a/Engine/Plugins/Runtime/ModelViewViewModel/Config/BaseModelViewViewModel.ini
+++ b/Engine/Plugins/Runtime/ModelViewViewModel/Config/BaseModelViewViewModel.ini
@@ -4,4 +4,4 @@
 
 
 [/Script/ModelViewViewModelBlueprint.MVVMDeveloperProjectSettings]
-+FieldSelectorPermissions=(("/Script/UMG.Widget", (DisallowedFieldNames=("TooltipWidget","ForceVolatile","SetUserFocus","Navigation"),AdvancedFieldNames=())))
\ No newline at end of file
++FieldSelectorPermissions=(("/Script/UMG.Widget", (DisallowedFieldNames=("TooltipWidget","ForceVolatile","SetUserFocus","Navigation","GetParent","GetOwningPlayer","GetOwningLocalPlayer","GetGameInstance"),AdvancedFieldNames=())))
\ No newline at end of file
diff --git a/Engine/Plugins/Runtime/ModelViewViewModel/Source/ModelViewViewModel/Private/View/MVVMView.cpp b/Engine/Plugins/Runtime/ModelViewViewModel/Source/ModelViewViewModel/Private/View/MVVMView.cpp
--- a/Engine/Plugins/Runtime/ModelViewViewModel/Source/ModelViewViewModel/Private/View/MVVMView.cpp
+++ b/Engine/Plugins/Runtime/ModelViewViewModel/Source/ModelViewViewModel/Private/View/MVVMView.cpp
@@ -130,7 +130,6 @@ bool UMVVMView::SetViewModel(FName ViewModelName, TScriptInterface<INotifyFieldV
 			const TArrayView<const FMVVMViewClass_CompiledBinding> CompiledBindings = ClassExtension->GetCompiledBindings();
 
 			// Unregister any from that source
-			if (PreviousValue)
 			{
 				BindingToReenabled.Reserve(CompiledBindings.Num());
 
@@ -420,6 +419,7 @@ void UMVVMView::DisableLibraryBinding(const FMVVMViewClass_CompiledBinding& Bind
 {
 	check(IsLibraryBindingEnabled(BindingIndex));
 
+
 	EMVVMExecutionMode ExecutionMode = Binding.GetExecuteMode();
 	EnabledLibraryBindings[BindingIndex] = false;
 	if (!Binding.IsOneTime() && ExecutionMode != EMVVMExecutionMode::Tick)
diff --git a/Engine/Plugins/Runtime/ModelViewViewModel/Source/ModelViewViewModelBlueprint/Private/MVVMDeveloperProjectSettings.cpp b/Engine/Plugins/Runtime/ModelViewViewModel/Source/ModelViewViewModelBlueprint/Private/MVVMDeveloperProjectSettings.cpp
--- a/Engine/Plugins/Runtime/ModelViewViewModel/Source/ModelViewViewModelBlueprint/Private/MVVMDeveloperProjectSettings.cpp
+++ b/Engine/Plugins/Runtime/ModelViewViewModel/Source/ModelViewViewModelBlueprint/Private/MVVMDeveloperProjectSettings.cpp
@@ -14,7 +14,7 @@ FText UMVVMDeveloperProjectSettings::GetSectionText() const
 	return LOCTEXT("MVVMProjectSettings", "Model View Viewmodel");
 }
 
-bool UMVVMDeveloperProjectSettings::IsPropertyAllowed(FProperty* Property) const
+bool UMVVMDeveloperProjectSettings::IsPropertyAllowed(const FProperty* Property) const
 {
 	check(Property);
 	TStringBuilder<256> StringBuilder;
@@ -28,7 +28,7 @@ bool UMVVMDeveloperProjectSettings::IsPropertyAllowed(FProperty* Property) const
 	return true;
 }
 
-bool UMVVMDeveloperProjectSettings::IsFunctionAllowed(UFunction* Function) const
+bool UMVVMDeveloperProjectSettings::IsFunctionAllowed(const UFunction* Function) const
 {
 	check(Function);
 	TStringBuilder<256> StringBuilder;
diff --git a/Engine/Plugins/Runtime/ModelViewViewModel/Source/ModelViewViewModelBlueprint/Public/MVVMDeveloperProjectSettings.h b/Engine/Plugins/Runtime/ModelViewViewModel/Source/ModelViewViewModelBlueprint/Public/MVVMDeveloperProjectSettings.h
--- a/Engine/Plugins/Runtime/ModelViewViewModel/Source/ModelViewViewModelBlueprint/Public/MVVMDeveloperProjectSettings.h
+++ b/Engine/Plugins/Runtime/ModelViewViewModel/Source/ModelViewViewModelBlueprint/Public/MVVMDeveloperProjectSettings.h
@@ -38,8 +38,8 @@ public:
 	virtual FName GetCategoryName() const override;
 	virtual FText GetSectionText() const override;
 
-	bool IsPropertyAllowed(FProperty* Property) const;
-	bool IsFunctionAllowed(UFunction* Function) const;
+	bool IsPropertyAllowed(const FProperty* Property) const;
+	bool IsFunctionAllowed(const UFunction* Function) const;
 
 private:
 	/** Permission list for filtering which properties are visible in UI. */
diff --git a/Engine/Plugins/Runtime/ModelViewViewModel/Source/ModelViewViewModelEditor/Private/Widgets/SMVVMViewModelBindingListWidget.cpp b/Engine/Plugins/Runtime/ModelViewViewModel/Source/ModelViewViewModelEditor/Private/Widgets/SMVVMViewModelBindingListWidget.cpp
--- a/Engine/Plugins/Runtime/ModelViewViewModel/Source/ModelViewViewModelEditor/Private/Widgets/SMVVMViewModelBindingListWidget.cpp
+++ b/Engine/Plugins/Runtime/ModelViewViewModel/Source/ModelViewViewModelEditor/Private/Widgets/SMVVMViewModelBindingListWidget.cpp
@@ -55,19 +55,49 @@ namespace Private
 		{
 			if (FieldVariant.IsFunction())
 			{
-				if (EnumHasAllFlags(FieldVisibilityFlags, EFieldVisibility::Writable) && !Binding.IsWritable())
+				const UFunction* Function = FieldVariant.GetFunction();
+				if (Function == nullptr)
 				{
 					return TOptional<FFieldVariant>();
 				}
 
-				if (!GetDefault<UMVVMDeveloperProjectSettings>()->IsFunctionAllowed(FieldVariant.GetFunction()))
+				const FProperty* ReturnProperty = BindingHelper::GetReturnProperty(Function);
+				bool bDoCompatibleTest = ReturnProperty == nullptr || AssignableTo != nullptr;
+				bool bDoWritableTest = true;
+				// Do we allow walking up the tree
+				if (CastField<FObjectPropertyBase>(ReturnProperty))
 				{
-					return TOptional<FFieldVariant>();
+					bDoCompatibleTest = bDoCompatibleTest && bDoObjectProperty;
+					bDoWritableTest = bDoObjectProperty;
 				}
 
-				const UFunction* Function = FieldVariant.GetFunction();
-				if (Function != nullptr && AssignableTo != nullptr && 
-					!BindingHelper::ArePropertiesCompatible(BindingHelper::GetReturnProperty(Function), AssignableTo))
+				if (bDoWritableTest)
+				{
+					if (EnumHasAllFlags(FieldVisibilityFlags, EFieldVisibility::Writable) && !Binding.IsWritable())
+					{
+						return TOptional<FFieldVariant>();
+					}
+				}
+
+				if (bDoCompatibleTest)
+				{
+					if (EnumHasAllFlags(FieldVisibilityFlags, EFieldVisibility::Writable))
+					{
+						if (!BindingHelper::ArePropertiesCompatible(AssignableTo, BindingHelper::GetFirstArgumentProperty(Function)))
+						{
+							return TOptional<FFieldVariant>();
+						}
+					}
+					if (EnumHasAllFlags(FieldVisibilityFlags, EFieldVisibility::Readable))
+					{
+						if (!BindingHelper::ArePropertiesCompatible(ReturnProperty, AssignableTo))
+						{
+							return TOptional<FFieldVariant>();
+						}
+					}
+				}
+
+				if (!GetDefault<UMVVMDeveloperProjectSettings>()->IsFunctionAllowed(Function))
 				{
 					return TOptional<FFieldVariant>();
 				}
@@ -76,20 +106,23 @@ namespace Private
 			}
 			else if (FieldVariant.IsProperty())
 			{
-				if (!GetDefault<UMVVMDeveloperProjectSettings>()->IsPropertyAllowed(FieldVariant.GetProperty()))
+				const FProperty* Property = FieldVariant.GetProperty();
+				if (Property == nullptr)
 				{
 					return TOptional<FFieldVariant>();
 				}
 
-				const FProperty* Property = FieldVariant.GetProperty();
-				if (Property != nullptr && AssignableTo != nullptr && 
-					!BindingHelper::ArePropertiesCompatible(Property, AssignableTo))
+				bool bDoCompatibleTest = AssignableTo != nullptr;
+				bool bDoWritableTest = true;
+				// Do we allow walking up the tree
+				if (CastField<FObjectPropertyBase>(Property))
 				{
-					return TOptional<FFieldVariant>();
+					bDoCompatibleTest = bDoCompatibleTest && bDoObjectProperty;
+					bDoWritableTest = bDoObjectProperty;
 				}
 
 				// If the path ends with the object property, then it needs to follow the writable rule
-				if (bDoObjectProperty || !CastField<FObjectPropertyBase>(FieldVariant.GetProperty()))
+				if (bDoWritableTest)
 				{
 					if (EnumHasAllFlags(FieldVisibilityFlags, EFieldVisibility::Writable) && !Binding.IsWritable())
 					{
@@ -97,6 +130,29 @@ namespace Private
 					}
 				}
 
+				if (bDoCompatibleTest)
+				{
+					if (EnumHasAllFlags(FieldVisibilityFlags, EFieldVisibility::Writable))
+					{
+						if (!BindingHelper::ArePropertiesCompatible(AssignableTo, Property))
+						{
+							return TOptional<FFieldVariant>();
+						}
+					}
+					if (EnumHasAllFlags(FieldVisibilityFlags, EFieldVisibility::Readable))
+					{
+						if (!BindingHelper::ArePropertiesCompatible(Property, AssignableTo))
+						{
+							return TOptional<FFieldVariant>();
+						}
+					}
+				}
+
+				if (!GetDefault<UMVVMDeveloperProjectSettings>()->IsPropertyAllowed(Property))
+				{
+					return TOptional<FFieldVariant>();
+				}
+
 				return FFieldVariant(Property);
 			}
 		}
@@ -176,15 +232,31 @@ TArray<FFieldVariant> FFieldIterator_Bindable::GetFields(const UStruct* Struct)
 	return Result;
 }
 
+/** */
+FFieldExpander_Bindable::FFieldExpander_Bindable()
+{
+	SetExpandObject(UE::PropertyViewer::FFieldExpander_Default::EObjectExpandFlag::UseInstanceClass);
+	SetExpandScriptStruct(true);
+	SetExpandFunction(UE::PropertyViewer::FFieldExpander_Default::EFunctionExpand::FunctionProperties);
+}
+
+TOptional<const UStruct*> FFieldExpander_Bindable::GetExpandedFunction(const UFunction* Function) const
+{
+	const FProperty* ReturnProperty = Function ? BindingHelper::GetReturnProperty(Function) : nullptr;
+	if (const FObjectPropertyBase* ObjectProperty = CastField<const FObjectPropertyBase>(ReturnProperty))
+	{
+		return ObjectProperty->PropertyClass;
+	}
+	return TOptional<const UStruct*>();
+}
+
 /** */
 void SSourceBindingList::Construct(const FArguments& InArgs, const UWidgetBlueprint* InWidgetBlueprint)
 {
 	WidgetBlueprint = InWidgetBlueprint;
 	FieldIterator = MakeUnique<FFieldIterator_Bindable>(InWidgetBlueprint, InArgs._FieldVisibilityFlags, InArgs._AssignableTo);
-	FieldExpander = MakeUnique<UE::PropertyViewer::FFieldExpander_Default>();
-	FieldExpander->SetExpandObject(UE::PropertyViewer::FFieldExpander_Default::EObjectExpandFlag::UseInstanceClass);
-	FieldExpander->SetExpandScriptStruct(true);
-	FieldExpander->SetExpandFunction(false);
+	FieldExpander = MakeUnique<FFieldExpander_Bindable>();
+
 
 	OnDoubleClicked = InArgs._OnDoubleClicked;
 
diff --git a/Engine/Plugins/Runtime/ModelViewViewModel/Source/ModelViewViewModelEditor/Private/Widgets/SMVVMViewModelBindingListWidget.h b/Engine/Plugins/Runtime/ModelViewViewModel/Source/ModelViewViewModelEditor/Private/Widgets/SMVVMViewModelBindingListWidget.h
--- a/Engine/Plugins/Runtime/ModelViewViewModel/Source/ModelViewViewModelEditor/Private/Widgets/SMVVMViewModelBindingListWidget.h
+++ b/Engine/Plugins/Runtime/ModelViewViewModel/Source/ModelViewViewModelEditor/Private/Widgets/SMVVMViewModelBindingListWidget.h
@@ -55,6 +55,16 @@ namespace UE::MVVM
 		const FProperty* AssignableTo = nullptr;
 	};
 
+	/**
+ *
+ */
+	class FFieldExpander_Bindable : public UE::PropertyViewer::FFieldExpander_Default
+	{
+	public:
+		FFieldExpander_Bindable();
+		virtual TOptional<const UStruct*> GetExpandedFunction(const UFunction* Function) const override;
+	};
+
 	/** 
 	 * 
 	 */
diff --git a/Engine/Source/Runtime/AdvancedWidgets/Private/Widgets/PropertyViewer/SPropertyViewerImpl.cpp b/Engine/Source/Runtime/AdvancedWidgets/Private/Widgets/PropertyViewer/SPropertyViewerImpl.cpp
--- a/Engine/Source/Runtime/AdvancedWidgets/Private/Widgets/PropertyViewer/SPropertyViewerImpl.cpp
+++ b/Engine/Source/Runtime/AdvancedWidgets/Private/Widgets/PropertyViewer/SPropertyViewerImpl.cpp
@@ -163,11 +163,15 @@ FPropertyPath FTreeNode::GetPropertyPath() const
 
 		if (TSharedPtr<FTreeNode> ParentNodePin = CurrentNode->ParentNode.Pin())
 		{
+			// The property path are temporary for function
+			if (ParentNodePin->GetField().Get<UFunction>())
+			{
+				return FPropertyPath();
+			}
 			CurrentNode = ParentNodePin.Get();
 		}
 		else if (ContainerPin == nullptr)
 		{
-			ensureMsgf(false, TEXT("The tree is not owned by a container"));
 			return FPropertyPath();
 		}
 	}
@@ -304,9 +308,10 @@ void FTreeNode::BuildChildNodesRecursive(IFieldIterator& FieldIterator, IFieldEx
 	}
 	else if (const UFunction* FunctionPtr = Function.Get())
 	{
-		if (FieldExpander.CanExpandFunction(FunctionPtr))
+		TOptional<const UStruct*> StructToExpand = FieldExpander.GetExpandedFunction(FunctionPtr);
+		if (StructToExpand.IsSet())
 		{
-			ChildStructType = FunctionPtr;
+			ChildStructType = StructToExpand.GetValue();
 		}
 	}
 	else if (const TSharedPtr<FContainer> ContainerPin = Container.Pin())
diff --git a/Engine/Source/Runtime/AdvancedWidgets/Public/Framework/PropertyViewer/IFieldExpander.h b/Engine/Source/Runtime/AdvancedWidgets/Public/Framework/PropertyViewer/IFieldExpander.h
--- a/Engine/Source/Runtime/AdvancedWidgets/Public/Framework/PropertyViewer/IFieldExpander.h
+++ b/Engine/Source/Runtime/AdvancedWidgets/Public/Framework/PropertyViewer/IFieldExpander.h
@@ -3,6 +3,7 @@
 #pragma once
 
 #include "CoreMinimal.h"
+#include "UObject/Class.h"
 #include "UObject/Object.h"
 #include "UObject/UnrealType.h"
 
@@ -23,8 +24,15 @@ public:
 	virtual TOptional<const UClass*> CanExpandObject(const FObjectPropertyBase* Property, const UObject* Instance) const = 0;
 	/** @return true if the struct property can be expended. */
 	virtual bool CanExpandScriptStruct(const FStructProperty* StructProperty) const = 0;
-	virtual bool CanExpandFunction(const UFunction*) const = 0;
+	/** @return the struct that can be expanded from the function. ie. This can be the function itself (to see the arguments) or a return property. */
+	virtual TOptional<const UStruct*> GetExpandedFunction(const UFunction* Function) const = 0;
 	virtual ~IFieldExpander() = default;
+
+	UE_DEPRECATED(5.2, "CanExpandFunction is deprecated and GetExpandedFunction should be used instead.")
+	virtual bool CanExpandFunction(const UFunction* Function) const
+	{
+		return GetExpandedFunction(Function).IsSet();
+	}
 };
 
 class ADVANCEDWIDGETS_API FFieldExpander_Default : public IFieldExpander
@@ -39,6 +47,14 @@ public:
 		RequireValidInstance = 1 << 2,
 	};
 
+
+	enum class EFunctionExpand : uint8
+	{
+		None = 0,
+		FunctionProperties,
+		FunctionReturnProperty,
+	};
+
 	virtual TOptional<const UClass*> CanExpandObject(const FObjectPropertyBase* Property, const UObject* Instance) const override
 	{
 		const bool bUseInstanceClass = EnumHasAllFlags(ExpandObject, EObjectExpandFlag::UseInstanceClass);
@@ -53,8 +69,27 @@ public:
 		}
 		return Property->PropertyClass;
 	}
-	virtual bool CanExpandScriptStruct(const FStructProperty*) const override { return bExpandScriptStruct; }
-	virtual bool CanExpandFunction(const UFunction*) const override { return bExpandFunction; }
+
+	virtual bool CanExpandScriptStruct(const FStructProperty*) const override
+	{
+		return bExpandScriptStruct;
+	}
+
+	virtual TOptional<const UStruct*> GetExpandedFunction(const UFunction* Function) const override
+	{
+		if (Function && ExpandFunction == EFunctionExpand::FunctionProperties)
+		{
+			return Function;
+		}
+		else if (Function && ExpandFunction == EFunctionExpand::FunctionReturnProperty)
+		{
+			if (const FObjectPropertyBase* ObjectProperty = CastField<const FObjectPropertyBase>(Function->GetReturnProperty()))
+			{
+				return ObjectProperty->PropertyClass;
+			}
+		}
+		return TOptional<const UStruct*>();
+	}
 
 	void SetExpandObject(EObjectExpandFlag InExpandObject)
 	{
@@ -65,16 +100,22 @@ public:
 	{
 		bExpandScriptStruct = bInExpandScriptStruct;
 	}
-	
+
+	UE_DEPRECATED(5.2, "Use the SetExpandFunction with the EExpandFunctionType")
 	void SetExpandFunction(bool bInExpandFunction)
 	{
-		bExpandFunction = bInExpandFunction;
+		ExpandFunction = bInExpandFunction ? EFunctionExpand::FunctionProperties : EFunctionExpand::None;
+	}
+
+	void SetExpandFunction(EFunctionExpand InExpandFunction)
+	{
+		ExpandFunction = InExpandFunction;
 	}
 
 private:
 	EObjectExpandFlag ExpandObject = EObjectExpandFlag::None;
 	bool bExpandScriptStruct = false;
-	bool bExpandFunction = false;
+	EFunctionExpand ExpandFunction = EFunctionExpand::None;
 };
 ENUM_CLASS_FLAGS(FFieldExpander_Default::EObjectExpandFlag);
 
