MVVM - Field selector clear now works for viewmodel properties.

It was erroneously not clearing if FixedSource wasn't set (ie. a widget field), but the intent was to clear only the path when editing a widget property, and the entire path when editing anything else.

#rnx

diff --git a/Engine/Plugins/Runtime/ModelViewViewModel/Source/ModelViewViewModelEditor/Private/Widgets/SMVVMFieldSelector.cpp b/Engine/Plugins/Runtime/ModelViewViewModel/Source/ModelViewViewModelEditor/Private/Widgets/SMVVMFieldSelector.cpp
--- a/Engine/Plugins/Runtime/ModelViewViewModel/Source/ModelViewViewModelEditor/Private/Widgets/SMVVMFieldSelector.cpp
+++ b/Engine/Plugins/Runtime/ModelViewViewModel/Source/ModelViewViewModelEditor/Private/Widgets/SMVVMFieldSelector.cpp
@@ -251,10 +251,22 @@ bool SFieldSelector::IsClearEnabled() const
 
 FReply SFieldSelector::OnClearClicked()
 {
+	FMVVMBlueprintPropertyPath NewProperty;
+	
 	if (FixedSource.IsSet())
 	{
-		SetPropertySelection(FMVVMBlueprintPropertyPath());
+		FBindingSource Source = FixedSource.GetValue();
+		if (Source.ViewModelId.IsValid())
+		{
+			NewProperty.SetViewModelId(Source.ViewModelId);
+		}
+		else 
+		{
+			NewProperty.SetWidgetName(Source.Name);
+		}
 	}
+
+	SetPropertySelection(NewProperty);
 	SetConversionFunctionSelection(nullptr);
 	return FReply::Handled();
 }
