MVVM: Simple empty structure that could be use to notify from an event.

diff --git a/Engine/Plugins/Runtime/ModelViewViewModel/Source/ModelViewViewModel/Public/Types/MVVMEventField.h b/Engine/Plugins/Runtime/ModelViewViewModel/Source/ModelViewViewModel/Public/Types/MVVMEventField.h
new file mode 100644
--- /dev/null
+++ b/Engine/Plugins/Runtime/ModelViewViewModel/Source/ModelViewViewModel/Public/Types/MVVMEventField.h
@@ -0,0 +1,26 @@
+// Copyright Epic Games, Inc. All Rights Reserved.
+
+#pragma once
+
+#include "MVVMEventField.generated.h"
+
+/**
+ * Generic structure to notify when an event occurs.
+ * 
+ * class UMyViewmodel : public UMVVMViewModelBase
+ * {
+ *   UPROPERTY(FieldNotify)
+ *   FMVVMEventField SomeEvent;
+ *
+ *   void OnSomeEvent()
+ *   {
+ *     UE_MVVM_BROADCAST_FIELD_VALUE_CHANGED(SomeEvent);
+ *   }
+ * };
+};
+ */
+USTRUCT(BlueprintType)
+struct FMVVMEventField
+{
+	GENERATED_BODY()
+};
