MVVM: Selection fixes. Can now select the viewmodel self. Fix the reverse complexe function.

diff --git a/Engine/Plugins/Runtime/ModelViewViewModel/Source/ModelViewViewModel/Private/View/MVVMView.cpp b/Engine/Plugins/Runtime/ModelViewViewModel/Source/ModelViewViewModel/Private/View/MVVMView.cpp
--- a/Engine/Plugins/Runtime/ModelViewViewModel/Source/ModelViewViewModel/Private/View/MVVMView.cpp
+++ b/Engine/Plugins/Runtime/ModelViewViewModel/Source/ModelViewViewModel/Private/View/MVVMView.cpp
@@ -151,9 +151,9 @@ bool UMVVMView::SetViewModel(FName ViewModelName, TScriptInterface<INotifyFieldV
 			if (NewValue.GetObject())
 			{
 				// Register back any binding that was previously enabled
-				if (PreviousValue)
+				if (BindingToReenabled.Num() > 0)
 				{
-					for (int32 Index = 0; Index < BindingToReenabled.Num(); ++Index)
+					for (int32 Index : BindingToReenabled)
 					{
 						const FMVVMViewClass_CompiledBinding& Binding = CompiledBindings[Index];
 						EnableLibraryBinding(Binding, Index);
diff --git a/Engine/Plugins/Runtime/ModelViewViewModel/Source/ModelViewViewModelBlueprint/Private/Bindings/MVVMConversionFunctionHelper.cpp b/Engine/Plugins/Runtime/ModelViewViewModel/Source/ModelViewViewModelBlueprint/Private/Bindings/MVVMConversionFunctionHelper.cpp
--- a/Engine/Plugins/Runtime/ModelViewViewModel/Source/ModelViewViewModelBlueprint/Private/Bindings/MVVMConversionFunctionHelper.cpp
+++ b/Engine/Plugins/Runtime/ModelViewViewModel/Source/ModelViewViewModelBlueprint/Private/Bindings/MVVMConversionFunctionHelper.cpp
@@ -179,17 +179,17 @@ namespace Private
 		{
 			if (const FObjectProperty* Property = CastField<FObjectProperty>(RootRef.ResolveMember<FProperty>(WidgetBlueprint->SkeletonGeneratedClass)))
 			{
-				if (Property->PropertyClass->ImplementsInterface(UNotifyFieldValueChanged::StaticClass()))
+				if (Property->PropertyClass->IsChildOf<UWidget>() || Property->PropertyClass->IsChildOf<UWidgetBlueprint>())
+				{
+					ResultPath.SetWidgetName(Property->GetFName());
+				}
+				else if (Property->PropertyClass->ImplementsInterface(UNotifyFieldValueChanged::StaticClass()))
 				{
 					if (const FMVVMBlueprintViewModelContext* ViewModel = BlueprintView->FindViewModel(Property->GetFName()))
 					{
 						ResultPath.SetViewModelId(ViewModel->GetViewModelId());
 					}
 				}
-				else if (Property->PropertyClass->IsChildOf<UWidget>() || Property->PropertyClass->IsChildOf<UWidgetBlueprint>())
-				{
-					ResultPath.SetWidgetName(Property->GetFName());
-				}
 			}
 		}
 
diff --git a/Engine/Plugins/Runtime/ModelViewViewModel/Source/ModelViewViewModelBlueprint/Private/MVVMBlueprintViewBinding.cpp b/Engine/Plugins/Runtime/ModelViewViewModel/Source/ModelViewViewModelBlueprint/Private/MVVMBlueprintViewBinding.cpp
--- a/Engine/Plugins/Runtime/ModelViewViewModel/Source/ModelViewViewModelBlueprint/Private/MVVMBlueprintViewBinding.cpp
+++ b/Engine/Plugins/Runtime/ModelViewViewModel/Source/ModelViewViewModelBlueprint/Private/MVVMBlueprintViewBinding.cpp
@@ -37,17 +37,12 @@ namespace UE::MVVM::Private
 			PathBuilder << TEXT("<none>");
 		}
 
-		PathBuilder << TEXT(".");
-
 		const FString PropertyPath = ViewModelPath.GetBasePropertyPath();
 		if (!PropertyPath.IsEmpty())
 		{
+			PathBuilder << TEXT(".");
 			PathBuilder << PropertyPath;
 		}
-		else
-		{
-			PathBuilder << TEXT("<none>");
-		}
 	}
 	
 	FString GetBindingViewModelName(const UWidgetBlueprint* WidgetBlueprint, const FMVVMBlueprintViewBinding& Binding)
@@ -145,17 +140,12 @@ namespace UE::MVVM::Private
 			PathBuilder << WidgetPath.GetWidgetName();
 		}
 
-		PathBuilder << TEXT(".");
-
 		const FString PropertyPath = WidgetPath.GetBasePropertyPath();
 		if (!PropertyPath.IsEmpty())
 		{
+			PathBuilder << TEXT(".");
 			PathBuilder << PropertyPath;
 		}
-		else
-		{
-			PathBuilder << TEXT("<none>");
-		}
 	}
 
 	FString GetBindingWidgetName(const UWidgetBlueprint* WidgetBlueprint, const FMVVMBlueprintViewBinding& Binding)
diff --git a/Engine/Plugins/Runtime/ModelViewViewModel/Source/ModelViewViewModelBlueprint/Private/MVVMViewBlueprintCompiler.cpp b/Engine/Plugins/Runtime/ModelViewViewModel/Source/ModelViewViewModelBlueprint/Private/MVVMViewBlueprintCompiler.cpp
--- a/Engine/Plugins/Runtime/ModelViewViewModel/Source/ModelViewViewModelBlueprint/Private/MVVMViewBlueprintCompiler.cpp
+++ b/Engine/Plugins/Runtime/ModelViewViewModel/Source/ModelViewViewModelBlueprint/Private/MVVMViewBlueprintCompiler.cpp
@@ -173,14 +173,17 @@ void FMVVMViewBlueprintCompiler::CreateVariables(const FWidgetBlueprintCompilerC
 		CreateFunctionsDeclaration(Context, BlueprintView);
 	}
 
-	auto CreateVariable = [&Context](const FCompilerSourceContext& SourceContext) -> FProperty*
+	auto CreateVariable = [&Context](const FCompilerUserWidgetPropertyContext& SourceContext) -> FProperty*
 	{
-		FEdGraphPinType ViewModelPinType(UEdGraphSchema_K2::PC_Object, NAME_None, SourceContext.Class, EPinContainerType::None, false, FEdGraphTerminalType());
-		FProperty* NewProperty = Context.CreateVariable(SourceContext.PropertyName, ViewModelPinType);
+		FEdGraphPinType NewPropertyPinType(UEdGraphSchema_K2::PC_Object, NAME_None, SourceContext.Class, EPinContainerType::None, false, FEdGraphTerminalType());
+		FProperty* NewProperty = Context.CreateVariable(SourceContext.PropertyName, NewPropertyPinType);
 		if (NewProperty != nullptr)
 		{
-			NewProperty->SetPropertyFlags(CPF_BlueprintVisible | CPF_BlueprintReadOnly | CPF_RepSkip
-				| CPF_Transient | CPF_DuplicateTransient);
+			NewProperty->SetPropertyFlags(CPF_BlueprintVisible | CPF_RepSkip | CPF_Transient | CPF_DuplicateTransient);
+			if (SourceContext.BlueprintSetter.IsEmpty())
+			{
+				NewProperty->SetPropertyFlags(CPF_BlueprintReadOnly);
+			}
 			NewProperty->SetPropertyFlags(SourceContext.bExposeOnSpawn ? CPF_ExposeOnSpawn : CPF_DisableEditOnInstance);
 
 #if WITH_EDITOR
@@ -205,7 +208,7 @@ void FMVVMViewBlueprintCompiler::CreateVariables(const FWidgetBlueprintCompilerC
 		return NewProperty;
 	};
 
-	for (FCompilerSourceContext& SourceContext : CompilerSourceContexts)
+	for (FCompilerUserWidgetPropertyContext& SourceContext : CompilerSourceContexts)
 	{
 		SourceContext.Field = BindingHelper::FindFieldByName(Context.GetSkeletonGeneratedClass(), FMVVMBindingName(SourceContext.PropertyName));
 
@@ -348,7 +351,7 @@ void FMVVMViewBlueprintCompiler::CreateSourceLists(const FWidgetBlueprintCompile
 		}
 		checkf(FoundSourceCreatorContextIndex != INDEX_NONE, TEXT("The viewmodel was added after the skeleton was created?"));
 
-		FCompilerSourceContext SourceVariable;
+		FCompilerUserWidgetPropertyContext SourceVariable;
 		SourceVariable.Class = ViewModelContext.GetViewModelClass();
 		SourceVariable.PropertyName = ViewModelContext.GetViewModelName();
 		SourceVariable.DisplayName = ViewModelContext.GetDisplayName();
@@ -419,11 +422,11 @@ void FMVVMViewBlueprintCompiler::CreateSourceLists(const FWidgetBlueprintCompile
 					}
 
 					UWidget* Widget = *WidgetPtr;
-					FCompilerSourceContext SourceVariable;
+					FCompilerUserWidgetPropertyContext SourceVariable;
 					SourceVariable.Class = Widget->GetClass();
 					SourceVariable.PropertyName = PropertyPath.GetWidgetName();
-					SourceVariable.DisplayName = Widget->GetLabelText();
-					SourceVariable.CategoryName = !Widget->GetCategoryName().IsEmpty() ? Widget->GetCategoryName() : DefaultWidgetCategory.ToString();
+					SourceVariable.DisplayName = FText::FromString(Widget->GetDisplayLabel());
+					SourceVariable.CategoryName = TEXT("Widget");
 					Self->CompilerSourceContexts.Emplace(MoveTemp(SourceVariable));
 				}
 			}
@@ -670,7 +673,7 @@ bool FMVVMViewBlueprintCompiler::PreCompileBindingSources(UWidgetBlueprintGenera
 				return false;
 			}
 
-			if (!BindingSourceContext.bIsRootWidget && BindingSourceContext.CompilerSourceContextIndex == INDEX_NONE)
+			if (!BindingSourceContext.bIsRootWidget && BindingSourceContext.UserWidgetPropertyContextIndex == INDEX_NONE)
 			{
 				AddErrorForBinding(Binding, LOCTEXT("BindingInvalidSource", "Internal error. The binding could not find its source."), ArgumentName);
 				return false;
@@ -710,7 +713,7 @@ bool FMVVMViewBlueprintCompiler::PreCompileBindingSources(UWidgetBlueprintGenera
 		{
 			CreateSourcesForConversionFunction(false);
 
-			if (!Binding.WidgetPath.IsEmpty())
+			if (!Binding.WidgetPath.IsEmpty() && Binding.WidgetPath.HasPaths())
 			{
 				if (!CreateSourceContextForPropertyPath(Binding.WidgetPath, false))
 				{
@@ -1042,7 +1045,7 @@ bool FMVVMViewBlueprintCompiler::PreCompileBindings(UWidgetBlueprintGeneratedCla
 
 		FCompilerBinding NewBinding = AddBindingResult.StealValue();
 		NewBinding.BindingIndex = BindingSourceContext.BindingIndex;
-		NewBinding.CompilerSourceContextIndex = BindingSourceContext.CompilerSourceContextIndex;
+		NewBinding.CompilerSourceContextIndex = BindingSourceContext.UserWidgetPropertyContextIndex;
 		NewBinding.FieldIdHandle = AddFieldResult.StealValue();
 		NewBinding.bSourceIsUserWidget = BindingSourceContext.bIsRootWidget;
 		NewBinding.bFieldIdNeeded = !IsOneTimeBinding(Binding.BindingType);
@@ -1149,11 +1152,11 @@ TValueOrError<FMVVMViewBlueprintCompiler::FBindingSourceContext, FText> FMVVMVie
 		const FMVVMBlueprintViewModelContext* SourceViewModelContext = BlueprintView->FindViewModel(PropertyPath.GetViewModelId());
 		check(SourceViewModelContext);
 		const FName SourceName = SourceViewModelContext->GetViewModelName();
-		Result.CompilerSourceContextIndex = CompilerSourceContexts.IndexOfByPredicate([SourceName](const FCompilerSourceContext& Other) { return Other.PropertyName == SourceName; });
-		check(Result.CompilerSourceContextIndex != INDEX_NONE);
+		Result.UserWidgetPropertyContextIndex = CompilerSourceContexts.IndexOfByPredicate([SourceName](const FCompilerUserWidgetPropertyContext& Other) { return Other.PropertyName == SourceName; });
+		check(Result.UserWidgetPropertyContextIndex != INDEX_NONE);
 
-		Result.SourceClass = CompilerSourceContexts[Result.CompilerSourceContextIndex].Class;
-		Result.PropertyPath = CreatePropertyPath(Class, CompilerSourceContexts[Result.CompilerSourceContextIndex].PropertyName, PropertyPath.GetFields());
+		Result.SourceClass = CompilerSourceContexts[Result.UserWidgetPropertyContextIndex].Class;
+		Result.PropertyPath = CreatePropertyPath(Class, CompilerSourceContexts[Result.UserWidgetPropertyContextIndex].PropertyName, PropertyPath.GetFields());
 	}
 	else if (PropertyPath.IsFromWidget())
 	{
@@ -1161,16 +1164,16 @@ TValueOrError<FMVVMViewBlueprintCompiler::FBindingSourceContext, FText> FMVVMVie
 		Result.bIsRootWidget = SourceName == Class->ClassGeneratedBy->GetFName();
 		if (Result.bIsRootWidget)
 		{
-			Result.CompilerSourceContextIndex = INDEX_NONE;
+			Result.UserWidgetPropertyContextIndex = INDEX_NONE;
 			Result.SourceClass = Class->ClassGeneratedBy->GetClass();
 			Result.PropertyPath = CreatePropertyPath(Class, FName(), PropertyPath.GetFields());
 		}
 		else
 		{
-			Result.CompilerSourceContextIndex = CompilerSourceContexts.IndexOfByPredicate([SourceName](const FCompilerSourceContext& Other) { return Other.PropertyName == SourceName; });
-			check(Result.CompilerSourceContextIndex != INDEX_NONE);
-			Result.SourceClass = CompilerSourceContexts[Result.CompilerSourceContextIndex].Class;
-			Result.PropertyPath = CreatePropertyPath(Class, CompilerSourceContexts[Result.CompilerSourceContextIndex].PropertyName, PropertyPath.GetFields());
+			Result.UserWidgetPropertyContextIndex = CompilerSourceContexts.IndexOfByPredicate([SourceName](const FCompilerUserWidgetPropertyContext& Other) { return Other.PropertyName == SourceName; });
+			check(Result.UserWidgetPropertyContextIndex != INDEX_NONE);
+			Result.SourceClass = CompilerSourceContexts[Result.UserWidgetPropertyContextIndex].Class;
+			Result.PropertyPath = CreatePropertyPath(Class, CompilerSourceContexts[Result.UserWidgetPropertyContextIndex].PropertyName, PropertyPath.GetFields());
 		}
 	}
 	else
@@ -1195,7 +1198,7 @@ TValueOrError<FMVVMViewBlueprintCompiler::FBindingSourceContext, FText> FMVVMVie
 	if (bCreateADynamicSourceEntry)
 	{
 		Result.SourceClass = BindingInfo.NotifyFieldClass;
-		Result.CompilerSourceContextIndex = INDEX_NONE;
+		Result.UserWidgetPropertyContextIndex = INDEX_NONE;
 		Result.bIsRootWidget = false;
 
 		FString InterfacePath = FieldPathHelper::ToString(BindingInfo.NotifyFieldInterfacePath);
@@ -1217,14 +1220,7 @@ TValueOrError<FMVVMViewBlueprintCompiler::FBindingSourceContext, FText> FMVVMVie
 		Result.PropertyPath.RemoveAt(0, BindingInfo.ViewModelIndex);
 		if (CompilerSourceCreatorContexts[SourceCreatorContextIndex].Type == ECompilerSourceCreatorType::ViewModel)
 		{
-			Result.CompilerSourceContextIndex = SourceCreatorContextIndex;
-		}
-	}
-	else
-	{
-		if (Result.SourceClass != BindingInfo.NotifyFieldClass)
-		{
-			return MakeError(LOCTEXT("ClassDoesNotMatch", "Internal Error. The class doesn't match"));
+			Result.UserWidgetPropertyContextIndex = SourceCreatorContextIndex;
 		}
 	}
 
@@ -1245,7 +1241,7 @@ TArray<FMVVMConstFieldVariant> FMVVMViewBlueprintCompiler::CreateBindingDestinat
 		const FMVVMBlueprintViewModelContext* SourceViewModelContext = BlueprintView->FindViewModel(PropertyPath.GetViewModelId());
 		check(SourceViewModelContext);
 		FName DestinationName = SourceViewModelContext->GetViewModelName();
-		const int32 DestinationVariableContextIndex = CompilerSourceContexts.IndexOfByPredicate([DestinationName](const FCompilerSourceContext& Other) { return Other.PropertyName == DestinationName; });
+		const int32 DestinationVariableContextIndex = CompilerSourceContexts.IndexOfByPredicate([DestinationName](const FCompilerUserWidgetPropertyContext& Other) { return Other.PropertyName == DestinationName; });
 		check(DestinationVariableContextIndex != INDEX_NONE);
 
 		return CreatePropertyPath(Class, CompilerSourceContexts[DestinationVariableContextIndex].PropertyName, PropertyPath.GetFields());
@@ -1261,7 +1257,7 @@ TArray<FMVVMConstFieldVariant> FMVVMViewBlueprintCompiler::CreateBindingDestinat
 		}
 		else
 		{
-			const int32 DestinationVariableContextIndex = CompilerSourceContexts.IndexOfByPredicate([DestinationName](const FCompilerSourceContext& Other) { return Other.PropertyName == DestinationName; });
+			const int32 DestinationVariableContextIndex = CompilerSourceContexts.IndexOfByPredicate([DestinationName](const FCompilerUserWidgetPropertyContext& Other) { return Other.PropertyName == DestinationName; });
 			if (ensureAlwaysMsgf(DestinationVariableContextIndex != INDEX_NONE, TEXT("Could not find source context for destination '%s'"), *DestinationName.ToString()))
 			{
 				return CreatePropertyPath(Class, CompilerSourceContexts[DestinationVariableContextIndex].PropertyName, PropertyPath.GetFields());
diff --git a/Engine/Plugins/Runtime/ModelViewViewModel/Source/ModelViewViewModelBlueprint/Private/MVVMViewBlueprintCompiler.h b/Engine/Plugins/Runtime/ModelViewViewModel/Source/ModelViewViewModelBlueprint/Private/MVVMViewBlueprintCompiler.h
--- a/Engine/Plugins/Runtime/ModelViewViewModel/Source/ModelViewViewModelBlueprint/Private/MVVMViewBlueprintCompiler.h
+++ b/Engine/Plugins/Runtime/ModelViewViewModel/Source/ModelViewViewModelBlueprint/Private/MVVMViewBlueprintCompiler.h
@@ -25,7 +25,7 @@ namespace UE::MVVM::Private
 struct FMVVMViewBlueprintCompiler
 {
 private:
-	struct FCompilerSourceContext;
+	struct FCompilerUserWidgetPropertyContext;
 	struct FCompilerSourceCreatorContext;
 	struct FCompilerBinding;
 	struct FBindingSourceContext;
@@ -81,7 +81,12 @@ private:
 	bool IsPropertyPathValid(TArrayView<const FMVVMConstFieldVariant> PropertyPath) const;
 
 private:
-	struct FCompilerSourceContext
+	/**
+	 * Describe the Property that is needed to start a PropertyPath on the UserWidget.
+	 * They can be a Widget, a viewmodel, or any object owned by the UserWidget.
+	 * The property will be created if it doesn't exist yet.
+	 */
+	struct FCompilerUserWidgetPropertyContext
 	{
 		UClass* Class = nullptr;
 		FName PropertyName;
@@ -92,8 +97,11 @@ private:
 
 		bool bExposeOnSpawn = false;
 	};
-	TArray<FCompilerSourceContext> CompilerSourceContexts;
+	TArray<FCompilerUserWidgetPropertyContext> CompilerSourceContexts;
 
+	/** 
+	 * Describe the data initialize the view's properties/viewmodels.
+	 */
 	enum class ECompilerSourceCreatorType
 	{
 		ViewModel,
@@ -108,6 +116,10 @@ private:
 	};
 	TArray<FCompilerSourceCreatorContext> CompilerSourceCreatorContexts;
 
+	/** 
+	 * Describe a binding for the compiler.
+	 * 2 FCompilerBinding will be generated if the binding is TwoWays.
+	 */
 	enum class ECompilerBindingType
 	{
 		PropertyBinding, // normal binding
@@ -135,21 +147,25 @@ private:
 	};
 	TArray<FCompilerBinding> CompilerBindings;
 
+	/**
+	 * Source path of a binding that contains a FieldId that we can register/bind to
+	 */
 	struct FBindingSourceContext
 	{
 		int32 BindingIndex = INDEX_NONE;
 		bool bIsForwardBinding = false;
 
 		UClass* SourceClass = nullptr;
-		// if its the destination then the Field is not necessary
+		// The property that are registering to.
 		FFieldNotificationId FieldId;
+		// The path always start with the UserWidget as self.
 		// Viewmodel.Field.SubProperty.SubProperty
-		// or Widget.Field.SubProperty
-		// or Field.SubProperty (if the widget is the UserWidget)
+		// or Widget.Field.SubProperty.SubProperty
+		// or Field.SubProperty.SubProperty
 		TArray<UE::MVVM::FMVVMConstFieldVariant> PropertyPath;
 
 		// The source if it's a property
-		int32 CompilerSourceContextIndex = INDEX_NONE;
+		int32 UserWidgetPropertyContextIndex = INDEX_NONE;
 		// The source is the UserWidget itself
 		bool bIsRootWidget = false;
 	};
diff --git a/Engine/Plugins/Runtime/ModelViewViewModel/Source/ModelViewViewModelBlueprint/Public/MVVMPropertyPath.h b/Engine/Plugins/Runtime/ModelViewViewModel/Source/ModelViewViewModelBlueprint/Public/MVVMPropertyPath.h
--- a/Engine/Plugins/Runtime/ModelViewViewModel/Source/ModelViewViewModelBlueprint/Public/MVVMPropertyPath.h
+++ b/Engine/Plugins/Runtime/ModelViewViewModel/Source/ModelViewViewModelBlueprint/Public/MVVMPropertyPath.h
@@ -214,6 +214,11 @@ private:
 #endif
 
 public:
+	bool HasPaths() const
+	{
+		return Paths.Num() > 0;
+	}
+
 	/** Get the binding name, resolves reference deprecation / redirectors / etc before returning */
 	TArray<FName> GetPaths() const
 	{
diff --git a/Engine/Plugins/Runtime/ModelViewViewModel/Source/ModelViewViewModelEditor/Private/Types/MVVMBindingSource.cpp b/Engine/Plugins/Runtime/ModelViewViewModel/Source/ModelViewViewModelEditor/Private/Types/MVVMBindingSource.cpp
--- a/Engine/Plugins/Runtime/ModelViewViewModel/Source/ModelViewViewModelEditor/Private/Types/MVVMBindingSource.cpp
+++ b/Engine/Plugins/Runtime/ModelViewViewModel/Source/ModelViewViewModelEditor/Private/Types/MVVMBindingSource.cpp
@@ -7,6 +7,25 @@
 #include "MVVMEditorSubsystem.h"
 #include "WidgetBlueprint.h"
 
+FMVVMBindingName UE::MVVM::FBindingSource::ToBindingName(const UWidgetBlueprint* WidgetBlueprint) const
+{
+	if (ViewModelId.IsValid())
+	{
+		if (UMVVMBlueprintView* View = GEditor->GetEditorSubsystem<UMVVMEditorSubsystem>()->GetView(WidgetBlueprint))
+		{
+			if (const FMVVMBlueprintViewModelContext* ViewModel = View->FindViewModel(ViewModelId))
+			{
+				return FMVVMBindingName(ViewModel->GetViewModelName());
+			}
+		}
+		return FMVVMBindingName();
+	}
+	else
+	{
+		return FMVVMBindingName(Name);
+	}
+}
+
 UE::MVVM::FBindingSource UE::MVVM::FBindingSource::CreateForWidget(const UWidgetBlueprint* WidgetBlueprint, FName WidgetName)
 {
 	UE::MVVM::FBindingSource Source;
diff --git a/Engine/Plugins/Runtime/ModelViewViewModel/Source/ModelViewViewModelEditor/Private/Types/MVVMBindingSource.h b/Engine/Plugins/Runtime/ModelViewViewModel/Source/ModelViewViewModelEditor/Private/Types/MVVMBindingSource.h
--- a/Engine/Plugins/Runtime/ModelViewViewModel/Source/ModelViewViewModelEditor/Private/Types/MVVMBindingSource.h
+++ b/Engine/Plugins/Runtime/ModelViewViewModel/Source/ModelViewViewModelEditor/Private/Types/MVVMBindingSource.h
@@ -4,6 +4,7 @@
 
 #include "CoreMinimal.h"
 #include "Framework/Views/TableViewTypeTraits.h"
+#include "Types/MVVMBindingName.h"
 
 class UWidgetBlueprint;
 
@@ -46,6 +47,8 @@ namespace UE::MVVM
 			DisplayName = FText::GetEmpty();
 		}
 
+		FMVVMBindingName ToBindingName(const UWidgetBlueprint* WidgetBlueprint) const;
+
 		static FBindingSource CreateForWidget(const UWidgetBlueprint* WidgetBlueprint, FName WidgetName);
 		static FBindingSource CreateForViewModel(const UWidgetBlueprint* WidgetBlueprint, FGuid ViewModelId);
 		static FBindingSource CreateForViewModel(const UWidgetBlueprint* WidgetBlueprint, FName ViewModelName);
diff --git a/Engine/Plugins/Runtime/ModelViewViewModel/Source/ModelViewViewModelEditor/Private/Widgets/SMVVMFieldSelector.cpp b/Engine/Plugins/Runtime/ModelViewViewModel/Source/ModelViewViewModelEditor/Private/Widgets/SMVVMFieldSelector.cpp
--- a/Engine/Plugins/Runtime/ModelViewViewModel/Source/ModelViewViewModelEditor/Private/Widgets/SMVVMFieldSelector.cpp
+++ b/Engine/Plugins/Runtime/ModelViewViewModel/Source/ModelViewViewModelEditor/Private/Widgets/SMVVMFieldSelector.cpp
@@ -88,7 +88,7 @@ void SFieldSelector::Construct(const FArguments& InArgs, const UWidgetBlueprint*
 			+ SOverlay::Slot()
 			[
 				SAssignNew(SourceEntryBox, SHorizontalBox)
-				.Visibility_Lambda([this]() { return CachedSelectedConversionFunction == nullptr && CachedSelectedField.GetFields().Num() > 0 ? EVisibility::Visible : EVisibility::Collapsed; })
+				.Visibility_Lambda([this]() { return CachedSelectedConversionFunction == nullptr && !CachedSelectedField.IsEmpty() ? EVisibility::Visible : EVisibility::Collapsed; })
 			]
 
 			// is a conversion function set?
@@ -167,7 +167,8 @@ void SFieldSelector::Construct(const FArguments& InArgs, const UWidgetBlueprint*
 			[
 				SNew(SImage)
 				.Image(FAppStyle::Get().GetBrush("Icons.ChevronRight"))
-			];					
+				.Visibility_Lambda([this]() { return CachedSelectedField.HasPaths() ? EVisibility::Visible : EVisibility::Collapsed; })
+			];
 	}
 
 	SourceEntryBox->AddSlot()
@@ -211,7 +212,7 @@ bool SFieldSelector::IsSelectEnabled() const
 	if (BindingList.IsValid())
 	{
 		FMVVMBlueprintPropertyPath Path = BindingList->GetSelectedProperty();
-		if ((Path.IsFromViewModel() || Path.IsFromWidget()) && !Path.GetBasePropertyPath().IsEmpty())
+		if (Path.IsFromViewModel() || Path.IsFromWidget())
 		{
 			return true;
 		}
diff --git a/Engine/Plugins/Runtime/ModelViewViewModel/Source/ModelViewViewModelEditor/Private/Widgets/SMVVMFunctionParameter.cpp b/Engine/Plugins/Runtime/ModelViewViewModel/Source/ModelViewViewModelEditor/Private/Widgets/SMVVMFunctionParameter.cpp
--- a/Engine/Plugins/Runtime/ModelViewViewModel/Source/ModelViewViewModelEditor/Private/Widgets/SMVVMFunctionParameter.cpp
+++ b/Engine/Plugins/Runtime/ModelViewViewModel/Source/ModelViewViewModelEditor/Private/Widgets/SMVVMFunctionParameter.cpp
@@ -64,16 +64,7 @@ void SFunctionParameter::Construct(const FArguments& InArgs)
 
 	ValueWidget->SetVisibility(TAttribute<EVisibility>::CreateSP(this, &SFunctionParameter::OnGetVisibility, true));
 
-	bool bFromViewModel;
-	if (bSourceToDestination)
-	{
-		bFromViewModel = UE::MVVM::IsForwardBinding(Binding->BindingType);
-	}
-	else
-	{
-		bFromViewModel = UE::MVVM::IsBackwardBinding(Binding->BindingType);
-	}
-
+	bool bFromViewModel = UE::MVVM::IsForwardBinding(Binding->BindingType);
 	TSharedPtr<SHorizontalBox> HBox;
 
 	ChildSlot
diff --git a/Engine/Plugins/Runtime/ModelViewViewModel/Source/ModelViewViewModelEditor/Private/Widgets/SMVVMViewModelBindingListWidget.cpp b/Engine/Plugins/Runtime/ModelViewViewModel/Source/ModelViewViewModelEditor/Private/Widgets/SMVVMViewModelBindingListWidget.cpp
--- a/Engine/Plugins/Runtime/ModelViewViewModel/Source/ModelViewViewModelEditor/Private/Widgets/SMVVMViewModelBindingListWidget.cpp
+++ b/Engine/Plugins/Runtime/ModelViewViewModel/Source/ModelViewViewModelEditor/Private/Widgets/SMVVMViewModelBindingListWidget.cpp
@@ -379,7 +379,7 @@ void SSourceBindingList::AddSources(TArrayView<const FBindingSource> InSources)
 
 FMVVMBlueprintPropertyPath SSourceBindingList::CreateBlueprintPropertyPath(SPropertyViewer::FHandle Handle, TArrayView<const FFieldVariant> FieldPath) const
 {
-	if (FieldPath.Num() < 0)
+	if (FieldPath.Num() < 0 || !Handle.IsValid())
 	{
 		return FMVVMBlueprintPropertyPath();
 	}
@@ -402,57 +402,80 @@ FMVVMBlueprintPropertyPath SSourceBindingList::CreateBlueprintPropertyPath(SProp
 		return FMVVMBlueprintPropertyPath();
 	}
 
-	// Backward, test if the object can be access.
-	//The last property can be a struct variable, inside a struct, inside..., inside an object. 
-	bool bPassFilter = false;
-	TSubclassOf<UObject> AccessorClass = WidgetBlueprintPtr ? WidgetBlueprintPtr->GeneratedClass : nullptr;
-	for (int32 Index = FieldPath.Num() - 1; Index >= 0; --Index)
+	TSubclassOf<UObject> AccessorClass = WidgetBlueprintPtr ? WidgetBlueprintPtr->SkeletonGeneratedClass : nullptr;
+	FMVVMBlueprintPropertyPath PropertyPath;
+	if (FieldPath.Num() > 0)
 	{
-		const FFieldVariant& FieldVariant = FieldPath[Index];
-		const UStruct* OwnerStruct = nullptr;
-		FName FieldName;
-		if (const FProperty* Property = FieldVariant.Get<FProperty>())
-		{
-			OwnerStruct = Property->GetOwnerStruct();
-			FieldName = Property->GetFName();
-		}
-		else if (const UFunction* Function = FieldVariant.Get<UFunction>())
+		// Backward, test if the object can be access.
+		//The last property can be a struct variable, inside a struct, inside..., inside an object. 
+		bool bPassFilter = false;
+		for (int32 Index = FieldPath.Num() - 1; Index >= 0; --Index)
 		{
-			OwnerStruct = Function->GetOwnerClass();
-			FieldName = Function->GetFName();
-		}
+			const FFieldVariant& FieldVariant = FieldPath[Index];
+			const UStruct* OwnerStruct = nullptr;
+			FName FieldName;
+			if (const FProperty* Property = FieldVariant.Get<FProperty>())
+			{
+				OwnerStruct = Property->GetOwnerStruct();
+				FieldName = Property->GetFName();
+			}
+			else if (const UFunction* Function = FieldVariant.Get<UFunction>())
+			{
+				OwnerStruct = Function->GetOwnerClass();
+				FieldName = Function->GetFName();
+			}
 
-		if (const UClass* OwnerClass = Cast<const UClass>(OwnerStruct))
-		{
-			FMVVMAvailableBinding Binding = GEngine->GetEngineSubsystem<UMVVMSubsystem>()->GetAvailableBinding(OwnerClass, FMVVMBindingName(FieldName), AccessorClass);
-			if (Binding.IsValid())
+			if (const UClass* OwnerClass = Cast<const UClass>(OwnerStruct))
 			{
-				bPassFilter = Private::PassFilter(Binding, OwnerClass, FieldIterator->GetFieldVisibilityFlags(), FieldIterator->GetAssignableTo(), true).IsSet();
+				FMVVMAvailableBinding Binding = GEngine->GetEngineSubsystem<UMVVMSubsystem>()->GetAvailableBinding(OwnerClass, FMVVMBindingName(FieldName), AccessorClass);
+				if (Binding.IsValid())
+				{
+					bPassFilter = Private::PassFilter(Binding, OwnerClass, FieldIterator->GetFieldVisibilityFlags(), FieldIterator->GetAssignableTo(), true).IsSet();
+				}
+				break;
 			}
-			break;
 		}
-	}
 
-	FMVVMBlueprintPropertyPath PropertyPath;
-	if (bPassFilter)
-	{
-		if (Source->Key.ViewModelId.IsValid())
-		{
-			PropertyPath.SetViewModelId(Source->Key.ViewModelId);
-		}
-		else
+		if (bPassFilter)
 		{
-			PropertyPath.SetWidgetName(Source->Key.Name);
-		}
+			if (Source->Key.ViewModelId.IsValid())
+			{
+				PropertyPath.SetViewModelId(Source->Key.ViewModelId);
+			}
+			else
+			{
+				PropertyPath.SetWidgetName(Source->Key.Name);
+			}
 
-		PropertyPath.ResetBasePropertyPath();
+			PropertyPath.ResetBasePropertyPath();
 
-		for (const FFieldVariant& Field : FieldPath)
+			for (const FFieldVariant& Field : FieldPath)
+			{
+				PropertyPath.AppendBasePropertyPath(FMVVMConstFieldVariant(Field));
+			}
+		}
+	}
+	else if(AccessorClass.Get())
+	{
+		FMVVMBindingName BindingName = Source->Key.ToBindingName(WidgetBlueprintPtr);
+		FMVVMAvailableBinding Binding = GEngine->GetEngineSubsystem<UMVVMSubsystem>()->GetAvailableBinding(AccessorClass, BindingName, AccessorClass);
+		if (Binding.IsValid())
 		{
-			PropertyPath.AppendBasePropertyPath(FMVVMConstFieldVariant(Field));
+			bool bPassFilter = Private::PassFilter(Binding, AccessorClass, FieldIterator->GetFieldVisibilityFlags(), FieldIterator->GetAssignableTo(), true).IsSet();
+			if (bPassFilter)
+			{
+				if (Source->Key.ViewModelId.IsValid())
+				{
+					PropertyPath.SetViewModelId(Source->Key.ViewModelId);
+				}
+				else
+				{
+					PropertyPath.SetWidgetName(Source->Key.Name);
+				}
+				PropertyPath.ResetBasePropertyPath();
+			}
 		}
 	}
-
 	return PropertyPath;
 }
 
diff --git a/Engine/Source/Editor/UMGEditor/Private/WidgetBlueprintCompiler.cpp b/Engine/Source/Editor/UMGEditor/Private/WidgetBlueprintCompiler.cpp
--- a/Engine/Source/Editor/UMGEditor/Private/WidgetBlueprintCompiler.cpp
+++ b/Engine/Source/Editor/UMGEditor/Private/WidgetBlueprintCompiler.cpp
@@ -1109,13 +1109,10 @@ void FWidgetBlueprintCompilerContext::FinishCompilingClass(UClass* Class)
 
 	Super::FinishCompilingClass(Class);
 
-	if (!bIsSkeletonOnly)
+	CA_ASSUME(BPGClass);
+	if (UUserWidget* UserWidget = Cast<UUserWidget>(BPGClass->GetDefaultObject()))
 	{
-		CA_ASSUME(BPGClass);
-		if (UUserWidget* UserWidget = Cast<UUserWidget>(BPGClass->GetDefaultObject()))
-		{
-			BPGClass->InitializeFieldNotification(UserWidget);
-		}
+		BPGClass->InitializeFieldNotification(UserWidget);
 	}
 
 	UWidgetBlueprintExtension::ForEachExtension(WidgetBlueprint(), [BPGClass](UWidgetBlueprintExtension* InExtension)
