MVVM: Add a debugging interface.

diff --git a/Engine/Plugins/Runtime/ModelViewViewModel/Source/ModelViewViewModel/Private/Debugging/MVVMDebugging.cpp b/Engine/Plugins/Runtime/ModelViewViewModel/Source/ModelViewViewModel/Private/Debugging/MVVMDebugging.cpp
new file mode 100644
--- /dev/null
+++ b/Engine/Plugins/Runtime/ModelViewViewModel/Source/ModelViewViewModel/Private/Debugging/MVVMDebugging.cpp
@@ -0,0 +1,104 @@
+// Copyright Epic Games, Inc. All Rights Reserved.
+
+#include "Debugging/MVVMDebugging.h"
+#include "View/MVVMView.h"
+
+#define LOCTEXT_NAMESPACE "MVVMDebugging"
+
+#if UE_WITH_MVVM_DEBUGGING
+namespace UE::MVVM
+{
+/** */
+FDebugging::FView::FView(const UMVVMView* InView)
+	: UserWidget(InView->GetUserWidget())
+	, View(InView)
+{
+}
+
+FDebugging::FView::FView(const UUserWidget* InUserWidget, const UMVVMView* InView)
+	: UserWidget(InUserWidget)
+	, View(InView)
+{
+}
+
+/** */
+FDebugging::FViewConstructed FDebugging::OnViewConstructed;
+
+void FDebugging::BroadcastViewConstructed(const UMVVMView* InView)
+{
+	OnViewConstructed.Broadcast(InView, FViewConstructedArgs());
+}
+
+/** */
+FDebugging::FViewDestructing FDebugging::OnViewBeginDestruction;
+
+void FDebugging::BroadcastViewBeginDestruction(const UMVVMView* InView)
+{
+	OnViewBeginDestruction.Broadcast(InView, FViewBeginDestructionArgs());
+}
+
+/** */
+FDebugging::FLibraryBindingRegisteredArgs::FLibraryBindingRegisteredArgs(const FMVVMViewClass_CompiledBinding& InBinding, ERegisterLibraryBindingResult InResult)
+	: Binding(InBinding)
+	, Result(InResult)
+{
+}
+
+FDebugging::FLibraryBindingRegistered FDebugging::OnLibraryBindingRegistered;
+
+void FDebugging::BroadcastLibraryBindingRegistered(const UMVVMView* InView, const FLibraryBindingRegisteredArgs& Args)
+{
+	check(InView);
+	OnLibraryBindingRegistered.Broadcast(FView(InView), Args);
+}
+
+void FDebugging::BroadcastLibraryBindingRegistered(const UMVVMView* InView, const FMVVMViewClass_CompiledBinding& InBinding, ERegisterLibraryBindingResult InResult)
+{
+	if (OnLibraryBindingRegistered.IsBound())
+	{
+		check(InView);
+		OnLibraryBindingRegistered.Broadcast(FView(InView), FLibraryBindingRegisteredArgs(InBinding, InResult));
+	}
+}
+
+/** */
+FDebugging::FLibraryBindingExecutedArgs::FLibraryBindingExecutedArgs(const FMVVMViewClass_CompiledBinding& InBinding)
+	: Binding(InBinding)
+{
+}
+
+FDebugging::FLibraryBindingExecutedArgs::FLibraryBindingExecutedArgs(const FMVVMViewClass_CompiledBinding& InBinding, FMVVMCompiledBindingLibrary::EExecutionFailingReason InResult)
+	: Binding(InBinding)
+	, FailingReason(InResult)
+{
+}
+
+FDebugging::FLibraryBindingExecuted FDebugging::OnLibraryBindingExecuted;
+
+void FDebugging::BroadcastLibraryBindingExecuted(const UMVVMView* InView, const FLibraryBindingExecutedArgs& Args)
+{
+	check(InView);
+	OnLibraryBindingExecuted.Broadcast(FView(InView), Args);
+}
+
+void FDebugging::BroadcastLibraryBindingExecuted(const UMVVMView* InView, const FMVVMViewClass_CompiledBinding& InBinding)
+{
+	if (OnLibraryBindingExecuted.IsBound())
+	{
+		check(InView);
+		OnLibraryBindingExecuted.Broadcast(FView(InView), FLibraryBindingExecutedArgs(InBinding));
+	}
+}
+
+void FDebugging::BroadcastLibraryBindingExecuted(const UMVVMView* InView, const FMVVMViewClass_CompiledBinding& InBinding, FMVVMCompiledBindingLibrary::EExecutionFailingReason InResult)
+{
+	if (OnLibraryBindingExecuted.IsBound())
+	{
+		check(InView);
+		OnLibraryBindingExecuted.Broadcast(FView(InView), FLibraryBindingExecutedArgs(InBinding, InResult));
+	}
+}
+} // namespace UE::MVVM
+#endif //UE_WITH_MVVM_DEBUGGING
+
+#undef LOCTEXT_NAMESPACE
\ No newline at end of file
diff --git a/Engine/Plugins/Runtime/ModelViewViewModel/Source/ModelViewViewModel/Private/MVVMMessageLog.cpp b/Engine/Plugins/Runtime/ModelViewViewModel/Source/ModelViewViewModel/Private/MVVMMessageLog.cpp
--- a/Engine/Plugins/Runtime/ModelViewViewModel/Source/ModelViewViewModel/Private/MVVMMessageLog.cpp
+++ b/Engine/Plugins/Runtime/ModelViewViewModel/Source/ModelViewViewModel/Private/MVVMMessageLog.cpp
@@ -2,4 +2,5 @@
 
 #include "MVVMMessageLog.h"
 
-const FName UE::MVVM::FMessageLog::LogName = "Model View Viewmodel";
+DEFINE_LOG_CATEGORY(LogMVVM);
+const FName UE::MVVM::FMessageLog::LogName = "Model View Viewmodel";
\ No newline at end of file
diff --git a/Engine/Plugins/Runtime/ModelViewViewModel/Source/ModelViewViewModel/Private/MVVMMessageLog.h b/Engine/Plugins/Runtime/ModelViewViewModel/Source/ModelViewViewModel/Private/MVVMMessageLog.h
--- a/Engine/Plugins/Runtime/ModelViewViewModel/Source/ModelViewViewModel/Private/MVVMMessageLog.h
+++ b/Engine/Plugins/Runtime/ModelViewViewModel/Source/ModelViewViewModel/Private/MVVMMessageLog.h
@@ -6,6 +6,8 @@
 
 #include "Blueprint/UserWidget.h"
 
+DECLARE_LOG_CATEGORY_EXTERN(LogMVVM, Log, All);
+
 namespace UE::MVVM
 {
 	class MODELVIEWVIEWMODEL_API FMessageLog : private ::FMessageLog
diff --git a/Engine/Plugins/Runtime/ModelViewViewModel/Source/ModelViewViewModel/Private/View/MVVMView.cpp b/Engine/Plugins/Runtime/ModelViewViewModel/Source/ModelViewViewModel/Private/View/MVVMView.cpp
--- a/Engine/Plugins/Runtime/ModelViewViewModel/Source/ModelViewViewModel/Private/View/MVVMView.cpp
+++ b/Engine/Plugins/Runtime/ModelViewViewModel/Source/ModelViewViewModel/Private/View/MVVMView.cpp
@@ -4,6 +4,7 @@
 #include "View/MVVMViewClass.h"
 #include "View/MVVMViewWorldSubsystem.h"
 
+#include "Debugging/MVVMDebugging.h"
 #include "Blueprint/UserWidget.h"
 #include "Blueprint/WidgetTree.h"
 #include "Engine/World.h"
@@ -37,6 +38,10 @@ void UMVVMView::Construct()
 		Item.CreateInstance(ClassExtension, this, GetUserWidget());
 	}
 
+#if UE_WITH_MVVM_DEBUGGING
+	UE::MVVM::FDebugging::BroadcastViewConstructed(this);
+#endif
+
 	bHasEveryTickBinding = false;
 
 	const TArrayView<const FMVVMViewClass_CompiledBinding> CompiledBindings = ClassExtension->GetCompiledBindings();
@@ -63,6 +68,10 @@ void UMVVMView::Destruct()
 	check(bConstructed == true);
 	bConstructed = false;
 
+#if UE_WITH_MVVM_DEBUGGING
+	UE::MVVM::FDebugging::BroadcastViewBeginDestruction(this);
+#endif
+
 	for (FRegisteredSource& Source : AllSources)
 	{
 		UObject* SourceObject = Source.Source.Get();
@@ -91,8 +100,7 @@ bool UMVVMView::SetViewModel(FName ViewModelName, TScriptInterface<INotifyFieldV
 		FObjectPropertyBase* FoundObjectProperty = FindFProperty<FObjectPropertyBase>(GetUserWidget()->GetClass(), ViewModelName);
 		if (FoundObjectProperty == nullptr)
 		{
-			UE::MVVM::FMessageLog Log(GetUserWidget());
-			Log.Error(LOCTEXT("SetViewModelInvalidName", "Invalid viewmodel name for SetViewModel."));
+			UE_LOG(LogMVVM, Error, TEXT("There is no viewmodel named '%s' in the view '%s'"), *ViewModelName.ToString(), *GetFullName());
 			return false;
 		}
 
@@ -107,6 +115,7 @@ bool UMVVMView::SetViewModel(FName ViewModelName, TScriptInterface<INotifyFieldV
 					{
 						UE::MVVM::FMessageLog Log(GetUserWidget());
 						Log.Error(LOCTEXT("SetViewModelInvalidValueType", "The new viewmodel is not of the expected type."));
+						UE_LOG(LogMVVM, Error, TEXT("The viewmodel name '%s' is invalid for the view '%s'"), *ViewModelName.ToString(), *GetFullName());
 						return false;
 					}
 					break;
@@ -309,6 +318,18 @@ void UMVVMView::ExecuteLibraryBinding(const FMVVMViewClass_CompiledBinding& Bind
 
 	FMVVMCompiledBindingLibrary::EConversionFunctionType FunctionType = Binding.IsConversionFunctionComplex() ? FMVVMCompiledBindingLibrary::EConversionFunctionType::Complex : FMVVMCompiledBindingLibrary::EConversionFunctionType::Simple;
 	TValueOrError<void, FMVVMCompiledBindingLibrary::EExecutionFailingReason> ExecutionResult = ClassExtension->GetBindingLibrary().Execute(GetUserWidget(), Binding.GetBinding(), FunctionType);
+
+#if UE_WITH_MVVM_DEBUGGING
+	if (ExecutionResult.HasError())
+	{
+		UE::MVVM::FDebugging::BroadcastLibraryBindingExecuted(this, Binding, ExecutionResult.GetError());
+	}
+	else
+	{
+		UE::MVVM::FDebugging::BroadcastLibraryBindingExecuted(this, Binding);
+	}
+#endif
+
 	if (ExecutionResult.HasError())
 	{
 		UE::MVVM::FMessageLog Log(GetUserWidget());
@@ -331,6 +352,17 @@ void UMVVMView::ExecuteLibraryBinding(const FMVVMViewClass_CompiledBinding& Bind
 	TValueOrError<void, FMVVMCompiledBindingLibrary::EExecutionFailingReason> ExecutionResult = Binding.IsConversionFunctionComplex()
 		? ClassExtension->GetBindingLibrary().Execute(GetUserWidget(), Binding.GetBinding(), FMVVMCompiledBindingLibrary::EConversionFunctionType::Complex)
 		: ClassExtension->GetBindingLibrary().ExecuteWithSource(GetUserWidget(), Binding.GetBinding(), Source);
+
+#if UE_WITH_MVVM_DEBUGGING
+	if (ExecutionResult.HasError())
+	{
+		UE::MVVM::FDebugging::BroadcastLibraryBindingExecuted(this, Binding, ExecutionResult.GetError());
+	}
+	else
+	{
+		UE::MVVM::FDebugging::BroadcastLibraryBindingExecuted(this, Binding);
+	}
+#endif
 	if (ExecutionResult.HasError())
 	{
 		UE::MVVM::FMessageLog Log(GetUserWidget());
@@ -442,16 +474,20 @@ bool UMVVMView::RegisterLibraryBinding(const FMVVMViewClass_CompiledBinding& Bin
 	TValueOrError<UE::FieldNotification::FFieldId, void> FieldIdResult = ClassExtension->GetBindingLibrary().GetFieldId(Binding.GetSourceFieldId());
 	if (FieldIdResult.HasError())
 	{
-		UE::MVVM::FMessageLog Log(GetUserWidget());
-		Log.Error(FText::Format(LOCTEXT("RegisterBindingInvalidFieldId", "Can't register binding '{0}'. The FieldId is invalid."), FText::FromString(Binding.ToString())));
+#if UE_WITH_MVVM_DEBUGGING
+		UE::MVVM::FDebugging::BroadcastLibraryBindingRegistered(this, Binding, UE::MVVM::FDebugging::ERegisterLibraryBindingResult::Failed_InvalidFieldId);
+#endif
+		UE_LOG(LogMVVM, Error, TEXT("'%s' can't register binding '%s'. The FieldId is invalid."), *GetFullName(), *Binding.ToString());
 		return false;
 	}
 
 	UE::FieldNotification::FFieldId FieldId = FieldIdResult.StealValue();
 	if (!FieldId.IsValid())
 	{
-		UE::MVVM::FMessageLog Log(GetUserWidget());
-		Log.Error(FText::Format(LOCTEXT("RegisterBindingNoFieldId", "Can't register binding '{0}'. The FieldId was not found on the source."), FText::FromString(Binding.ToString())));
+#if UE_WITH_MVVM_DEBUGGING
+		UE::MVVM::FDebugging::BroadcastLibraryBindingRegistered(this, Binding, UE::MVVM::FDebugging::ERegisterLibraryBindingResult::Failed_FieldIdNotFound);
+#endif
+		UE_LOG(LogMVVM, Error, TEXT("'%s' can't register binding '%s'. The FieldId was not found on the source."), *GetFullName(), *Binding.ToString());
 		return false;
 	}
 
@@ -459,10 +495,12 @@ bool UMVVMView::RegisterLibraryBinding(const FMVVMViewClass_CompiledBinding& Bin
 	TScriptInterface<INotifyFieldValueChanged> Source = FindSource(Binding, true);
 	if (Source.GetInterface() == nullptr)
 	{
+#if UE_WITH_MVVM_DEBUGGING
+		UE::MVVM::FDebugging::BroadcastLibraryBindingRegistered(this, Binding, UE::MVVM::FDebugging::ERegisterLibraryBindingResult::Failed_InvalidSource);
+#endif
 		if (!Binding.IsRegistrationOptional())
 		{
-			UE::MVVM::FMessageLog Log(GetUserWidget());
-			Log.Error(FText::Format(LOCTEXT("RegisterBindingInvalidSourceInstance", "Can't register binding '{0}'. The source is invalid."), FText::FromString(Binding.ToString())));
+			UE_LOG(LogMVVM, Error, TEXT("'%s' can't register binding '%s'. The source is invalid."), *GetFullName(), *Binding.ToString());
 		}
 		return false;
 	}
@@ -470,14 +508,18 @@ bool UMVVMView::RegisterLibraryBinding(const FMVVMViewClass_CompiledBinding& Bin
 	// Only bind if the source and the destination are valid.
 	if (ClassExtension->GetBindingLibrary().EvaluateFieldPath(GetUserWidget(), Binding.GetBinding().GetSourceFieldPath()).HasError())
 	{
-		UE::MVVM::FMessageLog Log(GetUserWidget());
-		Log.Warning(FText::Format(LOCTEXT("RegisterBindingSourceNotEvaluated", "Can't register binding '{0}'. The source was not evaluated."), FText::FromString(Binding.ToString())));
+#if UE_WITH_MVVM_DEBUGGING
+		UE::MVVM::FDebugging::BroadcastLibraryBindingRegistered(this, Binding, UE::MVVM::FDebugging::ERegisterLibraryBindingResult::Failed_InvalidSourceField);
+#endif
+		UE_LOG(LogMVVM, Warning, TEXT("'%s' can't register binding '%s'. The destination was not evaluated."), *GetFullName(), *Binding.ToString());
 		return false;
 	}
 	if (ClassExtension->GetBindingLibrary().EvaluateFieldPath(GetUserWidget(), Binding.GetBinding().GetDestinationFieldPath()).HasError())
 	{
-		UE::MVVM::FMessageLog Log(GetUserWidget());
-		Log.Warning(FText::Format(LOCTEXT("RegisterBindingInvalidDestination2", "Can't register binding '{0}'. The destination was not evaluated."), FText::FromString(Binding.ToString())));
+#if UE_WITH_MVVM_DEBUGGING
+		UE::MVVM::FDebugging::BroadcastLibraryBindingRegistered(this, Binding, UE::MVVM::FDebugging::ERegisterLibraryBindingResult::Failed_InvalidDestinationField);
+#endif
+		UE_LOG(LogMVVM, Warning, TEXT("'%s' can't register binding '%s'. The destination was not evaluated."), *GetFullName(), *Binding.ToString());
 		return false;
 	}
 
@@ -499,6 +541,9 @@ bool UMVVMView::RegisterLibraryBinding(const FMVVMViewClass_CompiledBinding& Bin
 			AllSources.Add(MoveTemp(Item));
 		}
 	}
+#if UE_WITH_MVVM_DEBUGGING
+	UE::MVVM::FDebugging::BroadcastLibraryBindingRegistered(this, Binding, UE::MVVM::FDebugging::ERegisterLibraryBindingResult::Success);
+#endif
 	return bResult;
 }
 
@@ -538,8 +583,7 @@ TScriptInterface<INotifyFieldValueChanged> UMVVMView::FindSource(const FMVVMView
 	const FObjectPropertyBase* SourceObjectProperty = CastField<FObjectPropertyBase>(UserWidget->GetClass()->FindPropertyByName(Binding.GetSourceObjectPropertyName()));
 	if (SourceObjectProperty == nullptr)
 	{
-		UE::MVVM::FMessageLog Log(GetUserWidget());
-		Log.Error(FText::Format(LOCTEXT("FindSourceInvalidPropertyName", "We could not evaluate the source for binding '{0}'. The property name is invalid."), FText::FromString(Binding.ToString())));
+		UE_LOG(LogMVVM, Error, TEXT("'%s' could not evaluate the source for binding '%s'. The property name is invalid."), *GetFullName(), *Binding.ToString());
 		return TScriptInterface<INotifyFieldValueChanged>();
 	}
 
@@ -548,16 +592,14 @@ TScriptInterface<INotifyFieldValueChanged> UMVVMView::FindSource(const FMVVMView
 	{
 		if (!bAllowNull)
 		{
-			UE::MVVM::FMessageLog Log(GetUserWidget());
-			Log.Warning(FText::Format(LOCTEXT("FindSourceInvalidUObject", "We could not evaluate the source for binding '{0}'. The path point to an invalid object."), FText::FromString(Binding.ToString())));
+			UE_LOG(LogMVVM, Error, TEXT("'%s' could not evaluate the source for binding '%s'. The path point to an invalid object."), *GetFullName(), *Binding.ToString());
 		}
 		return TScriptInterface<INotifyFieldValueChanged>();
 	}
 
 	if (!Source->Implements<UNotifyFieldValueChanged>())
 	{
-		UE::MVVM::FMessageLog Log(GetUserWidget());
-		Log.Error(FText::Format(LOCTEXT("FindSourceNotTheCorrectInterface", "We could not evaluate the source for binding '{0}'. The object {1} doesn't implements INotifyFieldValueChanged."), FText::FromString(Binding.ToString()), FText::FromName(Source->GetFName())));
+		UE_LOG(LogMVVM, Error, TEXT("'%s' could not evaluate the source for binding '%s'. The object '%s' doesn't implements INotifyFieldValueChanged."), *GetFullName(), *Binding.ToString(), *Source->GetFName().ToString());
 		return TScriptInterface<INotifyFieldValueChanged>();
 	}
 
diff --git a/Engine/Plugins/Runtime/ModelViewViewModel/Source/ModelViewViewModel/Public/Debugging/MVVMDebugging.h b/Engine/Plugins/Runtime/ModelViewViewModel/Source/ModelViewViewModel/Public/Debugging/MVVMDebugging.h
new file mode 100644
--- /dev/null
+++ b/Engine/Plugins/Runtime/ModelViewViewModel/Source/ModelViewViewModel/Public/Debugging/MVVMDebugging.h
@@ -0,0 +1,123 @@
+// Copyright Epic Games, Inc. All Rights Reserved.
+
+#pragma once
+
+#include "CoreMinimal.h"
+#include "Misc/Optional.h"
+
+#ifndef UE_WITH_MVVM_DEBUGGING
+#define UE_WITH_MVVM_DEBUGGING !(UE_BUILD_SHIPPING || UE_BUILD_TEST)
+#endif
+
+
+#if UE_WITH_MVVM_DEBUGGING
+
+#include "Bindings/MVVMCompiledBindingLibrary.h"
+
+class UMVVMView;
+struct FMVVMViewClass_CompiledBinding;
+class UUserWidget;
+
+namespace UE::MVVM
+{
+/** */
+class MODELVIEWVIEWMODEL_API FDebugging
+{
+public:
+	struct FView
+	{
+	public:
+		FView() = default;
+		FView(const UMVVMView* View);
+		FView(const UUserWidget* UserWidget, const UMVVMView* View);
+
+	public:
+		const UUserWidget* GetUserWidget() const
+		{
+			return UserWidget;
+		}
+
+		const UMVVMView* GetView() const
+		{
+			return View;
+		}
+	private:
+		const UUserWidget* UserWidget = nullptr;
+		const UMVVMView* View = nullptr;
+	};
+public:
+	struct FViewConstructedArgs
+	{
+	};
+
+	DECLARE_EVENT_TwoParams(FDebugging, FViewConstructed, const FView&, const FViewConstructedArgs&);
+	/** Broadcast when a view is created and the viewmodels are instantiated and the binding are not registered. */
+	static FViewConstructed OnViewConstructed;
+	static void BroadcastViewConstructed(const UMVVMView* View);
+
+
+	struct FViewBeginDestructionArgs
+	{
+	};
+
+	DECLARE_EVENT_TwoParams(FDebugging, FViewDestructing, const FView&, const FViewBeginDestructionArgs&);
+	/** Broadcast before a view is destroyed. */
+	static FViewDestructing OnViewBeginDestruction;
+	static void BroadcastViewBeginDestruction(const UMVVMView* View);
+
+public:
+	enum class ERegisterLibraryBindingResult : uint8
+	{
+		Success,
+		Failed_InvalidFieldId,
+		Failed_FieldIdNotFound,
+		Failed_InvalidSource,
+		Failed_InvalidSourceField,
+		Failed_InvalidDestinationField,
+	};
+
+	struct FLibraryBindingRegisteredArgs
+	{
+		FLibraryBindingRegisteredArgs() = delete;
+		FLibraryBindingRegisteredArgs(const FMVVMViewClass_CompiledBinding& Binding, ERegisterLibraryBindingResult Result);
+		const FMVVMViewClass_CompiledBinding& Binding;
+		ERegisterLibraryBindingResult Result;
+	};
+
+	DECLARE_EVENT_TwoParams(FDebugging, FLibraryBindingRegistered, const FView&, const FLibraryBindingRegisteredArgs&);
+	/** Broadcast when a binding is registered to the view/viewmodel or failed to registered to the view/viewmodel. */
+	static FLibraryBindingRegistered OnLibraryBindingRegistered;
+	static void BroadcastLibraryBindingRegistered(const UMVVMView* View, const FLibraryBindingRegisteredArgs& Args);
+	static void BroadcastLibraryBindingRegistered(const UMVVMView* View, const FMVVMViewClass_CompiledBinding& Binding, ERegisterLibraryBindingResult Result);
+
+public:
+	struct FLibraryBindingExecutedArgs
+	{
+		FLibraryBindingExecutedArgs() = delete;
+		FLibraryBindingExecutedArgs(const FMVVMViewClass_CompiledBinding& Binding);
+		FLibraryBindingExecutedArgs(const FMVVMViewClass_CompiledBinding& Binding, FMVVMCompiledBindingLibrary::EExecutionFailingReason Result);
+		const FMVVMViewClass_CompiledBinding& Binding;
+		TOptional<FMVVMCompiledBindingLibrary::EExecutionFailingReason> FailingReason;
+	};
+
+	/**  */
+	DECLARE_EVENT_TwoParams(FDebugging, FLibraryBindingExecuted, const FView&, const FLibraryBindingExecutedArgs&);
+	/** Broadcast when a registered field is modified and a binding need to execute. */
+	static FLibraryBindingExecuted OnLibraryBindingExecuted;
+	static void BroadcastLibraryBindingExecuted(const UMVVMView* View, const FLibraryBindingExecutedArgs& Args);
+	static void BroadcastLibraryBindingExecuted(const UMVVMView* View, const FMVVMViewClass_CompiledBinding& Binding);
+	static void BroadcastLibraryBindingExecuted(const UMVVMView* View, const FMVVMViewClass_CompiledBinding& Binding, FMVVMCompiledBindingLibrary::EExecutionFailingReason Result);
+};
+} // UE::MVVM
+
+#else
+
+
+namespace UE::MVVM
+{
+class FDebugging
+{
+};
+} // UE::MVVM
+
+#endif
diff --git a/Engine/Plugins/Runtime/ModelViewViewModel/Source/ModelViewViewModel/Public/View/MVVMView.h b/Engine/Plugins/Runtime/ModelViewViewModel/Source/ModelViewViewModel/Public/View/MVVMView.h
--- a/Engine/Plugins/Runtime/ModelViewViewModel/Source/ModelViewViewModel/Public/View/MVVMView.h
+++ b/Engine/Plugins/Runtime/ModelViewViewModel/Source/ModelViewViewModel/Public/View/MVVMView.h
@@ -18,6 +18,10 @@ struct FMVVMViewClass_CompiledBinding;
 struct FMVVMViewDelayedBinding;
 class UMVVMViewModelBase;
 class UWidget;
+namespace UE::MVVM
+{
+	class FDebugging;
+}
 
 /**
  * Instance UMVVMClassExtension_View for the UUserWdiget
@@ -27,6 +31,8 @@ class MODELVIEWVIEWMODEL_API UMVVMView : public UUserWidgetExtension
 {
 	GENERATED_BODY()
 
+	friend UE::MVVM::FDebugging;
+
 public:
 	void ConstructView(const UMVVMViewClass* ClassExtension);
 
@@ -35,6 +41,11 @@ public:
 	virtual void Destruct() override;
 	//~ End UUserWidgetExtension implementation
 
+	const UMVVMViewClass* GetViewClass() const
+	{
+		return ClassExtension;
+	}
+
 	void ExecuteDelayedBinding(const FMVVMViewDelayedBinding& DelayedBinding) const;
 	void ExecuteEveryTickBindings() const;
 
