MVVM: Move the description for enum to be pickup correctly by the docgen
#rnx

diff --git a/Engine/Plugins/Runtime/ModelViewViewModel/Source/ModelViewViewModel/Public/Bindings/MVVMCompiledBindingLibrary.h b/Engine/Plugins/Runtime/ModelViewViewModel/Source/ModelViewViewModel/Public/Bindings/MVVMCompiledBindingLibrary.h
--- a/Engine/Plugins/Runtime/ModelViewViewModel/Source/ModelViewViewModel/Public/Bindings/MVVMCompiledBindingLibrary.h
+++ b/Engine/Plugins/Runtime/ModelViewViewModel/Source/ModelViewViewModel/Public/Bindings/MVVMCompiledBindingLibrary.h
@@ -261,8 +261,10 @@ public:
 	/** */
 	enum class EConversionFunctionType
 	{
-		Simple, 	// accept an single argument and returns a single property
-		Complex,	// returns a single property. The argument are fetch from inside the function 
+		// accept an single argument and returns a single property
+		Simple,
+		// returns a single property. The argument are fetch from inside the function 
+		Complex,
 	};
 
 	/** Fetch the FProperty and UFunction. */
diff --git a/Engine/Plugins/Runtime/ModelViewViewModel/Source/ModelViewViewModelBlueprint/Public/MVVMBlueprintViewModelContext.h b/Engine/Plugins/Runtime/ModelViewViewModel/Source/ModelViewViewModelBlueprint/Public/MVVMBlueprintViewModelContext.h
--- a/Engine/Plugins/Runtime/ModelViewViewModel/Source/ModelViewViewModelBlueprint/Public/MVVMBlueprintViewModelContext.h
+++ b/Engine/Plugins/Runtime/ModelViewViewModel/Source/ModelViewViewModelBlueprint/Public/MVVMBlueprintViewModelContext.h
@@ -15,10 +15,14 @@
 UENUM()
 enum class EMVVMBlueprintViewModelContextCreationType : uint8
 {
-	Manual,	// The viewmodel will be assigned later.
-	CreateInstance, // A new instance of the viewmodel will be created when the widget is created.
-	GlobalViewModelCollection, // The viewmodel exists and is added to the MVVMSubsystem. It will be fetched there.
-	PropertyPath, // The viewmodel will be fetched by evaluating a function or a property path.
+	// The viewmodel will be assigned later.
+	Manual,
+	// A new instance of the viewmodel will be created when the widget is created.
+	CreateInstance,
+	// The viewmodel exists and is added to the MVVMSubsystem. It will be fetched there.
+	GlobalViewModelCollection,
+	// The viewmodel will be fetched by evaluating a function or a property path.
+	PropertyPath,
 };
 
 /**
