MVVM: Add the viewmodal used by the UserWidget in the metadata

diff --git a/Engine/Plugins/Runtime/ModelViewViewModel/Source/ModelViewViewModelBlueprint/Private/MVVMBlueprintView.cpp b/Engine/Plugins/Runtime/ModelViewViewModel/Source/ModelViewViewModelBlueprint/Private/MVVMBlueprintView.cpp
--- a/Engine/Plugins/Runtime/ModelViewViewModel/Source/ModelViewViewModelBlueprint/Private/MVVMBlueprintView.cpp
+++ b/Engine/Plugins/Runtime/ModelViewViewModel/Source/ModelViewViewModelBlueprint/Private/MVVMBlueprintView.cpp
@@ -2,9 +2,10 @@
 
 #include "MVVMBlueprintView.h"
 
-#include "MVVMWidgetBlueprintExtension_View.h"
+#include "Containers/StringFwd.h"
 #include "Blueprint/WidgetTree.h"
 #include "Kismet2/BlueprintEditorUtils.h"
+#include "MVVMWidgetBlueprintExtension_View.h"
 #include "UObject/UnrealType.h"
 
 #include UE_INLINE_GENERATED_CPP_BY_NAME(MVVMBlueprintView)
@@ -234,6 +235,30 @@ void UMVVMBlueprintView::PostEditChangeChainProperty(FPropertyChangedChainEvent&
 	}
 }
 
+void UMVVMBlueprintView::AddAssetTags(TArray<FAssetRegistryTag>& OutTags)
+{
+	if (AvailableViewModels.Num() > 0)
+	{
+		TStringBuilder<512> Builder;
+		for (const FMVVMBlueprintViewModelContext& Context : AvailableViewModels)
+		{
+			if (Context.IsValid())
+			{
+				if (Builder.Len() > 0)
+				{
+					Builder << TEXT(',');
+				}
+				Builder << Context.GetViewModelClass()->GetPathName();
+			}
+		}
+
+		if (Builder.Len() > 0)
+		{
+			OutTags.Emplace(TEXT("Viewmodels"), Builder.ToString(), FAssetRegistryTag::TT_Hidden);
+		}
+	}
+}
+
 void UMVVMBlueprintView::WidgetRenamed(FName OldObjectName, FName NewObjectName)
 {
 	bool bRenamed = false;
diff --git a/Engine/Plugins/Runtime/ModelViewViewModel/Source/ModelViewViewModelBlueprint/Public/MVVMBlueprintView.h b/Engine/Plugins/Runtime/ModelViewViewModel/Source/ModelViewViewModelBlueprint/Public/MVVMBlueprintView.h
--- a/Engine/Plugins/Runtime/ModelViewViewModel/Source/ModelViewViewModelBlueprint/Public/MVVMBlueprintView.h
+++ b/Engine/Plugins/Runtime/ModelViewViewModel/Source/ModelViewViewModelBlueprint/Public/MVVMBlueprintView.h
@@ -71,6 +71,7 @@ public:
 	virtual void PostEditChangeProperty(FPropertyChangedEvent& PropertyChangedEvent) override;
 	virtual void PostEditChangeChainProperty(FPropertyChangedChainEvent& PropertyChainEvent) override;
 
+	void AddAssetTags(TArray<FAssetRegistryTag>& OutTags);
 	void WidgetRenamed(FName OldObjectName, FName NewObjectName);
 #endif
 
diff --git a/Engine/Plugins/Runtime/ModelViewViewModel/Source/ModelViewViewModelEditor/Private/ModelViewViewModelEditorModule.cpp b/Engine/Plugins/Runtime/ModelViewViewModel/Source/ModelViewViewModelEditor/Private/ModelViewViewModelEditorModule.cpp
--- a/Engine/Plugins/Runtime/ModelViewViewModel/Source/ModelViewViewModelEditor/Private/ModelViewViewModelEditorModule.cpp
+++ b/Engine/Plugins/Runtime/ModelViewViewModel/Source/ModelViewViewModelEditor/Private/ModelViewViewModelEditorModule.cpp
@@ -6,6 +6,7 @@
 #include "BlueprintModes/WidgetBlueprintApplicationMode.h"
 #include "BlueprintModes/WidgetBlueprintApplicationModes.h"
 #include "Customizations/MVVMPropertyBindingExtension.h"
+#include "Editor.h"
 #include "EdGraphSchema_K2.h"
 #include "Framework/Application/SlateApplication.h"
 #include "Framework/Docking/LayoutExtender.h"
@@ -18,6 +19,7 @@
 #include "MVVMBlueprintView.h"
 #include "MVVMBlueprintViewModelContext.h"
 #include "MVVMEditorCommands.h"
+#include "MVVMEditorSubsystem.h"
 #include "MVVMWidgetBlueprintExtension_View.h"
 #include "StatusBarSubsystem.h"
 #include "Styling/MVVMEditorStyle.h"
@@ -69,11 +71,13 @@ void FModelViewViewModelEditorModule::StartupModule()
 	}
 
 	FMVVMEditorCommands::Register();
+	FWidgetBlueprintDelegates::GetAssetTags.AddRaw(this, &FModelViewViewModelEditorModule::HandleAssetTags);
 }
 
 
 void FModelViewViewModelEditorModule::ShutdownModule()
 {
+	FWidgetBlueprintDelegates::GetAssetTags.RemoveAll(this);
 	if (FMessageLogModule* MessageLogModule = FModuleManager::GetModulePtr<FMessageLogModule>("MessageLog"))
 	{
 		MessageLogModule->UnregisterLogListing("Model View Viewmodel");
@@ -198,6 +202,20 @@ void FModelViewViewModelEditorModule::HandleActivateMode(FWidgetBlueprintApplica
 	}
 }
 
+void FModelViewViewModelEditorModule::HandleAssetTags(const UWidgetBlueprint* WidgetBlueprint, TArray<UObject::FAssetRegistryTag>& OutTags)
+{
+	if (WidgetBlueprint && GEditor)
+	{
+		if (UMVVMEditorSubsystem* Subsystem = GEditor->GetEditorSubsystem<UMVVMEditorSubsystem>())
+		{
+			if (UMVVMBlueprintView* BlueprintView = Subsystem->GetView(WidgetBlueprint))
+			{
+				BlueprintView->AddAssetTags(OutTags);
+			}
+		}
+	}
+}
+
 IMPLEMENT_MODULE(FModelViewViewModelEditorModule, ModelViewViewModelEditor);
 
 #undef LOCTEXT_NAMESPACE
\ No newline at end of file
diff --git a/Engine/Plugins/Runtime/ModelViewViewModel/Source/ModelViewViewModelEditor/Private/ModelViewViewModelEditorModule.h b/Engine/Plugins/Runtime/ModelViewViewModel/Source/ModelViewViewModelEditor/Private/ModelViewViewModelEditorModule.h
--- a/Engine/Plugins/Runtime/ModelViewViewModel/Source/ModelViewViewModelEditor/Private/ModelViewViewModelEditorModule.h
+++ b/Engine/Plugins/Runtime/ModelViewViewModel/Source/ModelViewViewModelEditor/Private/ModelViewViewModelEditorModule.h
@@ -5,6 +5,7 @@
 #include "CoreMinimal.h"
 #include "Modules/ModuleInterface.h"
 #include "Modules/ModuleManager.h"
+#include "UObject/Object.h"
 
 namespace UE::MVVM
 {
@@ -16,6 +17,7 @@ class FMVVMPropertyBindingExtension;
 class FWidgetBlueprintApplicationMode;
 class FWorkflowAllowedTabSet;
 class UBlueprint;
+class UWidgetBlueprint;
 
 /**
  *
@@ -35,6 +37,7 @@ private:
 	void HandleRenameVariableReferences(UBlueprint* Blueprint, UClass* VariableClass, const FName& OldVarName, const FName& NewVarName);
 	void HandleDeactiveMode(FWidgetBlueprintApplicationMode& InDesignerMode);
 	void HandleActivateMode(FWidgetBlueprintApplicationMode& InDesignerMode);
+	void HandleAssetTags(const UWidgetBlueprint* Widget, TArray<UObject::FAssetRegistryTag>& OutTags);
 
 private:
 	TSharedPtr<FMVVMPropertyBindingExtension> PropertyBindingExtension;
