Fixed typo of NotifyFieldValudChangedClass -> NotifyFieldValueChangedClass.

#rnx

diff --git a/Engine/Plugins/Runtime/ModelViewViewModel/Source/ModelViewViewModelBlueprint/Private/Bindings/MVVMCompiledBindingLibraryCompiler.cpp b/Engine/Plugins/Runtime/ModelViewViewModel/Source/ModelViewViewModelBlueprint/Private/Bindings/MVVMCompiledBindingLibraryCompiler.cpp
--- a/Engine/Plugins/Runtime/ModelViewViewModel/Source/ModelViewViewModelBlueprint/Private/Bindings/MVVMCompiledBindingLibraryCompiler.cpp
+++ b/Engine/Plugins/Runtime/ModelViewViewModel/Source/ModelViewViewModelBlueprint/Private/Bindings/MVVMCompiledBindingLibraryCompiler.cpp
@@ -19,7 +19,7 @@ static const FName NAME_BlueprintGetter = "BlueprintGetter";
 /** */
 struct FRawFieldId
 {
-	UClass* NotifyFieldValudChangedClass;
+	UClass* NotifyFieldValueChangedClass;
 	UE::FieldNotification::FFieldId FieldId;
 
 	int32 LoadedFieldIdIndex = INDEX_NONE;
@@ -164,12 +164,12 @@ TValueOrError<FCompiledBindingLibraryCompiler::FFieldIdHandle, FText> FCompiledB
 
 		int32 FoundFieldIdIndex = Impl->FieldIds.IndexOfByPredicate([FoundFieldId, SourceClass](const Private::FRawFieldId& Other)
 			{
-				return Other.FieldId == FoundFieldId && Other.NotifyFieldValudChangedClass == SourceClass.Get();
+				return Other.FieldId == FoundFieldId && Other.NotifyFieldValueChangedClass == SourceClass.Get();
 			});
 		if (FoundFieldIdIndex == INDEX_NONE)
 		{
 			Private::FRawFieldId RawFieldId;
-			RawFieldId.NotifyFieldValudChangedClass = SourceClass.Get();
+			RawFieldId.NotifyFieldValueChangedClass = SourceClass.Get();
 			RawFieldId.FieldId = FoundFieldId;
 			RawFieldId.IdHandle = FFieldIdHandle::MakeHandle();
 			FoundFieldIdIndex = Impl->FieldIds.Add(MoveTemp(RawFieldId));
@@ -574,9 +574,9 @@ TValueOrError<FCompiledBindingLibraryCompiler::FCompileResult, FText> FCompiledB
 		{
 			const Private::FRawFieldId& RawFieldId = Impl->FieldIds[Index];
 			check(RawFieldId.FieldId.IsValid());
-			check(RawFieldId.NotifyFieldValudChangedClass);
+			check(RawFieldId.NotifyFieldValueChangedClass);
 
-			FCompiledClassInfo& ClassInfo = MapOfFieldInClass.FindOrAdd(RawFieldId.NotifyFieldValudChangedClass);
+			FCompiledClassInfo& ClassInfo = MapOfFieldInClass.FindOrAdd(RawFieldId.NotifyFieldValueChangedClass);
 
 			// Test if the Field is not there more than one
 			{
