MVVM - Blueprint view is now transactional, which causes undo to function in many cases that did not work before.

#rnx

diff --git a/Engine/Plugins/Runtime/ModelViewViewModel/Source/ModelViewViewModelBlueprint/Private/MVVMWidgetBlueprintExtension_View.cpp b/Engine/Plugins/Runtime/ModelViewViewModel/Source/ModelViewViewModelBlueprint/Private/MVVMWidgetBlueprintExtension_View.cpp
--- a/Engine/Plugins/Runtime/ModelViewViewModel/Source/ModelViewViewModelBlueprint/Private/MVVMWidgetBlueprintExtension_View.cpp
+++ b/Engine/Plugins/Runtime/ModelViewViewModel/Source/ModelViewViewModelBlueprint/Private/MVVMWidgetBlueprintExtension_View.cpp
@@ -13,7 +13,7 @@
 
 void UMVVMWidgetBlueprintExtension_View::CreateBlueprintViewInstance()
 {
-	BlueprintView = NewObject<UMVVMBlueprintView>(this);
+	BlueprintView = NewObject<UMVVMBlueprintView>(this, FName(), RF_Transactional);
 	BlueprintViewChangedDelegate.Broadcast();
 }
 
@@ -28,10 +28,6 @@ void UMVVMWidgetBlueprintExtension_View::DestroyBlueprintViewInstance()
 void UMVVMWidgetBlueprintExtension_View::PostLoad()
 {
 	Super::PostLoad();
-	//for (FMVVMViewBindingEditorData& Data : ViewBindings)
-	//{
-	//	Data.Resolve();
-	//}
 }
 
 
@@ -41,18 +37,12 @@ void UMVVMWidgetBlueprintExtension_View::HandlePreloadObjectsForCompilation(UBlu
 	{
 		BlueprintView->ConditionalPostLoad();
 
-		for (const FMVVMBlueprintViewBinding& Bindings : BlueprintView->GetBindings())
-		{
-			//UBlueprint::ForceLoad(Bindings.ViewModelPath);
-			//UBlueprint::ForceLoad(Bindings.WidgetPath);
-		}
 		for (const FMVVMBlueprintViewModelContext& AvailableViewModel : BlueprintView->GetViewModels())
 		{
 			if (AvailableViewModel.GetViewModelClass())
 			{
 				UBlueprint::ForceLoad(AvailableViewModel.GetViewModelClass());
 			}
-			//UBlueprint::ForceLoad(ViewModelPropertyPath);
 		}
 	}
 }
