MVVM: Make static function in MVVMSubsystem for faster access.

diff --git a/Engine/Plugins/Runtime/ModelViewViewModel/Config/BaseModelViewViewModel.ini b/Engine/Plugins/Runtime/ModelViewViewModel/Config/BaseModelViewViewModel.ini
--- a/Engine/Plugins/Runtime/ModelViewViewModel/Config/BaseModelViewViewModel.ini
+++ b/Engine/Plugins/Runtime/ModelViewViewModel/Config/BaseModelViewViewModel.ini
@@ -1,7 +1,9 @@
 [CoreRedirects]
 +StructRedirects=(OldName="MVVMViewModelPropertyPath",NewName="/script/ModelViewViewModelBlueprint.MVVMBlueprintPropertyPath")
 +StructRedirects=(OldName="MVVMWidgetPropertyPath",NewName="/script/ModelViewViewModelBlueprint.MVVMBlueprintPropertyPath")
-
++FunctionRedirects=(OldName="/Script/ModelViewViewModel.MVVMSubsystem.GetViewFromUserWidget", NewName="/Script/ModelViewViewModel.MVVMSubsystem.K2_GetViewFromUserWidget")
++FunctionRedirects=(OldName="/Script/ModelViewViewModel.MVVMSubsystem.GetAvailableBindings", NewName="/Script/ModelViewViewModel.MVVMSubsystem.K2_GetAvailableBindings")
++FunctionRedirects=(OldName="/Script/ModelViewViewModel.MVVMSubsystem.GetAvailableBinding", NewName="/Script/ModelViewViewModel.MVVMSubsystem.K2_GetAvailableBinding")
 
 [/Script/ModelViewViewModelBlueprint.MVVMDeveloperProjectSettings]
 +FieldSelectorPermissions=(("/Script/UMG.Widget", (DisallowedFieldNames=("TooltipWidget","ForceVolatile","SetUserFocus","Navigation","GetParent","GetOwningPlayer","GetOwningLocalPlayer","GetGameInstance"),AdvancedFieldNames=())))
\ No newline at end of file
diff --git a/Engine/Plugins/Runtime/ModelViewViewModel/Source/ModelViewViewModel/Private/MVVMSubsystem.cpp b/Engine/Plugins/Runtime/ModelViewViewModel/Source/ModelViewViewModel/Private/MVVMSubsystem.cpp
--- a/Engine/Plugins/Runtime/ModelViewViewModel/Source/ModelViewViewModel/Private/MVVMSubsystem.cpp
+++ b/Engine/Plugins/Runtime/ModelViewViewModel/Source/ModelViewViewModel/Private/MVVMSubsystem.cpp
@@ -26,7 +26,13 @@ void UMVVMSubsystem::Deinitialize()
 }
 
 
-UMVVMView* UMVVMSubsystem::GetViewFromUserWidget(const UUserWidget* UserWidget) const
+UMVVMView* UMVVMSubsystem::K2_GetViewFromUserWidget(const UUserWidget* UserWidget) const
+{
+	return GetViewFromUserWidget(UserWidget);
+}
+
+
+UMVVMView* UMVVMSubsystem::GetViewFromUserWidget(const UUserWidget* UserWidget)
 {
 	return UserWidget ? UserWidget->GetExtension<UMVVMView>() : nullptr;
 }
@@ -250,25 +256,37 @@ namespace UE::MVVM::Private
 } //namespace
 
 
-TArray<FMVVMAvailableBinding> UMVVMSubsystem::GetAvailableBindings(const UClass* Class, const UClass* Accessor) const
+TArray<FMVVMAvailableBinding> UMVVMSubsystem::K2_GetAvailableBindings(const UClass* Class, const UClass* Accessor) const
+{
+	return GetAvailableBindings(Class, Accessor);
+}
+
+
+TArray<FMVVMAvailableBinding> UMVVMSubsystem::GetAvailableBindings(const UClass* Class, const UClass* Accessor)
 {
 	return UE::MVVM::Private::GetAvailableBindings(Class, Accessor);
 }
 
 
-TArray<FMVVMAvailableBinding> UMVVMSubsystem::GetAvailableBindingsForStruct(const UScriptStruct* Struct) const
+TArray<FMVVMAvailableBinding> UMVVMSubsystem::GetAvailableBindingsForStruct(const UScriptStruct* Struct)
 {
 	return UE::MVVM::Private::GetAvailableBindings(Struct, nullptr, nullptr);
 }
 
 
-FMVVMAvailableBinding UMVVMSubsystem::GetAvailableBinding(const UClass* Class, FMVVMBindingName BindingName, const UClass* Accessor) const
+FMVVMAvailableBinding UMVVMSubsystem::K2_GetAvailableBinding(const UClass* Class, FMVVMBindingName BindingName, const UClass* Accessor) const
+{
+	return GetAvailableBinding(Class, BindingName, Accessor);
+}
+
+
+FMVVMAvailableBinding UMVVMSubsystem::GetAvailableBinding(const UClass* Class, FMVVMBindingName BindingName, const UClass* Accessor)
 {
 	return UE::MVVM::Private::GetAvailableBinding(BindingName, Class, Accessor);
 }
 
 
-FMVVMAvailableBinding UMVVMSubsystem::GetAvailableBindingForVariant(UE::MVVM::FMVVMConstFieldVariant FieldVariant, const UClass* Accessor) const
+FMVVMAvailableBinding UMVVMSubsystem::GetAvailableBindingForField(UE::MVVM::FMVVMConstFieldVariant FieldVariant, const UClass* Accessor)
 {
 	return UE::MVVM::Private::GetAvailableBinding(FieldVariant, Accessor);
 }
diff --git a/Engine/Plugins/Runtime/ModelViewViewModel/Source/ModelViewViewModel/Public/MVVMSubsystem.h b/Engine/Plugins/Runtime/ModelViewViewModel/Source/ModelViewViewModel/Public/MVVMSubsystem.h
--- a/Engine/Plugins/Runtime/ModelViewViewModel/Source/ModelViewViewModel/Public/MVVMSubsystem.h
+++ b/Engine/Plugins/Runtime/ModelViewViewModel/Source/ModelViewViewModel/Public/MVVMSubsystem.h
@@ -29,28 +29,34 @@ public:
 	virtual void Deinitialize() override;
 	//~ End UEngineSubsystem interface
 
-	UFUNCTION(BlueprintCallable, Category="MVVM")
-	UMVVMView* GetViewFromUserWidget(const UUserWidget* UserWidget) const;
+	UFUNCTION(BlueprintCallable, Category="MVVM", meta = (DisplayName = "Get View From User Widget"))
+	UMVVMView* K2_GetViewFromUserWidget(const UUserWidget* UserWidget) const;
+
+	static UMVVMView* GetViewFromUserWidget(const UUserWidget* UserWidget);
 
 	UFUNCTION(BlueprintCallable, Category = "MVVM")
 	bool DoesWidgetTreeContainedWidget(const UWidgetTree* WidgetTree, const UWidget* ViewWidget) const;
 
 	/** @return The list of all the AvailableBindings that are available for the Class. */
-	UFUNCTION(BlueprintCallable, Category = "MVVM")
-	TArray<FMVVMAvailableBinding> GetAvailableBindings(const UClass* Class, const UClass* Accessor) const;
+	UFUNCTION(BlueprintCallable, Category = "MVVM", meta = (DisplayName = "Get Available Bindings"))
+	TArray<FMVVMAvailableBinding> K2_GetAvailableBindings(const UClass* Class, const UClass* Accessor) const;
+
+	static TArray<FMVVMAvailableBinding> GetAvailableBindings(const UClass* Class, const UClass* Accessor);
 
 	/**
 	 * @return The list of all the AvailableBindings that are available from the SriptStuct.
 	 * @note When FMVVMAvailableBinding::HasNotify is false, a notification can still be triggered by the owner of the struct. The struct changed but which property of the struct changed is unknown.
 	 */
-	TArray<FMVVMAvailableBinding> GetAvailableBindingsForStruct(const UScriptStruct* Struct) const;
+	static TArray<FMVVMAvailableBinding> GetAvailableBindingsForStruct(const UScriptStruct* Struct);
 
 	/** @return The AvailableBinding from a BindingName. */
-	UFUNCTION(BlueprintCallable, Category = "MVVM")
-	FMVVMAvailableBinding GetAvailableBinding(const UClass* Class, FMVVMBindingName BindingName, const UClass* Accessor) const;
+	UFUNCTION(BlueprintCallable, Category = "MVVM", meta = (DisplayName = "Get Available Binding"))
+	FMVVMAvailableBinding K2_GetAvailableBinding(const UClass* Class, FMVVMBindingName BindingName, const UClass* Accessor) const;
+
+	static FMVVMAvailableBinding GetAvailableBinding(const UClass* Class, FMVVMBindingName BindingName, const UClass* Accessor);
 
 	/** @return The AvailableBinding from a field. */
-	FMVVMAvailableBinding GetAvailableBindingForVariant(UE::MVVM::FMVVMConstFieldVariant Variant, const UClass* Accessor) const;
+	static FMVVMAvailableBinding GetAvailableBindingForField(UE::MVVM::FMVVMConstFieldVariant Variant, const UClass* Accessor);
 
 	UFUNCTION(BlueprintCallable, Category = "MVVM")
 	UMVVMViewModelCollectionObject* GetGlobalViewModelCollection() const
diff --git a/Engine/Plugins/Runtime/ModelViewViewModel/Source/ModelViewViewModelEditor/Private/Widgets/SMVVMViewModelBindingListWidget.cpp b/Engine/Plugins/Runtime/ModelViewViewModel/Source/ModelViewViewModelEditor/Private/Widgets/SMVVMViewModelBindingListWidget.cpp
--- a/Engine/Plugins/Runtime/ModelViewViewModel/Source/ModelViewViewModelEditor/Private/Widgets/SMVVMViewModelBindingListWidget.cpp
+++ b/Engine/Plugins/Runtime/ModelViewViewModel/Source/ModelViewViewModelEditor/Private/Widgets/SMVVMViewModelBindingListWidget.cpp
@@ -194,12 +194,12 @@ TArray<FFieldVariant> FFieldIterator_Bindable::GetFields(const UStruct* Struct)
 	{
 		const UWidgetBlueprint* WidgetBlueprintPtr = WidgetBlueprint.Get();
 		TSubclassOf<UObject> AccessorClass = WidgetBlueprintPtr ? WidgetBlueprintPtr->GeneratedClass : nullptr;
-		TArray<FMVVMAvailableBinding> Bindings = GEngine->GetEngineSubsystem<UMVVMSubsystem>()->GetAvailableBindings(const_cast<UClass*>(Class), AccessorClass);
+		TArray<FMVVMAvailableBinding> Bindings = UMVVMSubsystem::GetAvailableBindings(const_cast<UClass*>(Class), AccessorClass);
 		AddResult(Bindings);
 	}
 	else if (const UScriptStruct* ScriptStruct = Cast<const UScriptStruct>(Struct))
 	{
-		TArray<FMVVMAvailableBinding> Bindings = GEngine->GetEngineSubsystem<UMVVMSubsystem>()->GetAvailableBindingsForStruct(ScriptStruct);
+		TArray<FMVVMAvailableBinding> Bindings = UMVVMSubsystem::GetAvailableBindingsForStruct(ScriptStruct);
 		AddResult(Bindings);
 	}
 
@@ -259,7 +259,7 @@ TSharedRef<SWidget> ConstructFieldPreSlot(const UWidgetBlueprint* WidgetBlueprin
 {
 	TSharedRef<SWidget> ImageWidget = SNullWidget::NullWidget;
 	TSubclassOf<UObject> AccessorClass = WidgetBlueprint ? WidgetBlueprint->SkeletonGeneratedClass : nullptr;
-	FMVVMAvailableBinding Binding = GEngine->GetEngineSubsystem<UMVVMSubsystem>()->GetAvailableBindingForVariant(FMVVMConstFieldVariant(FieldPath), AccessorClass);
+	FMVVMAvailableBinding Binding = UMVVMSubsystem::GetAvailableBindingForField(FMVVMConstFieldVariant(FieldPath), AccessorClass);
 	if (Binding.IsValid())
 	{
 		const FSlateBrush* Brush = nullptr;
@@ -483,7 +483,7 @@ FMVVMBlueprintPropertyPath SSourceBindingList::CreateBlueprintPropertyPath(SProp
 
 			if (const UClass* OwnerClass = Cast<const UClass>(OwnerStruct))
 			{
-				FMVVMAvailableBinding Binding = GEngine->GetEngineSubsystem<UMVVMSubsystem>()->GetAvailableBinding(OwnerClass, FMVVMBindingName(FieldName), AccessorClass);
+				FMVVMAvailableBinding Binding = UMVVMSubsystem::GetAvailableBinding(OwnerClass, FMVVMBindingName(FieldName), AccessorClass);
 				if (Binding.IsValid())
 				{
 					bPassFilter = Private::PassFilter(Binding, OwnerClass, FieldIterator->GetFieldVisibilityFlags(), FieldIterator->GetAssignableTo(), true).IsSet();
@@ -514,7 +514,7 @@ FMVVMBlueprintPropertyPath SSourceBindingList::CreateBlueprintPropertyPath(SProp
 	else if(AccessorClass.Get())
 	{
 		FMVVMBindingName BindingName = Source->Key.ToBindingName(WidgetBlueprintPtr);
-		FMVVMAvailableBinding Binding = GEngine->GetEngineSubsystem<UMVVMSubsystem>()->GetAvailableBinding(AccessorClass, BindingName, AccessorClass);
+		FMVVMAvailableBinding Binding = UMVVMSubsystem::GetAvailableBinding(AccessorClass, BindingName, AccessorClass);
 		if (Binding.IsValid())
 		{
 			bool bPassFilter = Private::PassFilter(Binding, AccessorClass, FieldIterator->GetFieldVisibilityFlags(), FieldIterator->GetAssignableTo(), true).IsSet();
