MVVM: Rework GetBindingInfoFromFieldPath to support self.

diff --git a/Engine/Plugins/Runtime/ModelViewViewModel/Source/ModelViewViewModel/Private/Bindings/MVVMFieldPathHelper.cpp b/Engine/Plugins/Runtime/ModelViewViewModel/Source/ModelViewViewModel/Private/Bindings/MVVMFieldPathHelper.cpp
--- a/Engine/Plugins/Runtime/ModelViewViewModel/Source/ModelViewViewModel/Private/Bindings/MVVMFieldPathHelper.cpp
+++ b/Engine/Plugins/Runtime/ModelViewViewModel/Source/ModelViewViewModel/Private/Bindings/MVVMFieldPathHelper.cpp
@@ -38,6 +38,20 @@ TValueOrError<UStruct*, FText> FindContainer(const FProperty* Property)
 }
 
 
+const FProperty* FindProperty(FMVVMConstFieldVariant Field)
+{
+	if (Field.IsProperty())
+	{
+		return Field.GetProperty();
+	}
+	else if (Field.IsFunction() && Field.GetFunction())
+	{
+		return BindingHelper::GetReturnProperty(Field.GetFunction());
+	}
+	return nullptr;
+}
+
+
 TValueOrError<FMVVMConstFieldVariant, FText> TransformWithAccessor(UStruct* CurrentContainer, FMVVMConstFieldVariant CurrentField, bool bForReading)
 {
 #if WITH_EDITORONLY_DATA
@@ -269,18 +283,15 @@ TValueOrError<FParsedBindingInfo, FText> GetBindingInfoFromFieldPath(const UClas
 		return MakeError(LOCTEXT("FieldPathEmpty", "The field path is empty."));
 	}
 
-	FParsedBindingInfo Result;
-	FParsedBindingInfo::FFieldVariantArray LocalFullPath;
-	bool bLookForNotifyFieldId = false;
-	const UStruct* CurrentContainerType = InFieldPath[0].GetOwner();
-	if (!InAccessor->IsChildOf(CurrentContainerType))
+	if (!InAccessor->IsChildOf(InFieldPath[0].GetOwner()))
 	{
 		return MakeError(LOCTEXT("FieldPathDoesntStartWithAccessor", "Field path that doesn't start with the accessor is not supported."));
 	}
 
-	for (int32 Index = 0; Index < InFieldPath.Num(); ++Index)
+	FParsedBindingInfo Result;
+
+	for (int32 Index = InFieldPath.Num() - 1; Index >= 0; --Index)
 	{
-		bool bLastField = Index == InFieldPath.Num() - 1;
 		FMVVMConstFieldVariant Field = InFieldPath[Index];
 
 		if (Field.IsEmpty())
@@ -288,94 +299,54 @@ TValueOrError<FParsedBindingInfo, FText> GetBindingInfoFromFieldPath(const UClas
 			return MakeError(FText::Format(LOCTEXT("FieldIsEmpty", "The field at index '{0}' is empty."), Index));
 		}
 
-		if (CurrentContainerType == nullptr)
+		const UStruct* FieldOwner = Field.GetOwner();
+		check(FieldOwner);
+		const UStruct* ParentContainerType = InAccessor;
+		check(ParentContainerType);
+		if (Index != 0)
 		{
-			return MakeError(FText::Format(LOCTEXT("FieldDoesNotExist", "The field '{0}' does not exist."), 
-				FText::FromName(Field.GetName()))
-			);
-		}
+			const FProperty* ParentProperty = Private::FindProperty(InFieldPath[Index - 1]);
+			if (ParentProperty == nullptr)
+			{
+				return MakeError(FText::Format(LOCTEXT("InvalidParentProperty", "The field at index '{0}' has an invalid parent."), Index));
+			}
 
-		const bool bIsChild = Field.GetOwner()->IsChildOf(CurrentContainerType);
-		const bool bIsDownCast = CurrentContainerType->IsChildOf(Field.GetOwner());
-		if (!(bIsChild || bIsDownCast))
-		{
-			return MakeError(FText::Format(LOCTEXT("FieldDoesNotExistInStruct", "The field '{0}' does not exist in the struct '{1}'."), 
-				FText::FromName(Field.GetName()), 
-				FText::FromString(CurrentContainerType->GetName()))
-			);
+			TValueOrError<UStruct*, FText> FoundParentContainer = Private::FindContainer(ParentProperty);
+			if (FoundParentContainer.HasError())
+			{
+				return MakeError(FoundParentContainer.StealError());
+			}
+			ParentContainerType = FoundParentContainer.GetValue();
 		}
 
-		const FProperty* FieldProperty = nullptr;
-		if (Field.IsProperty())
+		if (ParentContainerType == nullptr)
 		{
-			FieldProperty = Field.GetProperty();
-		}
-		else if (Field.IsFunction() && Field.GetFunction())
-		{
-			FieldProperty = BindingHelper::GetReturnProperty(Field.GetFunction());
+			return MakeError(FText::Format(LOCTEXT("InvalidFieldParent", "The field '{0}' is not a container."), Index - 1));
 		}
 
-		if (FieldProperty == nullptr)
+		const bool bIsChild = FieldOwner->IsChildOf(ParentContainerType);
+		const bool bIsDownCast = ParentContainerType->IsChildOf(FieldOwner);
+		if (!(bIsChild || bIsDownCast))
 		{
-			return MakeError(FText::Format(LOCTEXT("FieldHasInvalidProperty", "The field '{0}' has an invalid property."), 
-				FText::FromName(Field.GetName()))
+			return MakeError(FText::Format(LOCTEXT("FieldDoesNotExistInStruct", "The field '{0}' does not exist in the struct '{1}'."),
+				FText::FromName(Field.GetName()),
+				FText::FromString(ParentContainerType->GetName()))
 			);
 		}
 
-
-		LocalFullPath.Add(Field);
-
 		// Is field a viewmodel
-		bool bSaveNotifyFieldInterfacePath = false;
-		if (!bLastField)
-		{
-			if (const FObjectPropertyBase* FieldObjectProperty = CastField<const FObjectPropertyBase>(FieldProperty))
-			{
-				if (FieldObjectProperty->PropertyClass->ImplementsInterface(UNotifyFieldValueChanged::StaticClass()))
-				{
-					Result.NotifyFieldClass = FieldObjectProperty->PropertyClass;
-					Result.NotifyFieldId = FFieldNotificationId();
-					Result.NotifyFieldInterfacePath = LocalFullPath;
-					Result.ViewModelIndex = Index;
-					bSaveNotifyFieldInterfacePath = true;
-					bLookForNotifyFieldId = true;
-				}
-			}
-		}
-
-		if (!bSaveNotifyFieldInterfacePath)
-		{
-			if (bLookForNotifyFieldId && Result.NotifyFieldClass.Get() != nullptr)
-			{
-				FMVVMBindingName BindingName = FMVVMBindingName(Field.GetName());
-				FMVVMAvailableBinding AvailableBinding = GEngine->GetEngineSubsystem<UMVVMSubsystem>()->GetAvailableBinding(Result.NotifyFieldClass.Get(), BindingName, InAccessor);
-				if (AvailableBinding.IsValid())
-				{
-					Result.NotifyFieldId = FFieldNotificationId(AvailableBinding.GetBindingName().ToName());
-				}
-			}
-			bLookForNotifyFieldId = false;
-		}
-
-		if (!bLastField)
+		const UClass* FieldOwnerAsClass = Cast<const UClass>(FieldOwner);
+		if (FieldOwnerAsClass && FieldOwnerAsClass->ImplementsInterface(UNotifyFieldValueChanged::StaticClass()))
 		{
-			TValueOrError<UStruct*, FText> FoundContainer = Private::FindContainer(FieldProperty);
-			if (FoundContainer.HasError())
+			FMVVMAvailableBinding AvailableBinding = UMVVMSubsystem::GetAvailableBindingForField(Field, InAccessor);
+			if (AvailableBinding.IsValid() && AvailableBinding.HasNotify())
 			{
-				return MakeError(FoundContainer.StealError());
+				Result.NotifyFieldClass = const_cast<UClass*>(FieldOwnerAsClass);
+				Result.NotifyFieldId = FFieldNotificationId(AvailableBinding.GetBindingName().ToName());
+				Result.ViewModelIndex = Index-1;
+				break;
 			}
-			CurrentContainerType = FoundContainer.GetValue();
-		}
-	}
-
-	if (Result.ViewModelIndex > 0)
-	{
-		TValueOrError<TArray<FMVVMConstFieldVariant>, FText> FinalPath = GenerateFieldPathList(Result.NotifyFieldInterfacePath, true);
-		if (FinalPath.HasError())
-		{
-			return MakeError(FinalPath.StealError());
 		}
-		Result.NotifyFieldInterfacePath = FinalPath.StealValue();
 	}
 
 	return MakeValue(MoveTemp(Result));
diff --git a/Engine/Plugins/Runtime/ModelViewViewModel/Source/ModelViewViewModel/Private/View/MVVMView.cpp b/Engine/Plugins/Runtime/ModelViewViewModel/Source/ModelViewViewModel/Private/View/MVVMView.cpp
--- a/Engine/Plugins/Runtime/ModelViewViewModel/Source/ModelViewViewModel/Private/View/MVVMView.cpp
+++ b/Engine/Plugins/Runtime/ModelViewViewModel/Source/ModelViewViewModel/Private/View/MVVMView.cpp
@@ -580,30 +580,37 @@ TScriptInterface<INotifyFieldValueChanged> UMVVMView::FindSource(const FMVVMView
 	check(UserWidget);
 	check(ClassExtension);
 
-	const FObjectPropertyBase* SourceObjectProperty = CastField<FObjectPropertyBase>(UserWidget->GetClass()->FindPropertyByName(Binding.GetSourceObjectPropertyName()));
-	if (SourceObjectProperty == nullptr)
+	if (Binding.IsSourceObjectItself())
 	{
-		UE_LOG(LogMVVM, Error, TEXT("'%s' could not evaluate the source for binding '%s'. The property name is invalid."), *GetFullName(), *Binding.ToString());
-		return TScriptInterface<INotifyFieldValueChanged>();
+		return TScriptInterface<INotifyFieldValueChanged>(UserWidget);
 	}
-
-	UObject* Source = SourceObjectProperty->GetObjectPropertyValue_InContainer(UserWidget);
-	if (Source == nullptr)
+	else
 	{
-		if (!bAllowNull)
+		const FObjectPropertyBase* SourceObjectProperty = CastField<FObjectPropertyBase>(UserWidget->GetClass()->FindPropertyByName(Binding.GetSourceObjectPropertyName()));
+		if (SourceObjectProperty == nullptr)
 		{
-			UE_LOG(LogMVVM, Error, TEXT("'%s' could not evaluate the source for binding '%s'. The path point to an invalid object."), *GetFullName(), *Binding.ToString());
+			UE_LOG(LogMVVM, Error, TEXT("'%s' could not evaluate the source for binding '%s'. The property name is invalid."), *GetFullName(), *Binding.ToString());
+			return TScriptInterface<INotifyFieldValueChanged>();
 		}
-		return TScriptInterface<INotifyFieldValueChanged>();
-	}
 
-	if (!Source->Implements<UNotifyFieldValueChanged>())
-	{
-		UE_LOG(LogMVVM, Error, TEXT("'%s' could not evaluate the source for binding '%s'. The object '%s' doesn't implements INotifyFieldValueChanged."), *GetFullName(), *Binding.ToString(), *Source->GetFName().ToString());
-		return TScriptInterface<INotifyFieldValueChanged>();
-	}
+		UObject* Source = SourceObjectProperty->GetObjectPropertyValue_InContainer(UserWidget);
+		if (Source == nullptr)
+		{
+			if (!bAllowNull)
+			{
+				UE_LOG(LogMVVM, Error, TEXT("'%s' could not evaluate the source for binding '%s'. The path point to an invalid object."), *GetFullName(), *Binding.ToString());
+			}
+			return TScriptInterface<INotifyFieldValueChanged>();
+		}
+
+		if (!Source->Implements<UNotifyFieldValueChanged>())
+		{
+			UE_LOG(LogMVVM, Error, TEXT("'%s' could not evaluate the source for binding '%s'. The object '%s' doesn't implements INotifyFieldValueChanged."), *GetFullName(), *Binding.ToString(), *Source->GetFName().ToString());
+			return TScriptInterface<INotifyFieldValueChanged>();
+		}
 
-	return TScriptInterface<INotifyFieldValueChanged>(Source);
+		return TScriptInterface<INotifyFieldValueChanged>(Source);
+	}
 }
 
 
diff --git a/Engine/Plugins/Runtime/ModelViewViewModel/Source/ModelViewViewModel/Public/Bindings/MVVMFieldPathHelper.h b/Engine/Plugins/Runtime/ModelViewViewModel/Source/ModelViewViewModel/Public/Bindings/MVVMFieldPathHelper.h
--- a/Engine/Plugins/Runtime/ModelViewViewModel/Source/ModelViewViewModel/Public/Bindings/MVVMFieldPathHelper.h
+++ b/Engine/Plugins/Runtime/ModelViewViewModel/Source/ModelViewViewModel/Public/Bindings/MVVMFieldPathHelper.h
@@ -32,13 +32,9 @@ namespace UE::MVVM::FieldPathHelper
 		 */
 		FFieldNotificationId NotifyFieldId;
 
-		/**
-		 * The property path to get the viewmodel at runtime.
-		 * @return "ViewmodelA.ViewmodelB" or "ViewmodelA.GetViewmodelB" when the path is "ViewmodelA.ViewmodelB.Vector.X"
-		 */
-		FFieldVariantArray NotifyFieldInterfacePath;
 		/**
 		 * The viewmodel index in the original path.
+		 * When it's INDEX_NONE but the ParsedBindingInfo is valid, the class accessor is the Accessor class.
 		 * @return 1 when the path is "ViewmodelA.ViewmodelB.Vector.X"
 		 */
 		int32 ViewModelIndex = INDEX_NONE;
@@ -50,7 +46,11 @@ namespace UE::MVVM::FieldPathHelper
 	/** */
 	UE_NODISCARD MODELVIEWVIEWMODEL_API TValueOrError<TArray<FMVVMConstFieldVariant>, FText> GenerateFieldPathList(TArrayView<const FMVVMConstFieldVariant> Fields, bool bForSourceBinding);
 
-	/** */
+	/**
+	 * Generate the FParsedBindingInfo from the parsed field path.
+	 * This gives us which FieldId and Viewmodel should be used by the path.
+	 * Note. The code may have changed since the binding was created in the editor. As an example, a sub-object can be a viewmodel now.
+	 */
 	UE_NODISCARD MODELVIEWVIEWMODEL_API TValueOrError<FParsedBindingInfo, FText> GetBindingInfoFromFieldPath(const UClass* Accessor, TArrayView<const FMVVMConstFieldVariant> Fields);
 
 	/** */
diff --git a/Engine/Plugins/Runtime/ModelViewViewModel/Source/ModelViewViewModel/Public/View/MVVMViewClass.h b/Engine/Plugins/Runtime/ModelViewViewModel/Source/ModelViewViewModel/Public/View/MVVMViewClass.h
--- a/Engine/Plugins/Runtime/ModelViewViewModel/Source/ModelViewViewModel/Public/View/MVVMViewClass.h
+++ b/Engine/Plugins/Runtime/ModelViewViewModel/Source/ModelViewViewModel/Public/View/MVVMViewClass.h
@@ -112,6 +112,12 @@ public:
 		return SourcePropertyName;
 	}
 
+	/** @return true if SourceFieldId is not from a property but on the object itserf (UserWidget). */
+	bool IsSourceObjectItself() const
+	{
+		return (Flags & EBindingFlags::SourceObjectIsSelf) != 0;
+	}
+
 	/** @return the binding. From source to destination (if forward) or from destination to source (if backward). */
 	const FMVVMVCompiledBinding& GetBinding() const
 	{
@@ -170,14 +176,20 @@ private:
 	enum EBindingFlags
 	{
 		None = 0,
-		ForwardBinding = 1 << 0, // True when the binding goes from Source to Destination.
-		TwoWayBinding = 1 << 1, // The binding is one part of a 2 ways binding.
+		/** True when the binding goes from Source to Destination. */
+		ForwardBinding = 1 << 0,
+		/** The binding is one part of a 2 ways binding. */
+		TwoWayBinding = 1 << 1,
 		OneTime = 1 << 2,
 		EnabledByDefault = 1 << 3,
-		ViewModelOptional = 1 << 4,	// The source (viewmodel) can be nullptr and the binding could failed and should not log a warning.
-		ConversionFunctionIsComplex = 1 << 5,	// The conversion function is complex, there is no input. The inputs are calculated in the BP function.
-		// In development, (when the Blueprint maybe not be compiled with the latest data), the ExecutionMode may not reflect the default project setting value.
+		/** The source (viewmodel) can be nullptr and the binding could failed and should not log a warning. */
+		ViewModelOptional = 1 << 4,
+		/** The conversion function is complex, there is no input. The inputs are calculated in the BP function. */
+		ConversionFunctionIsComplex = 1 << 5,
+		/** In development, (when the Blueprint maybe not be compiled with the latest data), the ExecutionMode may not reflect the default project setting value. */
 		OverrideExecuteMode = 1 << 6,
+		/** When the source object is the object itself. */
+		SourceObjectIsSelf = 1 << 7,
 	};
 
 	UPROPERTY()
diff --git a/Engine/Plugins/Runtime/ModelViewViewModel/Source/ModelViewViewModelBlueprint/Private/MVVMViewBlueprintCompiler.cpp b/Engine/Plugins/Runtime/ModelViewViewModel/Source/ModelViewViewModelBlueprint/Private/MVVMViewBlueprintCompiler.cpp
--- a/Engine/Plugins/Runtime/ModelViewViewModel/Source/ModelViewViewModelBlueprint/Private/MVVMViewBlueprintCompiler.cpp
+++ b/Engine/Plugins/Runtime/ModelViewViewModel/Source/ModelViewViewModelBlueprint/Private/MVVMViewBlueprintCompiler.cpp
@@ -1075,9 +1075,11 @@ bool FMVVMViewBlueprintCompiler::CompileBindings(const FCompiledBindingLibraryCo
 		FMVVMBlueprintViewBinding& ViewBinding = *BlueprintView->GetBindingAt(CompileBinding.BindingIndex);
 
 		FMVVMViewClass_CompiledBinding NewBinding;
-
-		check(CompileBinding.CompilerSourceContextIndex != INDEX_NONE);
-		NewBinding.SourcePropertyName = CompilerSourceContexts[CompileBinding.CompilerSourceContextIndex].PropertyName;
+		const bool bIsSourceSelf = !CompilerSourceCreatorContexts.IsValidIndex(CompileBinding.CompilerSourceContextIndex);
+		if (!bIsSourceSelf)
+		{
+			NewBinding.SourcePropertyName = CompilerSourceContexts[CompileBinding.CompilerSourceContextIndex].PropertyName;
+		}
 
 		const FMVVMVCompiledFieldId* CompiledFieldId = CompileResult.FieldIds.Find(CompileBinding.FieldIdHandle);
 		if (CompiledFieldId == nullptr && CompileBinding.bFieldIdNeeded)
@@ -1098,8 +1100,9 @@ bool FMVVMViewBlueprintCompiler::CompileBindings(const FCompiledBindingLibraryCo
 		}
 
 		bool bIsOptional = false;
-		if (CompilerSourceCreatorContexts.IsValidIndex(CompileBinding.CompilerSourceContextIndex))
+		if (!bIsSourceSelf)
 		{
+			check(CompilerSourceCreatorContexts.IsValidIndex(CompileBinding.CompilerSourceContextIndex));
 			const FMVVMBlueprintViewModelContext& ModelContext = CompilerSourceCreatorContexts[CompileBinding.CompilerSourceContextIndex].ViewModelContext;
 			if (ModelContext.IsValid())
 			{
@@ -1119,6 +1122,7 @@ bool FMVVMViewBlueprintCompiler::CompileBindings(const FCompiledBindingLibraryCo
 		NewBinding.Flags |= (bIsOptional) ? FMVVMViewClass_CompiledBinding::EBindingFlags::ViewModelOptional : 0;
 		NewBinding.Flags |= (CompileBinding.bIsConversionFunctionComplex) ? FMVVMViewClass_CompiledBinding::EBindingFlags::ConversionFunctionIsComplex : 0;
 		NewBinding.Flags |= (ViewBinding.bOverrideExecutionMode) ? FMVVMViewClass_CompiledBinding::EBindingFlags::OverrideExecuteMode : 0;
+		NewBinding.Flags |= (bIsSourceSelf) ? FMVVMViewClass_CompiledBinding::EBindingFlags::SourceObjectIsSelf : 0;
 
 		ViewExtension->CompiledBindings.Emplace(MoveTemp(NewBinding));
 	}
@@ -1141,7 +1145,7 @@ TValueOrError<FMVVMViewBlueprintCompiler::FBindingSourceContext, FText> FMVVMVie
 	if (PropertyPath.IsEmpty())
 	{
 		ensureAlways(false);
-		return MakeError(LOCTEXT("EmptyPropertyPath", "Empty property path found. This is legal."));
+		return MakeError(LOCTEXT("EmptyPropertyPath", "Empty property path found. This is ilegal."));
 	}
 
 	FBindingSourceContext Result;
@@ -1165,7 +1169,7 @@ TValueOrError<FMVVMViewBlueprintCompiler::FBindingSourceContext, FText> FMVVMVie
 		if (Result.bIsRootWidget)
 		{
 			Result.UserWidgetPropertyContextIndex = INDEX_NONE;
-			Result.SourceClass = Class->ClassGeneratedBy->GetClass();
+			Result.SourceClass = const_cast<UWidgetBlueprintGeneratedClass*>(Class);
 			Result.PropertyPath = CreatePropertyPath(Class, FName(), PropertyPath.GetFields());
 		}
 		else
@@ -1197,31 +1201,32 @@ TValueOrError<FMVVMViewBlueprintCompiler::FBindingSourceContext, FText> FMVVMVie
 	// The INotifyFieldValueChanged/viewmodel is not the first property. Creates the ViaBinding source entries.
 	if (bCreateADynamicSourceEntry)
 	{
-		Result.SourceClass = BindingInfo.NotifyFieldClass;
-		Result.UserWidgetPropertyContextIndex = INDEX_NONE;
-		Result.bIsRootWidget = false;
-
-		FString InterfacePath = FieldPathHelper::ToString(BindingInfo.NotifyFieldInterfacePath);
-		int32 SourceCreatorContextIndex = CompilerSourceCreatorContexts.IndexOfByPredicate([&InterfacePath](const FCompilerSourceCreatorContext& Other)
-			{
-				return (Other.Type == ECompilerSourceCreatorType::ViewModel)
-					? (Other.ViewModelContext.CreationType == EMVVMBlueprintViewModelContextCreationType::PropertyPath
-						&& Other.ViewModelContext.ViewModelPropertyPath == InterfacePath)
-					: false;
-			});
-
-
-		// Add if the path doesn't already exists
-		if (!CompilerSourceCreatorContexts.IsValidIndex(SourceCreatorContextIndex))
-		{
-			return MakeError(LOCTEXT("ViewModelNotAtRoot", "The field in the path is from a viewmodel that is not the root of the path. This is not yet supported."));
-		}
-
-		Result.PropertyPath.RemoveAt(0, BindingInfo.ViewModelIndex);
-		if (CompilerSourceCreatorContexts[SourceCreatorContextIndex].Type == ECompilerSourceCreatorType::ViewModel)
-		{
-			Result.UserWidgetPropertyContextIndex = SourceCreatorContextIndex;
-		}
+		return MakeError(LOCTEXT("DynamicSourceEntryNotSupport", "Long source entry is not yet fully supported. Add the viewmodel manually."));
+		
+		//Result.SourceClass = BindingInfo.NotifyFieldClass;
+		//Result.UserWidgetPropertyContextIndex = INDEX_NONE;
+		//Result.bIsRootWidget = false;
+
+		//FString InterfacePath = FieldPathHelper::ToString(BindingInfo.NotifyFieldInterfacePath);
+		//int32 SourceCreatorContextIndex = CompilerSourceCreatorContexts.IndexOfByPredicate([&InterfacePath](const FCompilerSourceCreatorContext& Other)
+		//	{
+		//		return (Other.Type == ECompilerSourceCreatorType::ViewModel)
+		//			? (Other.ViewModelContext.CreationType == EMVVMBlueprintViewModelContextCreationType::PropertyPath
+		//				&& Other.ViewModelContext.ViewModelPropertyPath == InterfacePath)
+		//			: false;
+		//	});
+
+		//// Add if the path doesn't already exists
+		//if (!CompilerSourceCreatorContexts.IsValidIndex(SourceCreatorContextIndex))
+		//{
+		//	return MakeError(LOCTEXT("ViewModelNotAtRoot", "The field in the path is from a viewmodel that is not the root of the path. This is not yet supported."));
+		//}
+
+		//Result.PropertyPath.RemoveAt(0, BindingInfo.ViewModelIndex);
+		//if (CompilerSourceCreatorContexts[SourceCreatorContextIndex].Type == ECompilerSourceCreatorType::ViewModel)
+		//{
+		//	Result.UserWidgetPropertyContextIndex = SourceCreatorContextIndex;
+		//}
 	}
 
 	return MakeValue(Result);
