Fix bad index uses introduced with CL23494894. Rename CompilerSourceContexts to CompilerUserWidgetPropertyContexts to help prevent futur issues.

diff --git a/Engine/Plugins/Runtime/ModelViewViewModel/Source/ModelViewViewModelBlueprint/Private/MVVMViewBlueprintCompiler.cpp b/Engine/Plugins/Runtime/ModelViewViewModel/Source/ModelViewViewModelBlueprint/Private/MVVMViewBlueprintCompiler.cpp
--- a/Engine/Plugins/Runtime/ModelViewViewModel/Source/ModelViewViewModelBlueprint/Private/MVVMViewBlueprintCompiler.cpp
+++ b/Engine/Plugins/Runtime/ModelViewViewModel/Source/ModelViewViewModelBlueprint/Private/MVVMViewBlueprintCompiler.cpp
@@ -208,7 +208,7 @@ void FMVVMViewBlueprintCompiler::CreateVariables(const FWidgetBlueprintCompilerC
 		return NewProperty;
 	};
 
-	for (FCompilerUserWidgetPropertyContext& SourceContext : CompilerSourceContexts)
+	for (FCompilerUserWidgetPropertyContext& SourceContext : CompilerUserWidgetPropertyContexts)
 	{
 		SourceContext.Field = BindingHelper::FindFieldByName(Context.GetSkeletonGeneratedClass(), FMVVMBindingName(SourceContext.PropertyName));
 
@@ -297,7 +297,7 @@ void FMVVMViewBlueprintCompiler::CreateWidgetMap(const FWidgetBlueprintCompilerC
 
 void FMVVMViewBlueprintCompiler::CreateSourceLists(const FWidgetBlueprintCompilerContext::FCreateVariableContext& Context, UMVVMBlueprintView* BlueprintView)
 {
-	CompilerSourceContexts.Reset();
+	CompilerUserWidgetPropertyContexts.Reset();
 	CompilerSourceCreatorContexts.Reset();
 
 	TSet<FGuid> ViewModelGuids;
@@ -358,7 +358,8 @@ void FMVVMViewBlueprintCompiler::CreateSourceLists(const FWidgetBlueprintCompile
 		SourceVariable.CategoryName = TEXT("Viewmodel");
 		SourceVariable.bExposeOnSpawn = bCreateSetterFunction;
 		SourceVariable.BlueprintSetter = CompilerSourceCreatorContexts[FoundSourceCreatorContextIndex].SetterFunctionName;
-		CompilerSourceContexts.Emplace(MoveTemp(SourceVariable));
+		SourceVariable.ViewModelId = ViewModelContext.GetViewModelId();
+		CompilerUserWidgetPropertyContexts.Emplace(MoveTemp(SourceVariable));
 	}
 
 	bAreSourceContextsValid = bAreSourcesCreatorValid;
@@ -427,7 +428,8 @@ void FMVVMViewBlueprintCompiler::CreateSourceLists(const FWidgetBlueprintCompile
 					SourceVariable.PropertyName = PropertyPath.GetWidgetName();
 					SourceVariable.DisplayName = FText::FromString(Widget->GetDisplayLabel());
 					SourceVariable.CategoryName = TEXT("Widget");
-					Self->CompilerSourceContexts.Emplace(MoveTemp(SourceVariable));
+					SourceVariable.ViewModelId = FGuid();
+					Self->CompilerUserWidgetPropertyContexts.Emplace(MoveTemp(SourceVariable));
 				}
 			}
 			else if (PropertyPath.IsFromViewModel())
@@ -1045,7 +1047,7 @@ bool FMVVMViewBlueprintCompiler::PreCompileBindings(UWidgetBlueprintGeneratedCla
 
 		FCompilerBinding NewBinding = AddBindingResult.StealValue();
 		NewBinding.BindingIndex = BindingSourceContext.BindingIndex;
-		NewBinding.CompilerSourceContextIndex = BindingSourceContext.UserWidgetPropertyContextIndex;
+		NewBinding.UserWidgetPropertyContextIndex = BindingSourceContext.UserWidgetPropertyContextIndex;
 		NewBinding.FieldIdHandle = AddFieldResult.StealValue();
 		NewBinding.bSourceIsUserWidget = BindingSourceContext.bIsRootWidget;
 		NewBinding.bFieldIdNeeded = !IsOneTimeBinding(Binding.BindingType);
@@ -1075,10 +1077,10 @@ bool FMVVMViewBlueprintCompiler::CompileBindings(const FCompiledBindingLibraryCo
 		FMVVMBlueprintViewBinding& ViewBinding = *BlueprintView->GetBindingAt(CompileBinding.BindingIndex);
 
 		FMVVMViewClass_CompiledBinding NewBinding;
-		const bool bIsSourceSelf = !CompilerSourceCreatorContexts.IsValidIndex(CompileBinding.CompilerSourceContextIndex);
+		const bool bIsSourceSelf = !CompilerUserWidgetPropertyContexts.IsValidIndex(CompileBinding.UserWidgetPropertyContextIndex);
 		if (!bIsSourceSelf)
 		{
-			NewBinding.SourcePropertyName = CompilerSourceContexts[CompileBinding.CompilerSourceContextIndex].PropertyName;
+			NewBinding.SourcePropertyName = CompilerUserWidgetPropertyContexts[CompileBinding.UserWidgetPropertyContextIndex].PropertyName;
 		}
 
 		const FMVVMVCompiledFieldId* CompiledFieldId = CompileResult.FieldIds.Find(CompileBinding.FieldIdHandle);
@@ -1102,11 +1104,19 @@ bool FMVVMViewBlueprintCompiler::CompileBindings(const FCompiledBindingLibraryCo
 		bool bIsOptional = false;
 		if (!bIsSourceSelf)
 		{
-			check(CompilerSourceCreatorContexts.IsValidIndex(CompileBinding.CompilerSourceContextIndex));
-			const FMVVMBlueprintViewModelContext& ModelContext = CompilerSourceCreatorContexts[CompileBinding.CompilerSourceContextIndex].ViewModelContext;
-			if (ModelContext.IsValid())
+			check(CompilerUserWidgetPropertyContexts.IsValidIndex(CompileBinding.UserWidgetPropertyContextIndex));
+			const FGuid ViewModelId = CompilerUserWidgetPropertyContexts[CompileBinding.UserWidgetPropertyContextIndex].ViewModelId;
+			if (ViewModelId.IsValid())
 			{
-				bIsOptional = ModelContext.bOptional;
+				const FMVVMBlueprintViewModelContext* ModelContext = BlueprintView->FindViewModel(ViewModelId);
+				if (ModelContext == nullptr)
+				{
+					AddErrorForBinding(ViewBinding, LOCTEXT("CompiledBindingWithInvalidIVewModelId", "Internal error: the view model became invalid."));
+					bIsBindingsValid = false;
+					continue;
+				}
+
+				bIsOptional = ModelContext->IsValid();
 			}
 		}
 
@@ -1156,11 +1166,11 @@ TValueOrError<FMVVMViewBlueprintCompiler::FBindingSourceContext, FText> FMVVMVie
 		const FMVVMBlueprintViewModelContext* SourceViewModelContext = BlueprintView->FindViewModel(PropertyPath.GetViewModelId());
 		check(SourceViewModelContext);
 		const FName SourceName = SourceViewModelContext->GetViewModelName();
-		Result.UserWidgetPropertyContextIndex = CompilerSourceContexts.IndexOfByPredicate([SourceName](const FCompilerUserWidgetPropertyContext& Other) { return Other.PropertyName == SourceName; });
+		Result.UserWidgetPropertyContextIndex = CompilerUserWidgetPropertyContexts.IndexOfByPredicate([SourceName](const FCompilerUserWidgetPropertyContext& Other) { return Other.PropertyName == SourceName; });
 		check(Result.UserWidgetPropertyContextIndex != INDEX_NONE);
 
-		Result.SourceClass = CompilerSourceContexts[Result.UserWidgetPropertyContextIndex].Class;
-		Result.PropertyPath = CreatePropertyPath(Class, CompilerSourceContexts[Result.UserWidgetPropertyContextIndex].PropertyName, PropertyPath.GetFields());
+		Result.SourceClass = CompilerUserWidgetPropertyContexts[Result.UserWidgetPropertyContextIndex].Class;
+		Result.PropertyPath = CreatePropertyPath(Class, CompilerUserWidgetPropertyContexts[Result.UserWidgetPropertyContextIndex].PropertyName, PropertyPath.GetFields());
 	}
 	else if (PropertyPath.IsFromWidget())
 	{
@@ -1174,10 +1184,10 @@ TValueOrError<FMVVMViewBlueprintCompiler::FBindingSourceContext, FText> FMVVMVie
 		}
 		else
 		{
-			Result.UserWidgetPropertyContextIndex = CompilerSourceContexts.IndexOfByPredicate([SourceName](const FCompilerUserWidgetPropertyContext& Other) { return Other.PropertyName == SourceName; });
+			Result.UserWidgetPropertyContextIndex = CompilerUserWidgetPropertyContexts.IndexOfByPredicate([SourceName](const FCompilerUserWidgetPropertyContext& Other) { return Other.PropertyName == SourceName; });
 			check(Result.UserWidgetPropertyContextIndex != INDEX_NONE);
-			Result.SourceClass = CompilerSourceContexts[Result.UserWidgetPropertyContextIndex].Class;
-			Result.PropertyPath = CreatePropertyPath(Class, CompilerSourceContexts[Result.UserWidgetPropertyContextIndex].PropertyName, PropertyPath.GetFields());
+			Result.SourceClass = CompilerUserWidgetPropertyContexts[Result.UserWidgetPropertyContextIndex].Class;
+			Result.PropertyPath = CreatePropertyPath(Class, CompilerUserWidgetPropertyContexts[Result.UserWidgetPropertyContextIndex].PropertyName, PropertyPath.GetFields());
 		}
 	}
 	else
@@ -1246,10 +1256,10 @@ TArray<FMVVMConstFieldVariant> FMVVMViewBlueprintCompiler::CreateBindingDestinat
 		const FMVVMBlueprintViewModelContext* SourceViewModelContext = BlueprintView->FindViewModel(PropertyPath.GetViewModelId());
 		check(SourceViewModelContext);
 		FName DestinationName = SourceViewModelContext->GetViewModelName();
-		const int32 DestinationVariableContextIndex = CompilerSourceContexts.IndexOfByPredicate([DestinationName](const FCompilerUserWidgetPropertyContext& Other) { return Other.PropertyName == DestinationName; });
+		const int32 DestinationVariableContextIndex = CompilerUserWidgetPropertyContexts.IndexOfByPredicate([DestinationName](const FCompilerUserWidgetPropertyContext& Other) { return Other.PropertyName == DestinationName; });
 		check(DestinationVariableContextIndex != INDEX_NONE);
 
-		return CreatePropertyPath(Class, CompilerSourceContexts[DestinationVariableContextIndex].PropertyName, PropertyPath.GetFields());
+		return CreatePropertyPath(Class, CompilerUserWidgetPropertyContexts[DestinationVariableContextIndex].PropertyName, PropertyPath.GetFields());
 	}
 	else if (PropertyPath.IsFromWidget())
 	{
@@ -1262,10 +1272,10 @@ TArray<FMVVMConstFieldVariant> FMVVMViewBlueprintCompiler::CreateBindingDestinat
 		}
 		else
 		{
-			const int32 DestinationVariableContextIndex = CompilerSourceContexts.IndexOfByPredicate([DestinationName](const FCompilerUserWidgetPropertyContext& Other) { return Other.PropertyName == DestinationName; });
+			const int32 DestinationVariableContextIndex = CompilerUserWidgetPropertyContexts.IndexOfByPredicate([DestinationName](const FCompilerUserWidgetPropertyContext& Other) { return Other.PropertyName == DestinationName; });
 			if (ensureAlwaysMsgf(DestinationVariableContextIndex != INDEX_NONE, TEXT("Could not find source context for destination '%s'"), *DestinationName.ToString()))
 			{
-				return CreatePropertyPath(Class, CompilerSourceContexts[DestinationVariableContextIndex].PropertyName, PropertyPath.GetFields());
+				return CreatePropertyPath(Class, CompilerUserWidgetPropertyContexts[DestinationVariableContextIndex].PropertyName, PropertyPath.GetFields());
 			}
 		}
 	}
diff --git a/Engine/Plugins/Runtime/ModelViewViewModel/Source/ModelViewViewModelBlueprint/Private/MVVMViewBlueprintCompiler.h b/Engine/Plugins/Runtime/ModelViewViewModel/Source/ModelViewViewModelBlueprint/Private/MVVMViewBlueprintCompiler.h
--- a/Engine/Plugins/Runtime/ModelViewViewModel/Source/ModelViewViewModelBlueprint/Private/MVVMViewBlueprintCompiler.h
+++ b/Engine/Plugins/Runtime/ModelViewViewModel/Source/ModelViewViewModelBlueprint/Private/MVVMViewBlueprintCompiler.h
@@ -94,10 +94,12 @@ private:
 		FString CategoryName;
 		FString BlueprintSetter;
 		FMVVMConstFieldVariant Field;
+		// If the class is a viewmodel on the userwidget, what is the id of that viewmodel
+		FGuid ViewModelId;
 
 		bool bExposeOnSpawn = false;
 	};
-	TArray<FCompilerUserWidgetPropertyContext> CompilerSourceContexts;
+	TArray<FCompilerUserWidgetPropertyContext> CompilerUserWidgetPropertyContexts;
 
 	/** 
 	 * Describe the data initialize the view's properties/viewmodels.
@@ -127,7 +129,7 @@ private:
 	struct FCompilerBinding
 	{
 		int32 BindingIndex = INDEX_NONE;
-		int32 CompilerSourceContextIndex = INDEX_NONE;
+		int32 UserWidgetPropertyContextIndex = INDEX_NONE;
 		bool bSourceIsUserWidget = false;
 		bool bFieldIdNeeded = false;
 		bool bIsConversionFunctionComplex = false;
