MVVM - Added new borders to Bindings view as per design.

#rnx

diff --git a/Engine/Plugins/Runtime/ModelViewViewModel/Source/ModelViewViewModelEditor/Private/Styling/MVVMEditorStyle.cpp b/Engine/Plugins/Runtime/ModelViewViewModel/Source/ModelViewViewModelEditor/Private/Styling/MVVMEditorStyle.cpp
--- a/Engine/Plugins/Runtime/ModelViewViewModel/Source/ModelViewViewModelEditor/Private/Styling/MVVMEditorStyle.cpp
+++ b/Engine/Plugins/Runtime/ModelViewViewModel/Source/ModelViewViewModelEditor/Private/Styling/MVVMEditorStyle.cpp
@@ -38,10 +38,29 @@ FMVVMEditorStyle::FMVVMEditorStyle()
 
 	Set("BlueprintView.TabIcon", new IMAGE_BRUSH_SVG("Slate/ViewModel", Icon16x16));
 
-	Set("BindingView.ManageViewModels", new IMAGE_BRUSH_SVG("Slate/ViewModel", Icon16x16));
 	Set("BindingView.AddBinding", new IMAGE_BRUSH_SVG("Slate/ViewModel_AddBinding", Icon16x16));
 	Set("BindingView.Background", new FSlateRoundedBoxBrush(FStyleColors::Recessed, 0.0f, FStyleColors::Panel, 4.0f));
 	Set("BindingView.ViewModelWarning", new FSlateRoundedBoxBrush(FStyleColors::Transparent, 4.0f, FStyleColors::Hover, 1.0f));
+	Set("BindingView.WidgetRow", FTableRowStyle(FAppStyle::Get().GetWidgetStyle<FTableRowStyle>("TableView.Row"))
+		.SetEvenRowBackgroundBrush(*FAppStyle::Get().GetBrush("Brushes.Header"))
+		.SetOddRowBackgroundBrush(*FAppStyle::Get().GetBrush("Brushes.Header"))
+		.SetEvenRowBackgroundHoveredBrush(*FAppStyle::Get().GetBrush("Brushes.Dropdown"))
+		.SetOddRowBackgroundHoveredBrush(*FAppStyle::Get().GetBrush("Brushes.Dropdown"))
+	);
+
+	Set("BindingView.BindingRow", FTableRowStyle(FAppStyle::Get().GetWidgetStyle<FTableRowStyle>("TableView.Row"))
+		.SetEvenRowBackgroundBrush(*FAppStyle::Get().GetBrush("Brushes.Panel"))
+		.SetOddRowBackgroundBrush(*FAppStyle::Get().GetBrush("Brushes.Panel"))
+		.SetEvenRowBackgroundHoveredBrush(*FAppStyle::Get().GetBrush("Brushes.Header"))
+		.SetOddRowBackgroundHoveredBrush(*FAppStyle::Get().GetBrush("Brushes.Header"))
+	);
+
+	Set("BindingView.ParameterRow", FTableRowStyle(FAppStyle::Get().GetWidgetStyle<FTableRowStyle>("TableView.Row"))
+		.SetEvenRowBackgroundBrush(*FAppStyle::Get().GetBrush("Brushes.Panel"))
+		.SetOddRowBackgroundBrush(*FAppStyle::Get().GetBrush("Brushes.Panel"))
+		.SetEvenRowBackgroundHoveredBrush(*FAppStyle::Get().GetBrush("Brushes.Header"))
+		.SetOddRowBackgroundHoveredBrush(*FAppStyle::Get().GetBrush("Brushes.Header"))
+	);
 
 	// ViewModelSelectionWidget Icons
 	{
@@ -70,6 +89,8 @@ FMVVMEditorStyle::FMVVMEditorStyle()
 		FComboButtonStyle(FAppStyle::Get().GetWidgetStyle<FComboButtonStyle>("ComboButton"))
 		.SetDownArrowImage(*FCoreStyle::Get().GetBrush("Icons.Edit")));
 
+	Set("FunctionParameter.Border", new FSlateColorBrush(FStyleColors::Input));
+
 	FSlateStyleRegistry::RegisterSlateStyle(*this);
 }
 
diff --git a/Engine/Plugins/Runtime/ModelViewViewModel/Source/ModelViewViewModelEditor/Private/Widgets/SMVVMViewBindingListView.cpp b/Engine/Plugins/Runtime/ModelViewViewModel/Source/ModelViewViewModelEditor/Private/Widgets/SMVVMViewBindingListView.cpp
--- a/Engine/Plugins/Runtime/ModelViewViewModel/Source/ModelViewViewModelEditor/Private/Widgets/SMVVMViewBindingListView.cpp
+++ b/Engine/Plugins/Runtime/ModelViewViewModel/Source/ModelViewViewModelEditor/Private/Widgets/SMVVMViewBindingListView.cpp
@@ -19,6 +19,7 @@
 #include "MVVMEditorSubsystem.h"
 #include "MVVMSubsystem.h"
 #include "MVVMWidgetBlueprintExtension_View.h"
+#include "ScopedTransaction.h"
 #include "SEnumCombo.h"
 #include "SPrimaryButton.h"
 #include "SSimpleButton.h" 
@@ -177,7 +178,8 @@ public:
 
 		STableRow<TSharedPtr<FBindingEntry>>::Construct(
 			STableRow<TSharedPtr<FBindingEntry>>::FArguments()
-			.Padding(1.0f)
+			.Padding(2.0f)
+			.Style(FMVVMEditorStyle::Get(), "BindingView.WidgetRow")
 			[
 				SNew(SHorizontalBox)
 				+ SHorizontalBox::Slot()
@@ -211,6 +213,24 @@ public:
 			],
 			OwnerTableView
 		);
+
+		TSharedPtr<SWidget> ChildContent = ChildSlot.DetachWidget();
+		ChildSlot
+		[
+			/* Add a single pixel top and bottom border for this widget. */
+			SNew(SBorder)
+			.BorderImage(FAppStyle::Get().GetBrush("Brushes.Recessed"))
+			.Padding(0, 1)
+			[
+				/* Restore the border that we're meant to have that reacts to selection/hover/etc. */
+				SNew(SBorder)
+				.BorderImage(this, &SWidgetRow::GetBorderImage)
+				.Padding(0)
+				[
+					ChildContent.ToSharedRef()
+				]
+			]
+		];
 	}
 
 private:
@@ -278,133 +298,156 @@ public:
 
 		STableRow<TSharedPtr<FBindingEntry>>::Construct(
 			STableRow<TSharedPtr<FBindingEntry>>::FArguments()
-			.Padding(1.0f)
+			.ShowWires(true)
+			.Style(FMVVMEditorStyle::Get(), "BindingView.BindingRow")
 			[
-				SNew(SHorizontalBox)
-
-				+ SHorizontalBox::Slot()
-				.Padding(2, 0)
-				.VAlign(VAlign_Center)
-				.HAlign(HAlign_Left) 
-				.AutoWidth()
+				SNew(SBox)
+				.HeightOverride(30)
 				[
-					SNew(SCheckBox)
-					.IsChecked(this, &SBindingRow::IsBindingCompiled)
-					.OnCheckStateChanged(this, &SBindingRow::OnIsBindingCompileChanged)
-				]
+					SNew(SHorizontalBox)
 
-				+ SHorizontalBox::Slot()
-				.Padding(2, 0)
-				.VAlign(VAlign_Center)
-				.HAlign(HAlign_Left) 
-				.AutoWidth()
-				[
-					SNew(SSimpleButton)
-					.Icon(FAppStyle::Get().GetBrush("Icons.Error"))
-					.Visibility(this, &SBindingRow::GetErrorVisibility)
-					.ToolTipText(this, &SBindingRow::GetErrorToolTip)
-					.OnClicked(this, &SBindingRow::OnErrorButtonClicked)
-				]
+					+ SHorizontalBox::Slot()
+					.Padding(2, 0)
+					.VAlign(VAlign_Center)
+					.HAlign(HAlign_Left) 
+					.AutoWidth()
+					[
+						SNew(SCheckBox)
+						.IsChecked(this, &SBindingRow::IsBindingCompiled)
+						.OnCheckStateChanged(this, &SBindingRow::OnIsBindingCompileChanged)
+					]
 
-				+ SHorizontalBox::Slot()
-				.Padding(4, 0)
-				.VAlign(VAlign_Center)
-				.HAlign(HAlign_Left)
-				.AutoWidth()
-				[
-					SNew(SBox)
-					.MinDesiredWidth(150)
+					+ SHorizontalBox::Slot()
+					.Padding(2, 0)
+					.VAlign(VAlign_Center)
+					.HAlign(HAlign_Left) 
+					.AutoWidth()
 					[
-						SAssignNew(WidgetFieldSelector, SFieldSelector, InWidgetBlueprint, false)
-						.BindingMode(this, &SBindingRow::GetCurrentBindingMode)
-						.SelectedField(this, &SBindingRow::GetSelectedWidgetProperty)
-						.OnFieldSelectionChanged(this, &SBindingRow::OnWidgetPropertySelected)
-						.SelectedConversionFunction(this, &SBindingRow::GetSelectedConversionFunction, false)
-						.OnConversionFunctionSelectionChanged(this, &SBindingRow::OnConversionFunctionChanged, false)
-						.ShowConversionFunctions(this, &SBindingRow::ShouldShowConversionFunctions, false)
-						.ShowSource(false)
-						.Source(WidgetSource)
-						.AssignableTo(this, &SBindingRow::GetAssignableToProperty, false)
+						SNew(SSimpleButton)
+						.Icon(FAppStyle::Get().GetBrush("Icons.Error"))
+						.Visibility(this, &SBindingRow::GetErrorVisibility)
+						.ToolTipText(this, &SBindingRow::GetErrorToolTip)
+						.OnClicked(this, &SBindingRow::OnErrorButtonClicked)
 					]
-				]
 
-				+ SHorizontalBox::Slot()
-				.Padding(2, 0)
-				.VAlign(VAlign_Fill)
-				.HAlign(HAlign_Left) 
-				.AutoWidth()
-				[
-					SNew(SComboBox<FName>)
-					.OptionsSource(Private::GetBindingModeNames())
-					.InitiallySelectedItem(StaticEnum<EMVVMBindingMode>()->GetNameByValue((int64) ViewBinding->BindingType))
-					.OnSelectionChanged(this, &SBindingRow::OnBindingModeSelectionChanged)
-					.OnGenerateWidget(this, &SBindingRow::GenerateBindingModeWidget)
-					.ToolTipText(this, &SBindingRow::GetCurrentBindingModeLabel)
-					.Content()
+					+ SHorizontalBox::Slot()
+					.Padding(4, 0)
+					.VAlign(VAlign_Center)
+					.HAlign(HAlign_Left)
+					.AutoWidth()
 					[
 						SNew(SBox)
-						.HAlign(HAlign_Center)
-						.VAlign(VAlign_Center)
-						.WidthOverride(16)
-						.HeightOverride(16)
+						.MinDesiredWidth(150)
 						[
-							SNew(SImage)
-							.Image(this, &SBindingRow::GetCurrentBindingModeBrush)
+							SAssignNew(WidgetFieldSelector, SFieldSelector, InWidgetBlueprint, false)
+							.BindingMode(this, &SBindingRow::GetCurrentBindingMode)
+							.SelectedField(this, &SBindingRow::GetSelectedWidgetProperty)
+							.OnFieldSelectionChanged(this, &SBindingRow::OnWidgetPropertySelected)
+							.SelectedConversionFunction(this, &SBindingRow::GetSelectedConversionFunction, false)
+							.OnConversionFunctionSelectionChanged(this, &SBindingRow::OnConversionFunctionChanged, false)
+							.ShowConversionFunctions(this, &SBindingRow::ShouldShowConversionFunctions, false)
+							.ShowSource(false)
+							.Source(WidgetSource)
+							.AssignableTo(this, &SBindingRow::GetAssignableToProperty, false)
+						]
+					]
+
+					+ SHorizontalBox::Slot()
+					.Padding(4, 0)
+					.VAlign(VAlign_Fill)
+					.HAlign(HAlign_Left) 
+					.AutoWidth()
+					[
+						SNew(SComboBox<FName>)
+						.OptionsSource(Private::GetBindingModeNames())
+						.InitiallySelectedItem(StaticEnum<EMVVMBindingMode>()->GetNameByValue((int64) ViewBinding->BindingType))
+						.OnSelectionChanged(this, &SBindingRow::OnBindingModeSelectionChanged)
+						.OnGenerateWidget(this, &SBindingRow::GenerateBindingModeWidget)
+						.ToolTipText(this, &SBindingRow::GetCurrentBindingModeLabel)
+						.Content()
+						[
+							SNew(SBox)
+							.HAlign(HAlign_Center)
+							.VAlign(VAlign_Center)
+							.WidthOverride(16)
+							.HeightOverride(16)
+							[
+								SNew(SImage)
+								.Image(this, &SBindingRow::GetCurrentBindingModeBrush)
+							]
 						]
 					]
-				]
 				
-				+ SHorizontalBox::Slot()
-				.Padding(4, 0, 2, 0)
-				.VAlign(VAlign_Center)
-				.HAlign(HAlign_Left)
-				.AutoWidth()
-				[
-					SNew(SBox)
-					.MinDesiredWidth(150)
+					+ SHorizontalBox::Slot()
+					.Padding(4, 0, 2, 0)
+					.VAlign(VAlign_Center)
+					.HAlign(HAlign_Left)
+					.AutoWidth()
 					[
-						SAssignNew(ViewModelFieldSelector, SFieldSelector, InWidgetBlueprint, true)
-						.BindingMode(this, &SBindingRow::GetCurrentBindingMode)
-						.SelectedField(this, &SBindingRow::GetSelectedViewModelProperty)
-						.OnFieldSelectionChanged(this, &SBindingRow::OnViewModelPropertySelected)
-						.SelectedConversionFunction(this, &SBindingRow::GetSelectedConversionFunction, true)
-						.OnConversionFunctionSelectionChanged(this, &SBindingRow::OnConversionFunctionChanged, true)
-						.ShowConversionFunctions(this, &SBindingRow::ShouldShowConversionFunctions, true)
-						.AssignableTo(this, &SBindingRow::GetAssignableToProperty, true)
+						SNew(SBox)
+						.MinDesiredWidth(150)
+						[
+							SAssignNew(ViewModelFieldSelector, SFieldSelector, InWidgetBlueprint, true)
+							.BindingMode(this, &SBindingRow::GetCurrentBindingMode)
+							.SelectedField(this, &SBindingRow::GetSelectedViewModelProperty)
+							.OnFieldSelectionChanged(this, &SBindingRow::OnViewModelPropertySelected)
+							.SelectedConversionFunction(this, &SBindingRow::GetSelectedConversionFunction, true)
+							.OnConversionFunctionSelectionChanged(this, &SBindingRow::OnConversionFunctionChanged, true)
+							.ShowConversionFunctions(this, &SBindingRow::ShouldShowConversionFunctions, true)
+							.AssignableTo(this, &SBindingRow::GetAssignableToProperty, true)
+						]
 					]
-				]
 
-				+ SHorizontalBox::Slot()
-				[
-					SNew(SSpacer)
-				]
+					+ SHorizontalBox::Slot()
+					[
+						SNew(SSpacer)
+					]
 
-				+ SHorizontalBox::Slot()
-				.Padding(2, 1)
-				.VAlign(VAlign_Fill)
-				.HAlign(HAlign_Right)
-				.AutoWidth()
-				[
-					SNew(SCheckBox)
-					.IsChecked(this, &SBindingRow::IsExecutionModeOverrideChecked)
-					.OnCheckStateChanged(this, &SBindingRow::OnExecutionModeOverrideChanged)
-				]
+					+ SHorizontalBox::Slot()
+					.Padding(2, 1)
+					.VAlign(VAlign_Fill)
+					.HAlign(HAlign_Right)
+					.AutoWidth()
+					[
+						SNew(SCheckBox)
+						.IsChecked(this, &SBindingRow::IsExecutionModeOverrideChecked)
+						.OnCheckStateChanged(this, &SBindingRow::OnExecutionModeOverrideChanged)
+					]
 
-				+ SHorizontalBox::Slot()
-				.Padding(2, 1)
-				.VAlign(VAlign_Fill)
-				.HAlign(HAlign_Right)
-				.AutoWidth()
-				[
-					SNew(SEnumComboBox, StaticEnum<EMVVMExecutionMode>())
-					.ContentPadding(FMargin(4, 0))
-					.OnEnumSelectionChanged(this, &SBindingRow::OnExecutionModeSelectionChanged)
-					.CurrentValue(this, &SBindingRow::GetExecutioModeValue)
-					.IsEnabled(this, &SBindingRow::IsExecutionModeOverridden)
+					+ SHorizontalBox::Slot()
+					.Padding(2, 1)
+					.VAlign(VAlign_Fill)
+					.HAlign(HAlign_Right)
+					.AutoWidth()
+					[
+						SNew(SEnumComboBox, StaticEnum<EMVVMExecutionMode>())
+						.ContentPadding(FMargin(4, 0))
+						.OnEnumSelectionChanged(this, &SBindingRow::OnExecutionModeSelectionChanged)
+						.CurrentValue(this, &SBindingRow::GetExecutioModeValue)
+						.IsEnabled(this, &SBindingRow::IsExecutionModeOverridden)
+					]
 				]
 			],
 			OwnerTableView
 		);
+
+		TSharedPtr<SWidget> ChildContent = ChildSlot.DetachWidget();
+		ChildSlot
+		[
+			/* Add a single pixel top and bottom border for this widget. */
+			SNew(SBorder)
+			.BorderImage(FAppStyle::Get().GetBrush("Brushes.Recessed"))
+			.Padding(0, 2, 0, 1)
+			[
+				/* Restore the border that we're meant to have that reacts to selection/hover/etc. */
+				SNew(SBorder)
+				.BorderImage(this, &SBindingRow::GetBorderImage)
+				.Padding(0)
+				[
+					ChildContent.ToSharedRef()
+				]
+			]
+		];
 	}
 
 	~SBindingRow()
@@ -458,7 +501,7 @@ private:
 
 	EVisibility GetErrorVisibility() const
 	{
-		return GetThisViewBinding()->Errors.IsEmpty() ? EVisibility::Hidden : EVisibility::Visible;
+		return GetThisViewBinding()->Errors.IsEmpty() ? EVisibility::Collapsed : EVisibility::Visible;
 	}
 
 	FText GetErrorToolTip() const
@@ -940,20 +983,20 @@ class SFunctionParameterRow : public STableRow<TSharedPtr<FBindingEntry>>
 
 		STableRow<TSharedPtr<FBindingEntry>>::Construct(
 			STableRow<TSharedPtr<FBindingEntry>>::FArguments()
-		[
-			SNew(SBox)
-			.HeightOverride(30)
-			.ToolTipText(ToolTip)
+			.ShowWires(true)
+			.Style(FMVVMEditorStyle::Get(), "BindingView.ParameterRow")
 			[
-				SNew(SHorizontalBox)
-				+ SHorizontalBox::Slot()
-				.HAlign(HAlign_Left)
-				.VAlign(VAlign_Center)
-				.Padding(0, 0, 4, 0)
-				.AutoWidth()
+				SNew(SBox)
+				.HeightOverride(30)
+				.ToolTipText(ToolTip)
 				[
-					SNew(SBox)
-					.MinDesiredWidth(100)
+					SNew(SHorizontalBox)
+
+					+ SHorizontalBox::Slot()
+					.HAlign(HAlign_Left)
+					.VAlign(VAlign_Center)
+					.Padding(0, 0, 8, 0)
+					.AutoWidth()
 					[
 						SNew(SHorizontalBox)
 						+ SHorizontalBox::Slot()
@@ -972,26 +1015,27 @@ class SFunctionParameterRow : public STableRow<TSharedPtr<FBindingEntry>>
 							.Text(DisplayName)
 						]
 					]
-				]
-				+ SHorizontalBox::Slot()
-				.HAlign(HAlign_Left)
-				.VAlign(VAlign_Center)
-				.AutoWidth()
-				[
-					SNew(SBox)
-					.MinDesiredWidth(200)
+					+ SHorizontalBox::Slot()
+					.HAlign(HAlign_Left)
+					.VAlign(VAlign_Center)
+					.AutoWidth()
 					[
-						SNew(SFunctionParameter)
-						.OnGetBindingMode(this, &SFunctionParameterRow::OnGetBindingMode)
-						.WidgetBlueprint(InWidgetBlueprint)
-						.Binding(Binding)
-						.ParameterName(Entry->GetName())
-						.SourceToDestination(bSourceToDestination)
-						.AllowDefault(!bSimpleConversionFunction)
+						SNew(SBox)
+						.MinDesiredWidth(200)
+						[
+							SNew(SFunctionParameter)
+							.OnGetBindingMode(this, &SFunctionParameterRow::OnGetBindingMode)
+							.WidgetBlueprint(InWidgetBlueprint)
+							.Binding(Binding)
+							.ParameterName(Entry->GetName())
+							.SourceToDestination(bSourceToDestination)
+							.AllowDefault(!bSimpleConversionFunction)
+						]
 					]
 				]
-			]
-		], OwnerTableView);
+			], 
+			OwnerTableView
+		);
 	}
 
 private:
@@ -1243,11 +1287,15 @@ void SBindingsList::OnDeleteSelected()
 
 			if (FMessageDialog::Open(EAppMsgType::YesNo, EAppReturnType::Yes, Message, &Title) == EAppReturnType::Yes)
 			{
+				FScopedTransaction Transaction(LOCTEXT("DeleteBindingsTransaction", "Delete Bindings"));
+
 				if (TSharedPtr<SBindingsPanel> BindingPanelPtr = BindingPanel.Pin())
 				{
 					BindingPanelPtr->OnBindingListSelectionChanged(TConstArrayView<FMVVMBlueprintViewBinding*>());
 				}
 
+				BlueprintView->Modify();
+
 				for (const FMVVMBlueprintViewBinding* Binding : BindingsToRemove)
 				{
 					BlueprintView->RemoveBinding(Binding);
