[ModelViewViewModel]
* Fixed non unity vs2019 compile errors by adding includes

diff --git a/Engine/Plugins/Runtime/ModelViewViewModel/Source/ModelViewViewModel/Private/Bindings/MVVMFieldPathHelper.cpp b/Engine/Plugins/Runtime/ModelViewViewModel/Source/ModelViewViewModel/Private/Bindings/MVVMFieldPathHelper.cpp
--- a/Engine/Plugins/Runtime/ModelViewViewModel/Source/ModelViewViewModel/Private/Bindings/MVVMFieldPathHelper.cpp
+++ b/Engine/Plugins/Runtime/ModelViewViewModel/Source/ModelViewViewModel/Private/Bindings/MVVMFieldPathHelper.cpp
@@ -9,6 +9,7 @@
 #include "Templates/ValueOrError.h"
 #include "Types/MVVMAvailableBinding.h"
 #include "Types/MVVMFieldContext.h"
+#include "Types/MVVMViewModelContextInstance.h"
 
 #define LOCTEXT_NAMESPACE "MVVMFieldPathHelper"
 
diff --git a/Engine/Plugins/Runtime/ModelViewViewModel/Source/ModelViewViewModel/Private/Debugging/MVVMDebugging.cpp b/Engine/Plugins/Runtime/ModelViewViewModel/Source/ModelViewViewModel/Private/Debugging/MVVMDebugging.cpp
--- a/Engine/Plugins/Runtime/ModelViewViewModel/Source/ModelViewViewModel/Private/Debugging/MVVMDebugging.cpp
+++ b/Engine/Plugins/Runtime/ModelViewViewModel/Source/ModelViewViewModel/Private/Debugging/MVVMDebugging.cpp
@@ -1,11 +1,14 @@
 // Copyright Epic Games, Inc. All Rights Reserved.
 
 #include "Debugging/MVVMDebugging.h"
+
+#if UE_WITH_MVVM_DEBUGGING
+
+#include "FieldNotification/FieldId.h"
 #include "View/MVVMView.h"
 
 #define LOCTEXT_NAMESPACE "MVVMDebugging"
 
-#if UE_WITH_MVVM_DEBUGGING
 namespace UE::MVVM
 {
 /** */
@@ -99,6 +102,7 @@ void FDebugging::BroadcastLibraryBindingExecuted(const UMVVMView* InView, const
 	}
 }
 } // namespace UE::MVVM
-#endif //UE_WITH_MVVM_DEBUGGING
 
-#undef LOCTEXT_NAMESPACE
\ No newline at end of file
+#undef LOCTEXT_NAMESPACE
+
+#endif //UE_WITH_MVVM_DEBUGGING
diff --git a/Engine/Plugins/Runtime/ModelViewViewModel/Source/ModelViewViewModel/Private/MVVMSubsystem.cpp b/Engine/Plugins/Runtime/ModelViewViewModel/Source/ModelViewViewModel/Private/MVVMSubsystem.cpp
--- a/Engine/Plugins/Runtime/ModelViewViewModel/Source/ModelViewViewModel/Private/MVVMSubsystem.cpp
+++ b/Engine/Plugins/Runtime/ModelViewViewModel/Source/ModelViewViewModel/Private/MVVMSubsystem.cpp
@@ -7,6 +7,7 @@
 #include "Templates/ValueOrError.h"
 #include "View/MVVMView.h"
 #include "Types/MVVMAvailableBinding.h"
+#include "Types/MVVMViewModelContextInstance.h"
 
 #include UE_INLINE_GENERATED_CPP_BY_NAME(MVVMSubsystem)
 
diff --git a/Engine/Plugins/Runtime/ModelViewViewModel/Source/ModelViewViewModel/Private/View/MVVMViewClass.cpp b/Engine/Plugins/Runtime/ModelViewViewModel/Source/ModelViewViewModel/Private/View/MVVMViewClass.cpp
--- a/Engine/Plugins/Runtime/ModelViewViewModel/Source/ModelViewViewModel/Private/View/MVVMViewClass.cpp
+++ b/Engine/Plugins/Runtime/ModelViewViewModel/Source/ModelViewViewModel/Private/View/MVVMViewClass.cpp
@@ -2,6 +2,7 @@
 
 #include "View/MVVMViewClass.h"
 #include "Types/MVVMFieldContext.h"
+#include "Types/MVVMViewModelContextInstance.h"
 #include "View/MVVMView.h"
 
 #include "Bindings/MVVMFieldPathHelper.h"
