MVVM: Toggling the Bind button on a function parameter now restores the previously selected field instead of clearing it.

Non-boolean function parameters now have a dark background. Booleans are displayed as checkboxes and don't scale to the minimum width, so adding the background for them looks bad.

Tested in Bindings view with a conversion function: bind to viewmodel property, toggle binding off, toggle binding back on again. Observe that the viewmodel property is still bound.

diff --git a/Engine/Plugins/Runtime/ModelViewViewModel/Source/ModelViewViewModelEditor/Private/Widgets/SMVVMFunctionParameter.cpp b/Engine/Plugins/Runtime/ModelViewViewModel/Source/ModelViewViewModelEditor/Private/Widgets/SMVVMFunctionParameter.cpp
--- a/Engine/Plugins/Runtime/ModelViewViewModel/Source/ModelViewViewModelEditor/Private/Widgets/SMVVMFunctionParameter.cpp
+++ b/Engine/Plugins/Runtime/ModelViewViewModel/Source/ModelViewViewModelEditor/Private/Widgets/SMVVMFunctionParameter.cpp
@@ -1,11 +1,14 @@
 // Copyright Epic Games, Inc. All Rights Reserved.
 
 #include "SMVVMFunctionParameter.h" 
+#include "EdGraphSchema_K2.h"
 #include "Editor.h"
 #include "MVVMBlueprintView.h"
 #include "MVVMEditorSubsystem.h"
 #include "NodeFactory.h"
 #include "SGraphPin.h"
+#include "Styling/MVVMEditorStyle.h"
+#include "Types/MVVMBindingMode.h"
 #include "WidgetBlueprint.h"
 #include "Widgets/Images/SImage.h"
 #include "Widgets/SBoxPanel.h"
@@ -55,12 +58,27 @@ void SFunctionParameter::Construct(const FArguments& InArgs)
 		}
 	}
 
+	const bool bIsBooleanPin = Pin->PinType.PinCategory == UEdGraphSchema_K2::PC_Boolean;
+	
 	if (ValueWidget == SNullWidget::NullWidget)
 	{
 		ValueWidget = SNew(STextBlock)
 			.Text(LOCTEXT("DefaultValue", "Default Value"))
 			.TextStyle(FAppStyle::Get(), "HintText");
 	}
+	// booleans are represented by a checkbox which doesn't expand to the minsize we have, so don't put a border around them
+	else if (!bIsBooleanPin)
+	{
+		ValueWidget = SNew(SBorder)
+			.Padding(0)
+			.BorderImage(FMVVMEditorStyle::Get().GetBrush("FunctionParameter.Border"))
+			[
+				ValueWidget
+			];
+	}
+
+	FMVVMBlueprintPropertyPath Path = OnGetSelectedField();
+	bDefaultValueVisible = Path.IsEmpty();
 
 	ValueWidget->SetVisibility(TAttribute<EVisibility>::CreateSP(this, &SFunctionParameter::OnGetVisibility, true));
 
@@ -70,7 +88,7 @@ void SFunctionParameter::Construct(const FArguments& InArgs)
 	ChildSlot
 	[
 		SNew(SBox)
-		.MinDesiredWidth(100)
+		.MinDesiredWidth(bIsBooleanPin ? FOptionalSize() : 100)
 		[
 			SAssignNew(HBox, SHorizontalBox)
 			+ SHorizontalBox::Slot()
@@ -129,39 +147,30 @@ EVisibility SFunctionParameter::OnGetVisibility(bool bDefaultValue) const
 		return bDefaultValue ? EVisibility::Collapsed : EVisibility::Visible;
 	}
 
-	FMVVMBlueprintPropertyPath Path = OnGetSelectedField();
-
 	// if we're not bound then show the default value widget, otherwise show the binding widget
-	return Path.IsEmpty() == bDefaultValue ? EVisibility::Visible : EVisibility::Collapsed;
+	return bDefaultValue == bDefaultValueVisible ? EVisibility::Visible : EVisibility::Collapsed;
 }
 
 ECheckBoxState SFunctionParameter::OnGetIsBindArgumentChecked() const 
 {
-	UMVVMEditorSubsystem* EditorSubsystem = GEditor->GetEditorSubsystem<UMVVMEditorSubsystem>();
-	FMVVMBlueprintPropertyPath CurrentPath = EditorSubsystem->GetPathForConversionFunctionArgument(WidgetBlueprint.Get(), *Binding, ParameterName, bSourceToDestination);
-	return CurrentPath.IsEmpty() ? ECheckBoxState::Unchecked : ECheckBoxState::Checked;
+	return bDefaultValueVisible ? ECheckBoxState::Unchecked : ECheckBoxState::Checked;
 }
 
 void SFunctionParameter::OnBindArgumentChecked(ECheckBoxState Checked)
 {
 	UMVVMEditorSubsystem* EditorSubsystem = GEditor->GetEditorSubsystem<UMVVMEditorSubsystem>();
 
-	FMVVMBlueprintPropertyPath Path;
+	bDefaultValueVisible = (Checked != ECheckBoxState::Checked);
 
-	if (Checked == ECheckBoxState::Checked)
+	if (bDefaultValueVisible)
 	{
-		// HACK: Just set a placeholder viewmodel reference
-		if (const UMVVMBlueprintView* View = EditorSubsystem->GetView(WidgetBlueprint.Get()))
-		{
-			const TArrayView<const FMVVMBlueprintViewModelContext> ViewModels = View->GetViewModels();
-			if (ViewModels.Num() > 0)
-			{
-				Path.SetViewModelId(ViewModels[0].GetViewModelId());
-			}
-		}
+		PreviousSelectedField = OnGetSelectedField();
+		SetSelectedField(FMVVMBlueprintPropertyPath());
+	}
+	else
+	{
+		SetSelectedField(PreviousSelectedField);
 	}
-
-	EditorSubsystem->SetPathForConversionFunctionArgument(WidgetBlueprint.Get(), *Binding, ParameterName, Path, bSourceToDestination);
 }
 
 FMVVMBlueprintPropertyPath SFunctionParameter::OnGetSelectedField() const
@@ -170,10 +179,15 @@ FMVVMBlueprintPropertyPath SFunctionParameter::OnGetSelectedField() const
 	return EditorSubsystem->GetPathForConversionFunctionArgument(WidgetBlueprint.Get(), *Binding, ParameterName, bSourceToDestination);
 }
 
-void SFunctionParameter::OnFieldSelectionChanged(FMVVMBlueprintPropertyPath Selected)
+void SFunctionParameter::SetSelectedField(const FMVVMBlueprintPropertyPath& Path)
 {
 	UMVVMEditorSubsystem* EditorSubsystem = GEditor->GetEditorSubsystem<UMVVMEditorSubsystem>();
-	EditorSubsystem->SetPathForConversionFunctionArgument(WidgetBlueprint.Get(), *Binding, ParameterName, Selected, bSourceToDestination);
+	EditorSubsystem->SetPathForConversionFunctionArgument(WidgetBlueprint.Get(), *Binding, ParameterName, Path, bSourceToDestination);
+}
+
+void SFunctionParameter::OnFieldSelectionChanged(FMVVMBlueprintPropertyPath Selected)
+{
+	SetSelectedField(Selected);
 }
 
 }
diff --git a/Engine/Plugins/Runtime/ModelViewViewModel/Source/ModelViewViewModelEditor/Private/Widgets/SMVVMFunctionParameter.h b/Engine/Plugins/Runtime/ModelViewViewModel/Source/ModelViewViewModelEditor/Private/Widgets/SMVVMFunctionParameter.h
--- a/Engine/Plugins/Runtime/ModelViewViewModel/Source/ModelViewViewModelEditor/Private/Widgets/SMVVMFunctionParameter.h
+++ b/Engine/Plugins/Runtime/ModelViewViewModel/Source/ModelViewViewModelEditor/Private/Widgets/SMVVMFunctionParameter.h
@@ -2,10 +2,9 @@
 
 #pragma once
 
+#include "MVVMPropertyPath.h"
 #include "Widgets/SCompoundWidget.h"
 
-struct FMVVMBlueprintPropertyPath;
-
 class SGraphPin;
 class UK2Node_CallFunction;
 class UWidgetBlueprint;
@@ -44,7 +43,10 @@ private:
 	ECheckBoxState OnGetIsBindArgumentChecked() const;
 	void OnBindArgumentChecked(ECheckBoxState Checked);
 	EVisibility OnGetVisibility(bool bDefaultValue) const;
+
 	FMVVMBlueprintPropertyPath OnGetSelectedField() const;
+	void SetSelectedField(const FMVVMBlueprintPropertyPath& Path);
+
 	void OnFieldSelectionChanged(FMVVMBlueprintPropertyPath Selected);
 
 private:
@@ -57,8 +59,12 @@ private:
 	TSharedPtr<UE::MVVM::SFieldSelector> FieldSelector;
 	FGetBindingMode GetBindingModeDelegate;
 
+	// Previously selected field value so that toggling bind on and off doesn't clear the reference.
+	FMVVMBlueprintPropertyPath PreviousSelectedField;
+
 	bool bSourceToDestination = true;
 	bool bAllowDefault = true;
+	bool bDefaultValueVisible = true;
 };
 
 }
