MVVM: Change the MVVM category for Viewmodel

diff --git a/Engine/Plugins/Runtime/ModelViewViewModel/Source/ModelViewViewModel/Public/MVVMSubsystem.h b/Engine/Plugins/Runtime/ModelViewViewModel/Source/ModelViewViewModel/Public/MVVMSubsystem.h
--- a/Engine/Plugins/Runtime/ModelViewViewModel/Source/ModelViewViewModel/Public/MVVMSubsystem.h
+++ b/Engine/Plugins/Runtime/ModelViewViewModel/Source/ModelViewViewModel/Public/MVVMSubsystem.h
@@ -32,16 +32,16 @@ public:
 	virtual void Deinitialize() override;
 	//~ End UEngineSubsystem interface
 
-	UFUNCTION(BlueprintCallable, Category="MVVM", meta = (DisplayName = "Get View From User Widget"))
+	UFUNCTION(BlueprintCallable, Category = "Viewmodel", meta = (DisplayName = "Get View From User Widget"))
 	UMVVMView* K2_GetViewFromUserWidget(const UUserWidget* UserWidget) const;
 
 	static UMVVMView* GetViewFromUserWidget(const UUserWidget* UserWidget);
 
-	UFUNCTION(BlueprintCallable, Category = "MVVM")
+	UFUNCTION(BlueprintCallable, Category = "Viewmodel")
 	bool DoesWidgetTreeContainedWidget(const UWidgetTree* WidgetTree, const UWidget* ViewWidget) const;
 
 	/** @return The list of all the AvailableBindings that are available for the Class. */
-	UFUNCTION(BlueprintCallable, Category = "MVVM", meta = (DisplayName = "Get Available Bindings"))
+	UFUNCTION(BlueprintCallable, Category = "Viewmodel", meta = (DisplayName = "Get Available Bindings"))
 	TArray<FMVVMAvailableBinding> K2_GetAvailableBindings(const UClass* Class, const UClass* Accessor) const;
 
 	static TArray<FMVVMAvailableBinding> GetAvailableBindings(const UClass* Class, const UClass* Accessor);
@@ -53,7 +53,7 @@ public:
 	static TArray<FMVVMAvailableBinding> GetAvailableBindingsForStruct(const UScriptStruct* Struct);
 
 	/** @return The AvailableBinding from a BindingName. */
-	UFUNCTION(BlueprintCallable, Category = "MVVM", meta = (DisplayName = "Get Available Binding"))
+	UFUNCTION(BlueprintCallable, Category = "Viewmodel", meta = (DisplayName = "Get Available Binding"))
 	FMVVMAvailableBinding K2_GetAvailableBinding(const UClass* Class, FMVVMBindingName BindingName, const UClass* Accessor) const;
 
 	static FMVVMAvailableBinding GetAvailableBinding(const UClass* Class, FMVVMBindingName BindingName, const UClass* Accessor);
@@ -61,7 +61,7 @@ public:
 	/** @return The AvailableBinding from a field. */
 	static FMVVMAvailableBinding GetAvailableBindingForField(UE::MVVM::FMVVMConstFieldVariant Variant, const UClass* Accessor);
 
-	UFUNCTION(BlueprintCallable, Category = "MVVM")
+	UFUNCTION(BlueprintCallable, Category = "Viewmodel")
 	UMVVMViewModelCollectionObject* GetGlobalViewModelCollection() const
 	{
 		return GlobalViewModelCollection;
diff --git a/Engine/Plugins/Runtime/ModelViewViewModel/Source/ModelViewViewModel/Public/Types/MVVMBindingName.h b/Engine/Plugins/Runtime/ModelViewViewModel/Source/ModelViewViewModel/Public/Types/MVVMBindingName.h
--- a/Engine/Plugins/Runtime/ModelViewViewModel/Source/ModelViewViewModel/Public/Types/MVVMBindingName.h
+++ b/Engine/Plugins/Runtime/ModelViewViewModel/Source/ModelViewViewModel/Public/Types/MVVMBindingName.h
@@ -53,7 +53,7 @@ public:
 
 private:
 	//todo this should be and make helper function to build the name in BP with a picker
-	UPROPERTY(EditAnywhere, Category="MVVM")
+	UPROPERTY(EditAnywhere, Category = "Viewmodel")
 	FName BindingName;
 };
 
diff --git a/Engine/Plugins/Runtime/ModelViewViewModel/Source/ModelViewViewModel/Public/Types/MVVMViewModelContext.h b/Engine/Plugins/Runtime/ModelViewViewModel/Source/ModelViewViewModel/Public/Types/MVVMViewModelContext.h
--- a/Engine/Plugins/Runtime/ModelViewViewModel/Source/ModelViewViewModel/Public/Types/MVVMViewModelContext.h
+++ b/Engine/Plugins/Runtime/ModelViewViewModel/Source/ModelViewViewModel/Public/Types/MVVMViewModelContext.h
@@ -16,10 +16,10 @@ struct MODELVIEWVIEWMODEL_API FMVVMViewModelContext
 	GENERATED_BODY()
 
 public:
-	UPROPERTY(BlueprintReadWrite, EditDefaultsOnly, Category="MVVM")
+	UPROPERTY(BlueprintReadWrite, EditDefaultsOnly, Category="Viewmodel")
 	TSubclassOf<UMVVMViewModelBase> ContextClass;
 
-	UPROPERTY(BlueprintReadWrite, EditDefaultsOnly, Category="MVVM")
+	UPROPERTY(BlueprintReadWrite, EditDefaultsOnly, Category="Viewmodel")
 	FName ContextName;
 
 public:
diff --git a/Engine/Plugins/Runtime/ModelViewViewModel/Source/ModelViewViewModelBlueprint/Public/MVVMDeveloperProjectSettings.h b/Engine/Plugins/Runtime/ModelViewViewModel/Source/ModelViewViewModelBlueprint/Public/MVVMDeveloperProjectSettings.h
--- a/Engine/Plugins/Runtime/ModelViewViewModel/Source/ModelViewViewModelBlueprint/Public/MVVMDeveloperProjectSettings.h
+++ b/Engine/Plugins/Runtime/ModelViewViewModel/Source/ModelViewViewModelBlueprint/Public/MVVMDeveloperProjectSettings.h
@@ -14,11 +14,11 @@ struct FMVVMDeveloperProjectWidgetSettings
 	GENERATED_BODY()
 
 	/** Properties or functions name that should not be use for binding (read or write). */
-	UPROPERTY(EditAnywhere, config, Category=MVVM)
+	UPROPERTY(EditAnywhere, config, Category = "Viewmodel")
 	TSet<FName> DisallowedFieldNames;
 	
 	/** Properties or functions name that are displayed in the advanced category. */
-	UPROPERTY(EditAnywhere, config, Category=MVVM)
+	UPROPERTY(EditAnywhere, config, Category = "Viewmodel")
 	TSet<FName> AdvancedFieldNames;
 };
 
@@ -40,7 +40,7 @@ public:
 
 private:
 	/** Permission list for filtering which properties are visible in UI. */
-	UPROPERTY(EditAnywhere, config, Category=MVVM)
+	UPROPERTY(EditAnywhere, config, Category = "Viewmodel")
 	TMap<FSoftClassPath, FMVVMDeveloperProjectWidgetSettings> FieldSelectorPermissions;
 };
 
diff --git a/Engine/Plugins/Runtime/ModelViewViewModel/Source/ModelViewViewModelEditor/Private/ViewModel/MVVMViewModelBlueprintFactory.h b/Engine/Plugins/Runtime/ModelViewViewModel/Source/ModelViewViewModelEditor/Private/ViewModel/MVVMViewModelBlueprintFactory.h
--- a/Engine/Plugins/Runtime/ModelViewViewModel/Source/ModelViewViewModelEditor/Private/ViewModel/MVVMViewModelBlueprintFactory.h
+++ b/Engine/Plugins/Runtime/ModelViewViewModel/Source/ModelViewViewModelEditor/Private/ViewModel/MVVMViewModelBlueprintFactory.h
@@ -23,6 +23,6 @@ public:
 
 private:
 	// The parent class of the created blueprint
-	UPROPERTY(EditAnywhere, Category="MVVM", meta=(AllowAbstract="true"))
+	UPROPERTY(EditAnywhere, Category = "Viewmodel", meta=(AllowAbstract="true"))
 	TSubclassOf<class UMVVMViewModelBase> ParentClass;
 };
diff --git a/Engine/Plugins/Runtime/ModelViewViewModel/Source/ModelViewViewModelEditor/Public/MVVMEditorSubsystem.h b/Engine/Plugins/Runtime/ModelViewViewModel/Source/ModelViewViewModelEditor/Public/MVVMEditorSubsystem.h
--- a/Engine/Plugins/Runtime/ModelViewViewModel/Source/ModelViewViewModelEditor/Public/MVVMEditorSubsystem.h
+++ b/Engine/Plugins/Runtime/ModelViewViewModel/Source/ModelViewViewModelEditor/Public/MVVMEditorSubsystem.h
@@ -29,31 +29,31 @@ class MODELVIEWVIEWMODELEDITOR_API UMVVMEditorSubsystem : public UEditorSubsyste
 	GENERATED_BODY()
 
 public:
-	UFUNCTION(BlueprintCallable, Category="MVVM")
+	UFUNCTION(BlueprintCallable, Category = "Viewmodel")
 	UMVVMBlueprintView* RequestView(UWidgetBlueprint* WidgetBlueprint) const;
 
-	UFUNCTION(BlueprintCallable, Category = "MVVM")
+	UFUNCTION(BlueprintCallable, Category = "Viewmodel")
 	UMVVMBlueprintView* GetView(const UWidgetBlueprint* WidgetBlueprint) const;
 
-	UFUNCTION(BlueprintCallable, Category = "MVVM")
+	UFUNCTION(BlueprintCallable, Category = "Viewmodel")
 	FName AddViewModel(UWidgetBlueprint* WidgetBlueprint, const UClass* ViewModel);
 	
-	UFUNCTION(BlueprintCallable, Category = "MVVM")
+	UFUNCTION(BlueprintCallable, Category = "Viewmodel")
 	void RemoveViewModel(UWidgetBlueprint* WidgetBlueprint, FName ViewModel);
 
-	UFUNCTION(BlueprintCallable, Category = "MVVM")
+	UFUNCTION(BlueprintCallable, Category = "Viewmodel")
 	bool VerifyViewModelRename(UWidgetBlueprint* WidgetBlueprint, FName ViewModel, FName NewViewModel, FText& OutError);
 
-	UFUNCTION(BlueprintCallable, Category = "MVVM")
+	UFUNCTION(BlueprintCallable, Category = "Viewmodel")
 	bool RenameViewModel(UWidgetBlueprint* WidgetBlueprint, FName ViewModel, FName NewViewModel, FText& OutError);
 
-	UFUNCTION(BlueprintCallable, Category = "MVVM")
+	UFUNCTION(BlueprintCallable, Category = "Viewmodel")
 	FMVVMBlueprintViewBinding& AddBinding(UWidgetBlueprint* WidgetBlueprint);
 
-	UFUNCTION(BlueprintCallable, Category = "MVVM")
+	UFUNCTION(BlueprintCallable, Category = "Viewmodel")
 	void RemoveBinding(UWidgetBlueprint* WidgetBlueprint, const FMVVMBlueprintViewBinding& Binding);
 
-	UFUNCTION(BlueprintCallable, Category = "MVVM")
+	UFUNCTION(BlueprintCallable, Category = "Viewmodel")
 	TArray<FMVVMAvailableBinding> GetChildViewModels(TSubclassOf<UObject> Class, TSubclassOf<UObject> Accessor);
 
 	void SetSourceToDestinationConversionFunction(UWidgetBlueprint* WidgetBlueprint, FMVVMBlueprintViewBinding& Binding, const UFunction* ConversionFunction);
@@ -66,22 +66,22 @@ public:
 	void SetEnabledForBinding(UWidgetBlueprint* WidgetBlueprint, FMVVMBlueprintViewBinding& Binding, bool bEnabled);
 	void SetCompileForBinding(UWidgetBlueprint* WidgetBlueprint, FMVVMBlueprintViewBinding& Binding, bool bCompile);
 
-	UFUNCTION(BlueprintCallable, Category = "MVVM")
+	UFUNCTION(BlueprintCallable, Category = "Viewmodel")
 	bool IsValidConversionFunction(const UFunction* Function, const FMVVMBlueprintPropertyPath& Source, const FMVVMBlueprintPropertyPath& Destination) const;
 
-	UFUNCTION(BlueprintCallable, Category = "MVVM")
+	UFUNCTION(BlueprintCallable, Category = "Viewmodel")
 	bool IsSimpleConversionFunction(const UFunction* Function) const;
 
-	UFUNCTION(BlueprintCallable, Category = "MVVM")
+	UFUNCTION(BlueprintCallable, Category = "Viewmodel")
 	UEdGraph* GetConversionFunctionGraph(const UWidgetBlueprint* WidgetBlueprint, const FMVVMBlueprintViewBinding& Binding, bool bSourceToDestination) const;
 
-	UFUNCTION(BlueprintCallable, Category = "MVVM")
+	UFUNCTION(BlueprintCallable, Category = "Viewmodel")
 	UFunction* GetConversionFunction(const UWidgetBlueprint* WidgetBlueprint, const FMVVMBlueprintViewBinding& Binding, bool bSourceToDestination) const;
 
-	UFUNCTION(BlueprintCallable, Category = "MVVM")
+	UFUNCTION(BlueprintCallable, Category = "Viewmodel")
 	UK2Node_CallFunction* GetConversionFunctionNode(const UWidgetBlueprint* WidgetBlueprint, const FMVVMBlueprintViewBinding& Binding, bool bSourceToDestination) const;
 
-	UFUNCTION(BlueprintCallable, Category = "MVVM")
+	UFUNCTION(BlueprintCallable, Category = "Viewmodel")
 	TArray<UFunction*> GetAvailableConversionFunctions(const UWidgetBlueprint* WidgetBlueprint, const FMVVMBlueprintPropertyPath& Source, const FMVVMBlueprintPropertyPath& Destination) const;
 
 	FMVVMBlueprintPropertyPath GetPathForConversionFunctionArgument(const UWidgetBlueprint* WidgetBlueprint, const FMVVMBlueprintViewBinding& Binding, FName ArgumentName, bool bSourceToDestination) const;
