MVVM: New binding mode SVG icons for the View Model plugin.
- One Time / One Way
- One Time / One Way to Source
- One Time / Two Way

diff --git a/Engine/Plugins/Runtime/ModelViewViewModel/Source/ModelViewViewModelEditor/Private/Styling/MVVMEditorStyle.cpp b/Engine/Plugins/Runtime/ModelViewViewModel/Source/ModelViewViewModelEditor/Private/Styling/MVVMEditorStyle.cpp
--- a/Engine/Plugins/Runtime/ModelViewViewModel/Source/ModelViewViewModelEditor/Private/Styling/MVVMEditorStyle.cpp
+++ b/Engine/Plugins/Runtime/ModelViewViewModel/Source/ModelViewViewModelEditor/Private/Styling/MVVMEditorStyle.cpp
@@ -83,6 +83,9 @@ FMVVMEditorStyle::FMVVMEditorStyle()
 	Set("BindingMode.OneWayToSource", new IMAGE_BRUSH_SVG("Slate/BindingMode_OneWayToSource", Icon16x16));
 	Set("BindingMode.TwoWay", new IMAGE_BRUSH_SVG("Slate/BindingMode_TwoWay", Icon16x16));
 	Set("BindingMode.OneTime", new IMAGE_BRUSH_SVG("Slate/BindingMode_OneTime", Icon16x16));
+	Set("BindingMode.OneTimeOneWay", new IMAGE_BRUSH_SVG("Slate/BindingMode_OneTimeOneWayToSource", Icon16x16));
+	Set("BindingMode.OneTimeOneWayToSource", new IMAGE_BRUSH_SVG("Slate/BindingMode_OneTimeOneWay", Icon16x16));
+	Set("BindingMode.OneTimeTwoWay", new IMAGE_BRUSH_SVG("Slate/BindingMode_OneTimeTwoWay", Icon16x16));
 
 	Set("ConversionFunction.DestToSource", new IMAGE_BRUSH_SVG("Slate/ConversionFunction_DestToSource", Icon16x16));
 	Set("ConversionFunction.SourceToDest", new IMAGE_BRUSH_SVG("Slate/ConversionFunction_SourceToDest", Icon16x16));
diff --git a/Engine/Plugins/Runtime/ModelViewViewModel/Source/ModelViewViewModelEditor/Private/Widgets/SMVVMViewBindingListView.cpp b/Engine/Plugins/Runtime/ModelViewViewModel/Source/ModelViewViewModelEditor/Private/Widgets/SMVVMViewBindingListView.cpp
--- a/Engine/Plugins/Runtime/ModelViewViewModel/Source/ModelViewViewModelEditor/Private/Widgets/SMVVMViewBindingListView.cpp
+++ b/Engine/Plugins/Runtime/ModelViewViewModel/Source/ModelViewViewModelEditor/Private/Widgets/SMVVMViewBindingListView.cpp
@@ -787,7 +787,7 @@ private:
 		switch (BindingMode)
 		{
 		case EMVVMBindingMode::OneTimeToDestination:
-			return FMVVMEditorStyle::Get().GetBrush("BindingMode.OneTime");
+			return FMVVMEditorStyle::Get().GetBrush("BindingMode.OneTimeOneWay");
 		case EMVVMBindingMode::OneWayToDestination:
 			return FMVVMEditorStyle::Get().GetBrush("BindingMode.OneWayToSource");
 		case EMVVMBindingMode::OneWayToSource:
diff --git a/Engine/Plugins/Runtime/ModelViewViewModel/Source/ModelViewViewModelEditor/Private/Widgets/SMVVMViewModelBindingListWidget.cpp b/Engine/Plugins/Runtime/ModelViewViewModel/Source/ModelViewViewModelEditor/Private/Widgets/SMVVMViewModelBindingListWidget.cpp
--- a/Engine/Plugins/Runtime/ModelViewViewModel/Source/ModelViewViewModelEditor/Private/Widgets/SMVVMViewModelBindingListWidget.cpp
+++ b/Engine/Plugins/Runtime/ModelViewViewModel/Source/ModelViewViewModelEditor/Private/Widgets/SMVVMViewModelBindingListWidget.cpp
@@ -251,35 +251,41 @@ TSharedRef<SWidget> ConstructFieldPreSlot(const UWidgetBlueprint* WidgetBlueprin
 	if (Binding.IsValid())
 	{
 		const FSlateBrush* Brush = nullptr;
-		const FLinearColor ImageColor = FLinearColor(1.0f, 1.0f, 1.0, 0.5f);
-		if (Binding.HasNotify() && Binding.IsReadable() && Binding.IsWritable())
+		if (Binding.HasNotify())
 		{
-			Brush = FMVVMEditorStyle::Get().GetBrush("BindingMode.TwoWay");
-		}
-		else if (Binding.HasNotify() && Binding.IsReadable())
-		{
-			Brush = FMVVMEditorStyle::Get().GetBrush("BindingMode.OneWayToSource");
-		}
-		else if (Binding.IsWritable() && Binding.IsReadable())
-		{
-			ImageWidget = SNew(SLayeredImage, FMVVMEditorStyle::Get().GetBrush("BindingMode.OneWay"), ImageColor)
-				.Image(FMVVMEditorStyle::Get().GetBrush("BindingMode.OneTime"))
-				.ColorAndOpacity(ImageColor);
-		}
-		else if (Binding.IsWritable())
-		{
-			Brush = FMVVMEditorStyle::Get().GetBrush("BindingMode.OneWay");
+			if (Binding.IsReadable() && Binding.IsWritable())
+			{
+				Brush = FMVVMEditorStyle::Get().GetBrush("BindingMode.TwoWay");
+			}
+			else if (Binding.IsReadable())
+			{
+				Brush = FMVVMEditorStyle::Get().GetBrush("BindingMode.OneWayToSource");
+			}
+			else if (Binding.IsWritable())
+			{
+				Brush = FMVVMEditorStyle::Get().GetBrush("BindingMode.OneWay");
+			}
 		}
-		else if (Binding.IsReadable())
+		else
 		{
-			Brush = FMVVMEditorStyle::Get().GetBrush("BindingMode.OneTime");
+			if (Binding.IsReadable() && Binding.IsWritable())
+			{
+				Brush = FMVVMEditorStyle::Get().GetBrush("BindingMode.OneTimeTwoWay");
+			}
+			else if (Binding.IsReadable())
+			{
+				Brush = FMVVMEditorStyle::Get().GetBrush("BindingMode.OneTimeOneWay");
+			}
+			else if (Binding.IsWritable())
+			{
+				Brush = FMVVMEditorStyle::Get().GetBrush("BindingMode.OneTimeOneWayToSource");
+			}
 		}
-
+		
 		if (Brush)
 		{
 			ImageWidget = SNew(SImage)
-				.Image(Brush)
-				.ColorAndOpacity(ImageColor);
+				.Image(Brush);
 		}
 	}
 
diff --git a/Engine/Plugins/Runtime/ModelViewViewModel/Source/ModelViewViewModelEditor/Private/Widgets/SMVVMViewModelContextListWidget.cpp b/Engine/Plugins/Runtime/ModelViewViewModel/Source/ModelViewViewModelEditor/Private/Widgets/SMVVMViewModelContextListWidget.cpp
--- a/Engine/Plugins/Runtime/ModelViewViewModel/Source/ModelViewViewModelEditor/Private/Widgets/SMVVMViewModelContextListWidget.cpp
+++ b/Engine/Plugins/Runtime/ModelViewViewModel/Source/ModelViewViewModelEditor/Private/Widgets/SMVVMViewModelContextListWidget.cpp
@@ -643,7 +643,7 @@ namespace UE::MVVM::Private
 		switch (BindingMode)
 		{
 		case EMVVMBindingMode::OneTimeToDestination:
-			return FMVVMEditorStyle::Get().GetBrush("BindingMode.OneTime");
+			return FMVVMEditorStyle::Get().GetBrush("BindingMode.OneTimeOneWay");
 		case EMVVMBindingMode::OneWayToDestination:
 			return FMVVMEditorStyle::Get().GetBrush("BindingMode.OneWay");
 		case EMVVMBindingMode::OneWayToSource:
