MVVM: Fix bad access to raw pointer.

diff --git a/Engine/Plugins/Runtime/ModelViewViewModel/Source/ModelViewViewModelEditor/Private/Widgets/SMVVMViewBindingPanel.cpp b/Engine/Plugins/Runtime/ModelViewViewModel/Source/ModelViewViewModelEditor/Private/Widgets/SMVVMViewBindingPanel.cpp
--- a/Engine/Plugins/Runtime/ModelViewViewModel/Source/ModelViewViewModelEditor/Private/Widgets/SMVVMViewBindingPanel.cpp
+++ b/Engine/Plugins/Runtime/ModelViewViewModel/Source/ModelViewViewModelEditor/Private/Widgets/SMVVMViewBindingPanel.cpp
@@ -52,7 +52,7 @@ void SBindingsPanel::Construct(const FArguments& InArgs, TSharedPtr<FWidgetBluep
 	}
 	else
 	{
-		WidgetBlueprint->OnExtensionAdded.AddRaw(this, &SBindingsPanel::HandleExtensionAdded);
+		WidgetBlueprint->OnExtensionAdded.AddSP(this, &SBindingsPanel::HandleExtensionAdded);
 	}
 
 	{
