MVVM: Always include the viewmodel macro when including the base class.
#rnx

diff --git a/Engine/Plugins/Runtime/ModelViewViewModel/Source/ModelViewViewModel/Public/MVVMViewModelBase.h b/Engine/Plugins/Runtime/ModelViewViewModel/Source/ModelViewViewModel/Public/MVVMViewModelBase.h
--- a/Engine/Plugins/Runtime/ModelViewViewModel/Source/ModelViewViewModel/Public/MVVMViewModelBase.h
+++ b/Engine/Plugins/Runtime/ModelViewViewModel/Source/ModelViewViewModel/Public/MVVMViewModelBase.h
@@ -4,6 +4,7 @@
 
 
 #include "FieldNotification/IClassDescriptor.h"
+#include "FieldNotification/FieldNotificationDeclaration.h" // IWYU pragma: keep
 #include "ViewModel/MVVMFieldNotificationDelegates.h"
 
 #include "MVVMViewModelBase.generated.h"
