Fixed static analysis warning about null Pin in SMVVMFunctionParameter.cpp.

diff --git a/Engine/Plugins/Runtime/ModelViewViewModel/Source/ModelViewViewModelEditor/Private/Widgets/SMVVMFunctionParameter.cpp b/Engine/Plugins/Runtime/ModelViewViewModel/Source/ModelViewViewModelEditor/Private/Widgets/SMVVMFunctionParameter.cpp
--- a/Engine/Plugins/Runtime/ModelViewViewModel/Source/ModelViewViewModelEditor/Private/Widgets/SMVVMFunctionParameter.cpp
+++ b/Engine/Plugins/Runtime/ModelViewViewModel/Source/ModelViewViewModelEditor/Private/Widgets/SMVVMFunctionParameter.cpp
@@ -47,6 +47,7 @@ void SFunctionParameter::Construct(const FArguments& InArgs)
 
 	TSharedRef<SWidget> ValueWidget = SNullWidget::NullWidget;
 
+	bool bIsBooleanPin = false;
 	UEdGraphPin* Pin = EditorSubsystem->GetConversionFunctionArgumentPin(InArgs._WidgetBlueprint, *Binding, ParameterName, bSourceToDestination);
 	if (Pin != nullptr)
 	{
@@ -56,10 +57,10 @@ void SFunctionParameter::Construct(const FArguments& InArgs)
 			GraphPin = PinWidget;
 			ValueWidget = PinWidget->GetDefaultValueWidget();
 		}
+
+		bIsBooleanPin = Pin->PinType.PinCategory == UEdGraphSchema_K2::PC_Boolean;
 	}
 
-	const bool bIsBooleanPin = Pin->PinType.PinCategory == UEdGraphSchema_K2::PC_Boolean;
-	
 	if (ValueWidget == SNullWidget::NullWidget)
 	{
 		ValueWidget = SNew(STextBlock)
