MVVM: Use the class display name in the select viewmodel menu.

diff --git a/Engine/Plugins/Runtime/ModelViewViewModel/Source/ModelViewViewModelEditor/Private/ViewModel/MVVMViewModelBlueprintFactory.cpp b/Engine/Plugins/Runtime/ModelViewViewModel/Source/ModelViewViewModelEditor/Private/ViewModel/MVVMViewModelBlueprintFactory.cpp
--- a/Engine/Plugins/Runtime/ModelViewViewModel/Source/ModelViewViewModelEditor/Private/ViewModel/MVVMViewModelBlueprintFactory.cpp
+++ b/Engine/Plugins/Runtime/ModelViewViewModel/Source/ModelViewViewModelEditor/Private/ViewModel/MVVMViewModelBlueprintFactory.cpp
@@ -62,6 +62,7 @@ bool UMVVMViewModelBlueprintFactory::ConfigureProperties()
 	FClassViewerInitializationOptions Options;
 	Options.DisplayMode = EClassViewerDisplayMode::Type::TreeView;
 	Options.Mode = EClassViewerMode::ClassPicker;
+	Options.NameTypeToDisplay = EClassViewerNameTypeToDisplay::DisplayName;
 	Options.bShowNoneOption = false;
 	Options.bExpandAllNodes = true;
 
diff --git a/Engine/Plugins/Runtime/ModelViewViewModel/Source/ModelViewViewModelEditor/Private/Widgets/SMVVMSelectViewModel.cpp b/Engine/Plugins/Runtime/ModelViewViewModel/Source/ModelViewViewModelEditor/Private/Widgets/SMVVMSelectViewModel.cpp
--- a/Engine/Plugins/Runtime/ModelViewViewModel/Source/ModelViewViewModelEditor/Private/Widgets/SMVVMSelectViewModel.cpp
+++ b/Engine/Plugins/Runtime/ModelViewViewModel/Source/ModelViewViewModelEditor/Private/Widgets/SMVVMSelectViewModel.cpp
@@ -54,6 +54,7 @@ void SMVVMSelectViewModel::Construct(const FArguments& InArgs, const UWidgetBlue
 	FClassViewerInitializationOptions ClassViewerOptions;
 	ClassViewerOptions.DisplayMode = EClassViewerDisplayMode::TreeView;
 	ClassViewerOptions.Mode = EClassViewerMode::ClassPicker;
+	ClassViewerOptions.NameTypeToDisplay = EClassViewerNameTypeToDisplay::DisplayName;
 	ClassViewerOptions.ClassFilters.Add(MakeShared<FViewModelClassFilter>());
 
 	ClassViewer = FModuleManager::LoadModuleChecked<FClassViewerModule>("ClassViewer")
