MVVM: Add a allowed ExecutionMode list.
Move GlobalCollection to GameInstance subsystem.
Rename category from MVVM to Viewmodel.

diff --git a/Engine/Plugins/Runtime/ModelViewViewModel/Source/ModelViewViewModel/Private/MVVMGameSubsystem.cpp b/Engine/Plugins/Runtime/ModelViewViewModel/Source/ModelViewViewModel/Private/MVVMGameSubsystem.cpp
new file mode 100644
--- /dev/null
+++ b/Engine/Plugins/Runtime/ModelViewViewModel/Source/ModelViewViewModel/Private/MVVMGameSubsystem.cpp
@@ -0,0 +1,21 @@
+// Copyright Epic Games, Inc. All Rights Reserved.
+
+#include "MVVMGameSubsystem.h"
+
+#include "Types/MVVMViewModelContextInstance.h"
+
+#include UE_INLINE_GENERATED_CPP_BY_NAME(MVVMGameSubsystem)
+
+
+void UMVVMGameSubsystem::Initialize(FSubsystemCollectionBase& Collection)
+{
+	Super::Initialize(Collection);
+	GlobalViewModelCollection = NewObject<UMVVMViewModelCollectionObject>(this);
+}
+
+
+void UMVVMGameSubsystem::Deinitialize()
+{
+	GlobalViewModelCollection = nullptr;
+	Super::Deinitialize();
+}
diff --git a/Engine/Plugins/Runtime/ModelViewViewModel/Source/ModelViewViewModel/Private/MVVMSubsystem.cpp b/Engine/Plugins/Runtime/ModelViewViewModel/Source/ModelViewViewModel/Private/MVVMSubsystem.cpp
--- a/Engine/Plugins/Runtime/ModelViewViewModel/Source/ModelViewViewModel/Private/MVVMSubsystem.cpp
+++ b/Engine/Plugins/Runtime/ModelViewViewModel/Source/ModelViewViewModel/Private/MVVMSubsystem.cpp
@@ -4,6 +4,10 @@
 
 #include "Bindings/MVVMBindingHelper.h"
 #include "Blueprint/WidgetTree.h"
+#include "Engine/Engine.h"
+#include "Engine/GameInstance.h"
+#include "Engine/World.h"
+#include "MVVMGameSubsystem.h"
 #include "Templates/ValueOrError.h"
 #include "View/MVVMView.h"
 #include "Types/MVVMAvailableBinding.h"
@@ -15,13 +19,11 @@
 void UMVVMSubsystem::Initialize(FSubsystemCollectionBase& Collection)
 {
 	Super::Initialize(Collection);
-	GlobalViewModelCollection = NewObject<UMVVMViewModelCollectionObject>(this);
 }
 
 
 void UMVVMSubsystem::Deinitialize()
 {
-	GlobalViewModelCollection = nullptr;
 	Super::Deinitialize();
 }
 
@@ -292,6 +294,22 @@ FMVVMAvailableBinding UMVVMSubsystem::GetAvailableBindingForField(UE::MVVM::FMVV
 }
 
 
+UMVVMViewModelCollectionObject* UMVVMSubsystem::GetGlobalViewModelCollection(const UObject* WorldContextObject) const
+{
+	if (WorldContextObject)
+	{
+		if (const UWorld* World = GEngine->GetWorldFromContextObject(WorldContextObject, EGetWorldErrorMode::ReturnNull))
+		{
+			if (const UGameInstance* GameInstance = World->GetGameInstance())
+			{
+				return GameInstance->GetSubsystem<UMVVMGameSubsystem>()->GetGlobalViewModelCollection();
+			}
+		}
+	}
+	return nullptr;
+}
+
+
 TValueOrError<bool, FText> UMVVMSubsystem::IsBindingValid(FBindingArgs Args) const
 {
 	if (UE::MVVM::IsForwardBinding(Args.Mode))
diff --git a/Engine/Plugins/Runtime/ModelViewViewModel/Source/ModelViewViewModel/Private/View/MVVMViewClass.cpp b/Engine/Plugins/Runtime/ModelViewViewModel/Source/ModelViewViewModel/Private/View/MVVMViewClass.cpp
--- a/Engine/Plugins/Runtime/ModelViewViewModel/Source/ModelViewViewModel/Private/View/MVVMViewClass.cpp
+++ b/Engine/Plugins/Runtime/ModelViewViewModel/Source/ModelViewViewModel/Private/View/MVVMViewClass.cpp
@@ -7,6 +7,9 @@
 #include "Bindings/MVVMFieldPathHelper.h"
 #include "Blueprint/UserWidget.h"
 #include "Engine/Engine.h"
+#include "Engine/GameInstance.h"
+#include "Engine/World.h"
+#include "MVVMGameSubsystem.h"
 #include "MVVMMessageLog.h"
 #include "MVVMSubsystem.h"
 
@@ -112,7 +115,20 @@ UObject* FMVVMViewClass_SourceCreator::CreateInstance(const UMVVMViewClass* InVi
 		}
 		else if (GlobalViewModelInstance.IsValid())
 		{
-			UMVVMViewModelBase* FoundViewModelInstance = GEngine->GetEngineSubsystem<UMVVMSubsystem>()->GetGlobalViewModelCollection()->FindViewModelInstance(GlobalViewModelInstance);
+			UMVVMViewModelCollectionObject* Collection = nullptr;
+			UMVVMViewModelBase* FoundViewModelInstance = nullptr;
+			if (const UWorld* World = InUserWidget->GetWorld())
+			{
+				if (const UGameInstance* GameInstance = World->GetGameInstance())
+				{
+					Collection = GameInstance->GetSubsystem<UMVVMGameSubsystem>()->GetGlobalViewModelCollection();
+					if (Collection)
+					{
+						FoundViewModelInstance = Collection->FindViewModelInstance(GlobalViewModelInstance);
+					}
+				}
+			}
+
 			if (FoundViewModelInstance != nullptr)
 			{
 				ensureMsgf(FoundViewModelInstance->IsA(GlobalViewModelInstance.ContextClass), TEXT("The Global View Model Instance is not of the expected type."));
@@ -122,7 +138,14 @@ UObject* FMVVMViewClass_SourceCreator::CreateInstance(const UMVVMViewClass* InVi
 			else if (!bOptional)
 			{
 				UE::MVVM::FMessageLog Log(InUserWidget);
-				Log.Error(FText::Format(LOCTEXT("CreateInstanceFailedGlobal", "The viewmodel '{0}' was not found in the global view model collection."), FText::FromName(GlobalViewModelInstance.ContextName)));
+				if (Collection)
+				{
+					Log.Error(FText::Format(LOCTEXT("CreateInstanceFailedGlobal", "The source '{0}' was not found in the global view model collection."), FText::FromName(GlobalViewModelInstance.ContextName)));
+				}
+				else
+				{
+					Log.Error(FText::Format(LOCTEXT("CreateInstanceFailedGlobalInstance", "The source '{0}' will be invalid because the global view model collection could not be found."), FText::FromName(GlobalViewModelInstance.ContextName)));
+				}
 			}
 		}
 		else if (FieldPath.IsValid())
diff --git a/Engine/Plugins/Runtime/ModelViewViewModel/Source/ModelViewViewModel/Private/View/MVVMViewWorldSubsystem.h b/Engine/Plugins/Runtime/ModelViewViewModel/Source/ModelViewViewModel/Private/View/MVVMViewWorldSubsystem.h
--- a/Engine/Plugins/Runtime/ModelViewViewModel/Source/ModelViewViewModel/Private/View/MVVMViewWorldSubsystem.h
+++ b/Engine/Plugins/Runtime/ModelViewViewModel/Source/ModelViewViewModel/Private/View/MVVMViewWorldSubsystem.h
@@ -36,7 +36,7 @@ private:
 };
 
 /** */
-UCLASS()
+UCLASS(NotBlueprintable, DisplayName = "Viewmodel World Subsytem")
 class UMVVMViewWorldSubsystem : public UTickableWorldSubsystem
 {
 	GENERATED_BODY()
diff --git a/Engine/Plugins/Runtime/ModelViewViewModel/Source/ModelViewViewModel/Public/MVVMGameSubsystem.h b/Engine/Plugins/Runtime/ModelViewViewModel/Source/ModelViewViewModel/Public/MVVMGameSubsystem.h
new file mode 100644
--- /dev/null
+++ b/Engine/Plugins/Runtime/ModelViewViewModel/Source/ModelViewViewModel/Public/MVVMGameSubsystem.h
@@ -0,0 +1,32 @@
+// Copyright Epic Games, Inc. All Rights Reserved.
+
+#pragma once
+
+#include "Subsystems/GameInstanceSubsystem.h"
+#include "Types/MVVMViewModelCollection.h"
+
+#include "UObject/Package.h"
+#include "MVVMGameSubsystem.generated.h"
+
+
+/** */
+UCLASS(DisplayName="Viewmodel Game Subsytem")
+class MODELVIEWVIEWMODEL_API UMVVMGameSubsystem : public UGameInstanceSubsystem
+{
+	GENERATED_BODY()
+
+public:
+	//~ Begin UEngineSubsystem interface
+	virtual void Initialize(FSubsystemCollectionBase& Collection) override;
+	virtual void Deinitialize() override;
+	//~ End UEngineSubsystem interface
+
+	UFUNCTION(BlueprintCallable, Category = "Viewmodel")
+	UMVVMViewModelCollectionObject* GetGlobalViewModelCollection() const
+	{
+		return GlobalViewModelCollection;
+	}
+private:
+	UPROPERTY(Transient)
+	TObjectPtr<UMVVMViewModelCollectionObject> GlobalViewModelCollection;
+};
diff --git a/Engine/Plugins/Runtime/ModelViewViewModel/Source/ModelViewViewModel/Public/MVVMSubsystem.h b/Engine/Plugins/Runtime/ModelViewViewModel/Source/ModelViewViewModel/Public/MVVMSubsystem.h
--- a/Engine/Plugins/Runtime/ModelViewViewModel/Source/ModelViewViewModel/Public/MVVMSubsystem.h
+++ b/Engine/Plugins/Runtime/ModelViewViewModel/Source/ModelViewViewModel/Public/MVVMSubsystem.h
@@ -21,7 +21,7 @@ class UWidget;
 class UWidgetTree;
 
 /** */
-UCLASS()
+UCLASS(DisplayName="Viewmodel Engine Subsytem")
 class MODELVIEWVIEWMODEL_API UMVVMSubsystem : public UEngineSubsystem
 {
 	GENERATED_BODY()
@@ -61,11 +61,9 @@ public:
 	/** @return The AvailableBinding from a field. */
 	static FMVVMAvailableBinding GetAvailableBindingForField(UE::MVVM::FMVVMConstFieldVariant Variant, const UClass* Accessor);
 
-	UFUNCTION(BlueprintCallable, Category = "Viewmodel")
-	UMVVMViewModelCollectionObject* GetGlobalViewModelCollection() const
-	{
-		return GlobalViewModelCollection;
-	}
+	UE_DEPRECATED(5.3, "GetGlobalViewModelCollection has been deprecated, please use the game instance subsystem.")
+	UFUNCTION(BlueprintCallable, Category = "Viewmodel", meta=(WorldContext="WorldContextObject", DeprecatedFunction, DeprecatedMessage = "This version of GetGlobalViewModelCollection has been deprecated, please use GetGlobalViewModelCollection from the Viewmodel Game subsystem."))
+	UMVVMViewModelCollectionObject* GetGlobalViewModelCollection(const UObject* WorldContextObject) const;
 
 	struct FConstDirectionalBindingArgs
 	{
@@ -100,10 +98,6 @@ public:
 	UE_NODISCARD TValueOrError<bool, FText> IsBindingValid(FConstDirectionalBindingArgs Args) const;
 	UE_NODISCARD TValueOrError<bool, FText> IsBindingValid(FDirectionalBindingArgs Args) const;
 	UE_NODISCARD TValueOrError<bool, FText> IsBindingValid(FBindingArgs Args) const;
-
-private:
-	UPROPERTY(Transient)
-	TObjectPtr<UMVVMViewModelCollectionObject> GlobalViewModelCollection;
 };
 
 #if UE_ENABLE_INCLUDE_ORDER_DEPRECATED_IN_5_2
diff --git a/Engine/Plugins/Runtime/ModelViewViewModel/Source/ModelViewViewModel/Public/Types/MVVMViewModelCollection.h b/Engine/Plugins/Runtime/ModelViewViewModel/Source/ModelViewViewModel/Public/Types/MVVMViewModelCollection.h
--- a/Engine/Plugins/Runtime/ModelViewViewModel/Source/ModelViewViewModel/Public/Types/MVVMViewModelCollection.h
+++ b/Engine/Plugins/Runtime/ModelViewViewModel/Source/ModelViewViewModel/Public/Types/MVVMViewModelCollection.h
@@ -44,7 +44,7 @@ class UMVVMViewModelCollectionObject : public UObject
 	GENERATED_BODY()
 
 public:
-	UFUNCTION(BlueprintCallable, Category = "MVVM")
+	UFUNCTION(BlueprintCallable, Category = "Viewmodel")
 	UMVVMViewModelBase* FindViewModelInstance(FMVVMViewModelContext Context) const
 	{
 		return ViewModelCollection.FindViewModelInstance(Context);
@@ -54,25 +54,25 @@ public:
 	 * Finds a View Model of the given type.
 	 * If the collection contains multiple instances of the same type then this will return the first one found.
 	 */
-	UFUNCTION(BlueprintCallable, Category = "MVVM")
+	UFUNCTION(BlueprintCallable, Category = "Viewmodel")
 	UMVVMViewModelBase* FindFirstViewModelInstanceOfType(const TSubclassOf<UMVVMViewModelBase>& ViewModelClass) const
 	{
 		return ViewModelCollection.FindFirstViewModelInstanceOfType(ViewModelClass);
 	}
 	
-	UFUNCTION(BlueprintCallable, Category = "MVVM")
+	UFUNCTION(BlueprintCallable, Category = "Viewmodel")
 	bool AddViewModelInstance(FMVVMViewModelContext Context, UMVVMViewModelBase* ViewModel)
 	{
 		return ViewModelCollection.AddInstance(Context, ViewModel);
 	}
 
-	UFUNCTION(BlueprintCallable, Category = "MVVM")
+	UFUNCTION(BlueprintCallable, Category = "Viewmodel")
 	bool RemoveViewModel(FMVVMViewModelContext Context)
 	{
 		return ViewModelCollection.RemoveInstance(Context);
 	}
 
-	UFUNCTION(BlueprintCallable, Category = "MVVM")
+	UFUNCTION(BlueprintCallable, Category = "Viewmodel")
 	int32 RemoveAllViewModelInstance(UMVVMViewModelBase* ViewModel)
 	{
 		return ViewModelCollection.RemoveAllInstances(ViewModel);
diff --git a/Engine/Plugins/Runtime/ModelViewViewModel/Source/ModelViewViewModel/Public/View/MVVMView.h b/Engine/Plugins/Runtime/ModelViewViewModel/Source/ModelViewViewModel/Public/View/MVVMView.h
--- a/Engine/Plugins/Runtime/ModelViewViewModel/Source/ModelViewViewModel/Public/View/MVVMView.h
+++ b/Engine/Plugins/Runtime/ModelViewViewModel/Source/ModelViewViewModel/Public/View/MVVMView.h
@@ -48,13 +48,13 @@ public:
 	void ExecuteEveryTickBindings() const;
 
 // todo a way to identify a binding from outside. maybe a unique name in the editor?
-	//UFUNCTION(BlueprintCallable, Category = "MVVM")
+	//UFUNCTION(BlueprintCallable, Category = "Viewmodel")
 	//void SetLibraryBindingEnabled(FGuid ViewModelId, FMVVMBindingName BindingName, bool bEnable);
 
-	//UFUNCTION(BlueprintCallable, Category = "MVVM")
+	//UFUNCTION(BlueprintCallable, Category = "Viewmodel")
 	//bool IsLibraryBindingEnabled(FGuid ViewModelId, FMVVMBindingName BindingName) const;
 
-	UFUNCTION(BlueprintCallable, Category = "MVVM")
+	UFUNCTION(BlueprintCallable, Category = "Viewmodel")
 	bool SetViewModel(FName ViewModelName, TScriptInterface<INotifyFieldValueChanged> ViewModel);
 
 private:
diff --git a/Engine/Plugins/Runtime/ModelViewViewModel/Source/ModelViewViewModelBlueprint/Private/MVVMDeveloperProjectSettings.cpp b/Engine/Plugins/Runtime/ModelViewViewModel/Source/ModelViewViewModelBlueprint/Private/MVVMDeveloperProjectSettings.cpp
--- a/Engine/Plugins/Runtime/ModelViewViewModel/Source/ModelViewViewModelBlueprint/Private/MVVMDeveloperProjectSettings.cpp
+++ b/Engine/Plugins/Runtime/ModelViewViewModel/Source/ModelViewViewModelBlueprint/Private/MVVMDeveloperProjectSettings.cpp
@@ -4,6 +4,13 @@
 
 #define LOCTEXT_NAMESPACE "MVVMDeveloperProjectSettings"
 
+UMVVMDeveloperProjectSettings::UMVVMDeveloperProjectSettings()
+{
+	AllowedExecutionMode.Add(EMVVMExecutionMode::Immediate);
+	AllowedExecutionMode.Add(EMVVMExecutionMode::Delayed);
+	AllowedExecutionMode.Add(EMVVMExecutionMode::Tick);
+}
+
 FName UMVVMDeveloperProjectSettings::GetCategoryName() const
 {
 	return TEXT("Plugins");
diff --git a/Engine/Plugins/Runtime/ModelViewViewModel/Source/ModelViewViewModelBlueprint/Private/MVVMViewBlueprintCompiler.cpp b/Engine/Plugins/Runtime/ModelViewViewModel/Source/ModelViewViewModelBlueprint/Private/MVVMViewBlueprintCompiler.cpp
--- a/Engine/Plugins/Runtime/ModelViewViewModel/Source/ModelViewViewModelBlueprint/Private/MVVMViewBlueprintCompiler.cpp
+++ b/Engine/Plugins/Runtime/ModelViewViewModel/Source/ModelViewViewModelBlueprint/Private/MVVMViewBlueprintCompiler.cpp
@@ -6,6 +6,7 @@
 #include "Bindings/MVVMFieldPathHelper.h"
 #include "Components/Widget.h"
 #include "MVVMBlueprintView.h"
+#include "MVVMDeveloperProjectSettings.h"
 #include "HAL/IConsoleManager.h"
 #include "MVVMFunctionGraphHelper.h"
 #include "Templates/ValueOrError.h"
@@ -994,6 +995,14 @@ bool FMVVMViewBlueprintCompiler::PreCompileBindings(UWidgetBlueprintGeneratedCla
 			return MakeValue(FCompiledBindingLibraryCompiler::FFieldIdHandle());
 		};
 
+		if (Binding.bOverrideExecutionMode)
+		{
+			if (!GetDefault<UMVVMDeveloperProjectSettings>()->IsExecutionModeAllowed(Binding.OverrideExecutionMode))
+			{
+				AddErrorForBinding(Binding, LOCTEXT("NotAllowedExecutionMode", "The binding has a restricted execution mode."));
+			}
+		}
+
 		TValueOrError<FCompiledBindingLibraryCompiler::FFieldIdHandle, FText> AddFieldResult = AddFieldId(BindingSourceContext.SourceClass, true, Binding.BindingType, BindingSourceContext.FieldId.GetFieldName());
 		if (AddFieldResult.HasError())
 		{
diff --git a/Engine/Plugins/Runtime/ModelViewViewModel/Source/ModelViewViewModelBlueprint/Public/MVVMBlueprintView.h b/Engine/Plugins/Runtime/ModelViewViewModel/Source/ModelViewViewModelBlueprint/Public/MVVMBlueprintView.h
--- a/Engine/Plugins/Runtime/ModelViewViewModel/Source/ModelViewViewModelBlueprint/Public/MVVMBlueprintView.h
+++ b/Engine/Plugins/Runtime/ModelViewViewModel/Source/ModelViewViewModelBlueprint/Public/MVVMBlueprintView.h
@@ -88,10 +88,10 @@ public:
 	TArray<TObjectPtr<UEdGraph>> TemporaryGraph;
 
 private:
-	UPROPERTY(EditAnywhere, Category = "MVVM")
+	UPROPERTY(EditAnywhere, Category = "Viewmodel")
 	TArray<FMVVMBlueprintViewBinding> Bindings;
 
-	UPROPERTY(EditAnywhere, Category = "MVVM")
+	UPROPERTY(EditAnywhere, Category = "Viewmodel")
 	TArray<FMVVMBlueprintViewModelContext> AvailableViewModels;
 };
 
diff --git a/Engine/Plugins/Runtime/ModelViewViewModel/Source/ModelViewViewModelBlueprint/Public/MVVMBlueprintViewBinding.h b/Engine/Plugins/Runtime/ModelViewViewModel/Source/ModelViewViewModelBlueprint/Public/MVVMBlueprintViewBinding.h
--- a/Engine/Plugins/Runtime/ModelViewViewModel/Source/ModelViewViewModelBlueprint/Public/MVVMBlueprintViewBinding.h
+++ b/Engine/Plugins/Runtime/ModelViewViewModel/Source/ModelViewViewModelBlueprint/Public/MVVMBlueprintViewBinding.h
@@ -20,14 +20,14 @@ struct MODELVIEWVIEWMODELBLUEPRINT_API FMVVMBlueprintViewConversionPath
 	GENERATED_BODY()
 
 	/** The Conversion function when converting the value from the destination to the source. */
-	UPROPERTY(EditAnywhere, Category = "MVVM", AdvancedDisplay)
+	UPROPERTY(EditAnywhere, Category = "Viewmodel", AdvancedDisplay)
 	FMemberReference DestinationToSourceFunction;
 	
 	UPROPERTY()
 	FName DestinationToSourceWrapper;
 
 	/** The Conversion function when converting the value from the source to the destination. */
-	UPROPERTY(EditAnywhere, Category = "MVVM", AdvancedDisplay)
+	UPROPERTY(EditAnywhere, Category = "Viewmodel", AdvancedDisplay)
 	FMemberReference SourceToDestinationFunction;
 
 	UPROPERTY()
@@ -42,40 +42,40 @@ struct MODELVIEWVIEWMODELBLUEPRINT_API FMVVMBlueprintViewBinding
 {
 	GENERATED_BODY()
 
-	UPROPERTY(EditAnywhere, Category = "MVVM")
+	UPROPERTY(EditAnywhere, Category = "Viewmodel")
 	FMVVMBlueprintPropertyPath ViewModelPath;
 
-	UPROPERTY(EditAnywhere, Category = "MVVM")
+	UPROPERTY(EditAnywhere, Category = "Viewmodel")
 	FMVVMBlueprintPropertyPath WidgetPath;
 
 	/** */
-	UPROPERTY(EditAnywhere, Category = "MVVM")
+	UPROPERTY(EditAnywhere, Category = "Viewmodel")
 	EMVVMBindingMode BindingType = EMVVMBindingMode::OneWayToDestination;
 
 	UPROPERTY()
 	bool bOverrideExecutionMode = false;
 
-	UPROPERTY(EditAnywhere, Category = "MVVM", meta=(EditCondition="bOverrideExecutionMode"))
+	UPROPERTY(EditAnywhere, Category = "Viewmodel", meta=(EditCondition="bOverrideExecutionMode"))
 	EMVVMExecutionMode OverrideExecutionMode = EMVVMExecutionMode::Immediate;
 
 	/** */
-	UPROPERTY(EditAnywhere, Category = "MVVM")
+	UPROPERTY(EditAnywhere, Category = "Viewmodel")
 	FMVVMBlueprintViewConversionPath Conversion;
 
 	/** */
-	UPROPERTY(VisibleAnywhere, Category = "MVVM", Transient)
+	UPROPERTY(VisibleAnywhere, Category = "Viewmodel", Transient)
 	TArray<FText> Errors;
 
 	/** The unique ID of this binding. */
-	UPROPERTY(VisibleAnywhere, Category = "MVVM")
+	UPROPERTY(VisibleAnywhere, Category = "Viewmodel")
 	FGuid BindingId;
 
 	/** Whether the binding is enabled or disabled by default. The instance may enable the binding at runtime. */
-	UPROPERTY(EditAnywhere, Category = "MVVM")
+	UPROPERTY(EditAnywhere, Category = "Viewmodel")
 	bool bEnabled = true;
 
 	/** The binding is visible in the editor, but is not compiled and cannot be used at runtime. */
-	UPROPERTY(EditAnywhere, Category = "MVVM")
+	UPROPERTY(EditAnywhere, Category = "Viewmodel")
 	bool bCompile = true;
 
 	/**
diff --git a/Engine/Plugins/Runtime/ModelViewViewModel/Source/ModelViewViewModelBlueprint/Public/MVVMDeveloperProjectSettings.h b/Engine/Plugins/Runtime/ModelViewViewModel/Source/ModelViewViewModelBlueprint/Public/MVVMDeveloperProjectSettings.h
--- a/Engine/Plugins/Runtime/ModelViewViewModel/Source/ModelViewViewModelBlueprint/Public/MVVMDeveloperProjectSettings.h
+++ b/Engine/Plugins/Runtime/ModelViewViewModel/Source/ModelViewViewModelBlueprint/Public/MVVMDeveloperProjectSettings.h
@@ -3,6 +3,8 @@
 #pragma once
 
 #include "Engine/DeveloperSettings.h"
+#include "Types/MVVMExecutionMode.h"
+
 #include "MVVMDeveloperProjectSettings.generated.h"
 
 /**
@@ -32,16 +34,27 @@ class MODELVIEWVIEWMODELBLUEPRINT_API UMVVMDeveloperProjectSettings : public UDe
 	GENERATED_BODY()
 
 public:
+	UMVVMDeveloperProjectSettings();
+
 	virtual FName GetCategoryName() const override;
 	virtual FText GetSectionText() const override;
 
 	bool IsPropertyAllowed(const FProperty* Property) const;
 	bool IsFunctionAllowed(const UFunction* Function) const;
 
+	bool IsExecutionModeAllowed(EMVVMExecutionMode ExecutionMode) const
+	{
+		return AllowedExecutionMode.Contains(ExecutionMode);
+	}
+
 private:
 	/** Permission list for filtering which properties are visible in UI. */
 	UPROPERTY(EditAnywhere, config, Category = "Viewmodel")
 	TMap<FSoftClassPath, FMVVMDeveloperProjectWidgetSettings> FieldSelectorPermissions;
+
+	/** Permission list for filtering which execution mode is allowed. */
+	UPROPERTY(EditAnywhere, config, Category = "Viewmodel")
+	TSet<EMVVMExecutionMode> AllowedExecutionMode;
 };
 
 #if UE_ENABLE_INCLUDE_ORDER_DEPRECATED_IN_5_2
diff --git a/Engine/Plugins/Runtime/ModelViewViewModel/Source/ModelViewViewModelEditor/Private/Widgets/SMVVMViewBindingListView.cpp b/Engine/Plugins/Runtime/ModelViewViewModel/Source/ModelViewViewModelEditor/Private/Widgets/SMVVMViewBindingListView.cpp
--- a/Engine/Plugins/Runtime/ModelViewViewModel/Source/ModelViewViewModelEditor/Private/Widgets/SMVVMViewBindingListView.cpp
+++ b/Engine/Plugins/Runtime/ModelViewViewModel/Source/ModelViewViewModelEditor/Private/Widgets/SMVVMViewBindingListView.cpp
@@ -10,6 +10,7 @@
 #include "Misc/MessageDialog.h"
 #include "MVVMBlueprintView.h"
 #include "MVVMEditorSubsystem.h"
+#include "MVVMDeveloperProjectSettings.h"
 #include "MVVMWidgetBlueprintExtension_View.h"
 #include "ScopedTransaction.h"
 #include "SEnumCombo.h"
@@ -406,11 +407,16 @@ public:
 					.HAlign(HAlign_Right)
 					.AutoWidth()
 					[
-						SNew(SEnumComboBox, StaticEnum<EMVVMExecutionMode>())
-						.ContentPadding(FMargin(4, 0))
-						.OnEnumSelectionChanged(this, &SBindingRow::OnExecutionModeSelectionChanged)
-						.CurrentValue(this, &SBindingRow::GetExecutioModeValue)
+						SNew(SComboButton)
+						.ContentPadding(FMargin(4.f, 0.f))
+						.OnGetMenuContent(this, &SBindingRow::OnGetExecutionModeMenuContent)
 						.IsEnabled(this, &SBindingRow::IsExecutionModeOverridden)
+						.ButtonContent()
+						[
+							SNew(STextBlock)
+							.Text(this, &SBindingRow::GetExecutioModeValue)
+							.ToolTipText(this, &SBindingRow::GetExecutioModeValueToolTip)
+						]
 					]
 				]
 			],
@@ -723,16 +729,60 @@ private:
 		}
 	}
 
-	void OnExecutionModeSelectionChanged(int32 Value, ESelectInfo::Type)
+	void OnExecutionModeSelectionChanged(EMVVMExecutionMode Value)
 	{
 		if (FMVVMBlueprintViewBinding* ViewBinding = GetThisViewBinding())
 		{
 			UMVVMEditorSubsystem* Subsystem = GEditor->GetEditorSubsystem<UMVVMEditorSubsystem>();
-			Subsystem->OverrideExecutionModeForBinding(WidgetBlueprintWeak.Get(), *ViewBinding, static_cast<EMVVMExecutionMode>(Value));
-		}			
+			Subsystem->OverrideExecutionModeForBinding(WidgetBlueprintWeak.Get(), *ViewBinding, Value);
+		}
 	}
 
-	int32 GetExecutioModeValue() const
+	TSharedRef<SWidget> OnGetExecutionModeMenuContent()
+	{
+		const bool bCloseAfterSelection = true;
+		FMenuBuilder MenuBuilder(bCloseAfterSelection, nullptr, nullptr, true);
+
+		UEnum* Enum = StaticEnum<EMVVMExecutionMode>();
+		for (int32 Index = 0; Index < Enum->NumEnums() - 1; Index++)
+		{
+			EMVVMExecutionMode Mode = static_cast<EMVVMExecutionMode>(Enum->GetValueByIndex(Index));
+			if (!GetDefault<UMVVMDeveloperProjectSettings>()->IsExecutionModeAllowed(Mode))
+			{
+				continue;
+			}
+
+			MenuBuilder.AddMenuEntry(
+				Enum->GetDisplayNameTextByIndex(Index),
+				Enum->GetToolTipTextByIndex(Index),
+				FSlateIcon(),
+				FUIAction
+				(
+					FExecuteAction::CreateLambda([this, Mode]()
+					{
+						OnExecutionModeSelectionChanged(Mode);
+					})
+				)
+			);
+		}
+
+		return MenuBuilder.MakeWidget();
+	}
+
+	FText GetExecutioModeValue() const
+	{
+		EMVVMExecutionMode ExecutionMode = CVarDefaultExecutionMode ? (EMVVMExecutionMode)CVarDefaultExecutionMode->GetInt() : EMVVMExecutionMode::Immediate;
+		if (FMVVMBlueprintViewBinding* ViewBinding = GetThisViewBinding())
+		{
+			if (ViewBinding->bOverrideExecutionMode)
+			{
+				ExecutionMode = ViewBinding->OverrideExecutionMode;
+			}
+		}
+		return StaticEnum<EMVVMExecutionMode>()->GetDisplayNameTextByValue(static_cast<int64>(ExecutionMode));
+	}
+	
+	FText GetExecutioModeValueToolTip() const
 	{
 		EMVVMExecutionMode ExecutionMode = CVarDefaultExecutionMode ? (EMVVMExecutionMode)CVarDefaultExecutionMode->GetInt() : EMVVMExecutionMode::Immediate;
 		if (FMVVMBlueprintViewBinding* ViewBinding = GetThisViewBinding())
@@ -742,7 +792,7 @@ private:
 				ExecutionMode = ViewBinding->OverrideExecutionMode;
 			}
 		}
-		return static_cast<int32>(ExecutionMode);
+		return StaticEnum<EMVVMExecutionMode>()->GetToolTipTextByIndex(static_cast<int64>(ExecutionMode));
 	}
 	
 	bool IsExecutionModeOverridden() const
diff --git a/Engine/Plugins/Runtime/ModelViewViewModel/Source/ModelViewViewModelEditor/Public/MVVMEditorSubsystem.h b/Engine/Plugins/Runtime/ModelViewViewModel/Source/ModelViewViewModelEditor/Public/MVVMEditorSubsystem.h
--- a/Engine/Plugins/Runtime/ModelViewViewModel/Source/ModelViewViewModelEditor/Public/MVVMEditorSubsystem.h
+++ b/Engine/Plugins/Runtime/ModelViewViewModel/Source/ModelViewViewModelEditor/Public/MVVMEditorSubsystem.h
@@ -23,7 +23,7 @@ class UK2Node_CallFunction;
 class UWidgetBlueprint;
 
 /** */
-UCLASS()
+UCLASS(DisplayName="Viewmodel Editor Subsystem")
 class MODELVIEWVIEWMODELEDITOR_API UMVVMEditorSubsystem : public UEditorSubsystem
 {
 	GENERATED_BODY()
