MVVM: Move GlobalCollection to GameInstance subsystem.

diff --git a/Engine/Plugins/Runtime/ModelViewViewModel/Source/ModelViewViewModel/Private/MVVMGameSubsystem.cpp b/Engine/Plugins/Runtime/ModelViewViewModel/Source/ModelViewViewModel/Private/MVVMGameSubsystem.cpp
new file mode 100644
--- /dev/null
+++ b/Engine/Plugins/Runtime/ModelViewViewModel/Source/ModelViewViewModel/Private/MVVMGameSubsystem.cpp
@@ -0,0 +1,21 @@
+// Copyright Epic Games, Inc. All Rights Reserved.
+
+#include "MVVMGameSubsystem.h"
+
+#include "Types/MVVMViewModelContextInstance.h"
+
+#include UE_INLINE_GENERATED_CPP_BY_NAME(MVVMGameSubsystem)
+
+
+void UMVVMGameSubsystem::Initialize(FSubsystemCollectionBase& Collection)
+{
+	Super::Initialize(Collection);
+	ViewModelCollection = NewObject<UMVVMViewModelCollectionObject>(this);
+}
+
+
+void UMVVMGameSubsystem::Deinitialize()
+{
+	ViewModelCollection = nullptr;
+	Super::Deinitialize();
+}
diff --git a/Engine/Plugins/Runtime/ModelViewViewModel/Source/ModelViewViewModel/Private/MVVMSubsystem.cpp b/Engine/Plugins/Runtime/ModelViewViewModel/Source/ModelViewViewModel/Private/MVVMSubsystem.cpp
--- a/Engine/Plugins/Runtime/ModelViewViewModel/Source/ModelViewViewModel/Private/MVVMSubsystem.cpp
+++ b/Engine/Plugins/Runtime/ModelViewViewModel/Source/ModelViewViewModel/Private/MVVMSubsystem.cpp
@@ -15,13 +15,11 @@
 void UMVVMSubsystem::Initialize(FSubsystemCollectionBase& Collection)
 {
 	Super::Initialize(Collection);
-	GlobalViewModelCollection = NewObject<UMVVMViewModelCollectionObject>(this);
 }
 
 
 void UMVVMSubsystem::Deinitialize()
 {
-	GlobalViewModelCollection = nullptr;
 	Super::Deinitialize();
 }
 
@@ -292,6 +290,12 @@ FMVVMAvailableBinding UMVVMSubsystem::GetAvailableBindingForField(UE::MVVM::FMVV
 }
 
 
+UMVVMViewModelCollectionObject* UMVVMSubsystem::GetGlobalViewModelCollection() const
+{
+	return nullptr;
+}
+
+
 TValueOrError<bool, FText> UMVVMSubsystem::IsBindingValid(FBindingArgs Args) const
 {
 	if (UE::MVVM::IsForwardBinding(Args.Mode))
diff --git a/Engine/Plugins/Runtime/ModelViewViewModel/Source/ModelViewViewModel/Private/View/MVVMViewClass.cpp b/Engine/Plugins/Runtime/ModelViewViewModel/Source/ModelViewViewModel/Private/View/MVVMViewClass.cpp
--- a/Engine/Plugins/Runtime/ModelViewViewModel/Source/ModelViewViewModel/Private/View/MVVMViewClass.cpp
+++ b/Engine/Plugins/Runtime/ModelViewViewModel/Source/ModelViewViewModel/Private/View/MVVMViewClass.cpp
@@ -7,6 +7,9 @@
 #include "Bindings/MVVMFieldPathHelper.h"
 #include "Blueprint/UserWidget.h"
 #include "Engine/Engine.h"
+#include "Engine/GameInstance.h"
+#include "Engine/World.h"
+#include "MVVMGameSubsystem.h"
 #include "MVVMMessageLog.h"
 #include "MVVMSubsystem.h"
 
@@ -112,7 +115,20 @@ UObject* FMVVMViewClass_SourceCreator::CreateInstance(const UMVVMViewClass* InVi
 		}
 		else if (GlobalViewModelInstance.IsValid())
 		{
-			UMVVMViewModelBase* FoundViewModelInstance = GEngine->GetEngineSubsystem<UMVVMSubsystem>()->GetGlobalViewModelCollection()->FindViewModelInstance(GlobalViewModelInstance);
+			UMVVMViewModelCollectionObject* Collection = nullptr;
+			UMVVMViewModelBase* FoundViewModelInstance = nullptr;
+			if (const UWorld* World = InUserWidget->GetWorld())
+			{
+				if (const UGameInstance* GameInstance = World->GetGameInstance())
+				{
+					Collection = GameInstance->GetSubsystem<UMVVMGameSubsystem>()->GetViewModelCollection();
+					if (Collection)
+					{
+						FoundViewModelInstance = Collection->FindViewModelInstance(GlobalViewModelInstance);
+					}
+				}
+			}
+
 			if (FoundViewModelInstance != nullptr)
 			{
 				ensureMsgf(FoundViewModelInstance->IsA(GlobalViewModelInstance.ContextClass), TEXT("The Global View Model Instance is not of the expected type."));
@@ -122,7 +138,14 @@ UObject* FMVVMViewClass_SourceCreator::CreateInstance(const UMVVMViewClass* InVi
 			else if (!bOptional)
 			{
 				UE::MVVM::FMessageLog Log(InUserWidget);
-				Log.Error(FText::Format(LOCTEXT("CreateInstanceFailedGlobal", "The viewmodel '{0}' was not found in the global view model collection."), FText::FromName(GlobalViewModelInstance.ContextName)));
+				if (Collection)
+				{
+					Log.Error(FText::Format(LOCTEXT("CreateInstanceFailedGlobal", "The source '{0}' was not found in the global view model collection."), FText::FromName(GlobalViewModelInstance.ContextName)));
+				}
+				else
+				{
+					Log.Error(FText::Format(LOCTEXT("CreateInstanceFailedGlobalInstance", "The source '{0}' will be invalid because the global view model collection could not be found."), FText::FromName(GlobalViewModelInstance.ContextName)));
+				}
 			}
 		}
 		else if (FieldPath.IsValid())
diff --git a/Engine/Plugins/Runtime/ModelViewViewModel/Source/ModelViewViewModel/Public/MVVMGameSubsystem.h b/Engine/Plugins/Runtime/ModelViewViewModel/Source/ModelViewViewModel/Public/MVVMGameSubsystem.h
new file mode 100644
--- /dev/null
+++ b/Engine/Plugins/Runtime/ModelViewViewModel/Source/ModelViewViewModel/Public/MVVMGameSubsystem.h
@@ -0,0 +1,32 @@
+// Copyright Epic Games, Inc. All Rights Reserved.
+
+#pragma once
+
+#include "Subsystems/GameInstanceSubsystem.h"
+#include "Types/MVVMViewModelCollection.h"
+
+#include "UObject/Package.h"
+#include "MVVMGameSubsystem.generated.h"
+
+
+/** */
+UCLASS(DisplayName="Viewmodel Game Subsytem")
+class MODELVIEWVIEWMODEL_API UMVVMGameSubsystem : public UGameInstanceSubsystem
+{
+	GENERATED_BODY()
+
+public:
+	//~ Begin UEngineSubsystem interface
+	virtual void Initialize(FSubsystemCollectionBase& Collection) override;
+	virtual void Deinitialize() override;
+	//~ End UEngineSubsystem interface
+
+	UFUNCTION(BlueprintCallable, Category = "Viewmodel")
+	UMVVMViewModelCollectionObject* GetViewModelCollection() const
+	{
+		return ViewModelCollection;
+	}
+private:
+	UPROPERTY(Transient)
+	TObjectPtr<UMVVMViewModelCollectionObject> ViewModelCollection;
+};
diff --git a/Engine/Plugins/Runtime/ModelViewViewModel/Source/ModelViewViewModel/Public/MVVMSubsystem.h b/Engine/Plugins/Runtime/ModelViewViewModel/Source/ModelViewViewModel/Public/MVVMSubsystem.h
--- a/Engine/Plugins/Runtime/ModelViewViewModel/Source/ModelViewViewModel/Public/MVVMSubsystem.h
+++ b/Engine/Plugins/Runtime/ModelViewViewModel/Source/ModelViewViewModel/Public/MVVMSubsystem.h
@@ -61,11 +61,9 @@ public:
 	/** @return The AvailableBinding from a field. */
 	static FMVVMAvailableBinding GetAvailableBindingForField(UE::MVVM::FMVVMConstFieldVariant Variant, const UClass* Accessor);
 
-	UFUNCTION(BlueprintCallable, Category = "Viewmodel")
-	UMVVMViewModelCollectionObject* GetGlobalViewModelCollection() const
-	{
-		return GlobalViewModelCollection;
-	}
+	UE_DEPRECATED(5.3, "GetGlobalViewModelCollection has been deprecated, please use the game instance subsystem.")
+	UFUNCTION(BlueprintCallable, Category = "Viewmodel", meta=(DeprecatedFunction, DeprecatedMessage = "This version of GetGlobalViewModelCollection has been deprecated, please use GetGlobalViewModelCollection from the Viewmodel Game subsystem."))
+	UMVVMViewModelCollectionObject* GetGlobalViewModelCollection() const;
 
 	struct FConstDirectionalBindingArgs
 	{
@@ -100,10 +98,6 @@ public:
 	UE_NODISCARD TValueOrError<bool, FText> IsBindingValid(FConstDirectionalBindingArgs Args) const;
 	UE_NODISCARD TValueOrError<bool, FText> IsBindingValid(FDirectionalBindingArgs Args) const;
 	UE_NODISCARD TValueOrError<bool, FText> IsBindingValid(FBindingArgs Args) const;
-
-private:
-	UPROPERTY(Transient)
-	TObjectPtr<UMVVMViewModelCollectionObject> GlobalViewModelCollection;
 };
 
 #if UE_ENABLE_INCLUDE_ORDER_DEPRECATED_IN_5_2
diff --git a/Engine/Plugins/Runtime/ModelViewViewModel/Source/ModelViewViewModelEditor/Private/MVVMEditorSubsystem.cpp b/Engine/Plugins/Runtime/ModelViewViewModel/Source/ModelViewViewModelEditor/Private/MVVMEditorSubsystem.cpp
--- a/Engine/Plugins/Runtime/ModelViewViewModel/Source/ModelViewViewModelEditor/Private/MVVMEditorSubsystem.cpp
+++ b/Engine/Plugins/Runtime/ModelViewViewModel/Source/ModelViewViewModelEditor/Private/MVVMEditorSubsystem.cpp
@@ -673,8 +673,6 @@ TArray<UFunction*> UMVVMEditorSubsystem::GetAvailableConversionFunctions(const U
 {
 	TArray<UFunction*> ConversionFunctions;
 
-	const UMVVMSubsystem* Subsystem = GEngine->GetEngineSubsystem<UMVVMSubsystem>();
-
 	FBlueprintActionDatabase& ActionDatabase = FBlueprintActionDatabase::Get();
 	const FBlueprintActionDatabase::FActionRegistry& ActionRegistry = ActionDatabase.GetAllActions();
 
@@ -1004,7 +1002,7 @@ TArray<UE::MVVM::FBindingSource> UMVVMEditorSubsystem::GetBindableWidgets(const
 	// Add current widget as a possible binding source
 	if (UClass* BPClass = WidgetBlueprint->GeneratedClass)
 	{
-		TArray<FMVVMAvailableBinding> Bindings = GEngine->GetEngineSubsystem<UMVVMSubsystem>()->GetAvailableBindings(BPClass, WidgetBlueprint->GeneratedClass);
+		TArray<FMVVMAvailableBinding> Bindings = UMVVMSubsystem::GetAvailableBindings(BPClass, WidgetBlueprint->GeneratedClass);
 		if (Bindings.Num() > 0)
 		{
 			// at least one valid property, add it to our list
@@ -1018,7 +1016,7 @@ TArray<UE::MVVM::FBindingSource> UMVVMEditorSubsystem::GetBindableWidgets(const
 
 	for (const UWidget* Widget : AllWidgets)
 	{
-		TArray<FMVVMAvailableBinding> Bindings = GEngine->GetEngineSubsystem<UMVVMSubsystem>()->GetAvailableBindings(Widget->GetClass(), WidgetBlueprint->GeneratedClass);
+		TArray<FMVVMAvailableBinding> Bindings = UMVVMSubsystem::GetAvailableBindings(Widget->GetClass(), WidgetBlueprint->GeneratedClass);
 		if (Bindings.Num() > 0)
 		{
 			// at least one valid property, add it to our list
@@ -1062,7 +1060,7 @@ TArray<FMVVMAvailableBinding> UMVVMEditorSubsystem::GetChildViewModels(TSubclass
 		return TArray<FMVVMAvailableBinding>();
 	}
 
-	TArray<FMVVMAvailableBinding> ViewModelAvailableBindingsList = GEngine->GetEngineSubsystem<UMVVMSubsystem>()->GetAvailableBindings(Class, Accessor);
+	TArray<FMVVMAvailableBinding> ViewModelAvailableBindingsList = UMVVMSubsystem::GetAvailableBindings(Class, Accessor);
 	ViewModelAvailableBindingsList.RemoveAllSwap([Class](const FMVVMAvailableBinding& Value)
 		{
 			UE::MVVM::FMVVMFieldVariant Variant = UE::MVVM::BindingHelper::FindFieldByName(Class.Get(), Value.GetBindingName());
