MVVM: Remove the viewmodel editor from the build.

diff --git a/Engine/Plugins/Runtime/ModelViewViewModel/Source/ModelViewViewModelBlueprint/ModelViewViewModelBlueprint.Build.cs b/Engine/Plugins/Runtime/ModelViewViewModel/Source/ModelViewViewModelBlueprint/ModelViewViewModelBlueprint.Build.cs
--- a/Engine/Plugins/Runtime/ModelViewViewModel/Source/ModelViewViewModelBlueprint/ModelViewViewModelBlueprint.Build.cs
+++ b/Engine/Plugins/Runtime/ModelViewViewModel/Source/ModelViewViewModelBlueprint/ModelViewViewModelBlueprint.Build.cs
@@ -35,5 +35,8 @@ public class ModelViewViewModelBlueprint : ModuleRules
 			{
 				"ModelViewViewModel/Private",
 			});
+
+
+		PublicDefinitions.Add("UE_MVVM_WITH_VIEWMODEL_EDITOR=0");
 	}
 }
diff --git a/Engine/Plugins/Runtime/ModelViewViewModel/Source/ModelViewViewModelBlueprint/Private/ModelViewViewModelBlueprintModule.cpp b/Engine/Plugins/Runtime/ModelViewViewModel/Source/ModelViewViewModelBlueprint/Private/ModelViewViewModelBlueprintModule.cpp
--- a/Engine/Plugins/Runtime/ModelViewViewModel/Source/ModelViewViewModelBlueprint/Private/ModelViewViewModelBlueprintModule.cpp
+++ b/Engine/Plugins/Runtime/ModelViewViewModel/Source/ModelViewViewModelBlueprint/Private/ModelViewViewModelBlueprintModule.cpp
@@ -12,18 +12,22 @@ class FModelViewViewModelBlueprintModule : public IModuleInterface
 public:
 	virtual void StartupModule() override
 	{
+#if UE_MVVM_WITH_VIEWMODEL_EDITOR
 		IKismetCompilerInterface& KismetCompilerModule = FModuleManager::GetModuleChecked<IKismetCompilerInterface>("KismetCompiler");
 		KismetCompilerModule.GetCompilers().Add(&ViewModelBlueprintCompiler);
 
 		FKismetCompilerContext::RegisterCompilerForBP(UMVVMViewModelBlueprint::StaticClass(), &UMVVMViewModelBlueprint::GetCompilerForViewModelBlueprint);
+#endif
 	}
 
 	virtual void ShutdownModule() override
 	{
+#if UE_MVVM_WITH_VIEWMODEL_EDITOR
 		if (IKismetCompilerInterface* KismetCompilerModule = FModuleManager::GetModulePtr<IKismetCompilerInterface>("KismetCompiler"))
 		{
 			KismetCompilerModule->GetCompilers().Remove(&ViewModelBlueprintCompiler);
 		}
+#endif
 	}
 
 private:
diff --git a/Engine/Plugins/Runtime/ModelViewViewModel/Source/ModelViewViewModelEditor/Private/ModelViewViewModelEditorModule.cpp b/Engine/Plugins/Runtime/ModelViewViewModel/Source/ModelViewViewModelEditor/Private/ModelViewViewModelEditorModule.cpp
--- a/Engine/Plugins/Runtime/ModelViewViewModel/Source/ModelViewViewModelEditor/Private/ModelViewViewModelEditorModule.cpp
+++ b/Engine/Plugins/Runtime/ModelViewViewModel/Source/ModelViewViewModelEditor/Private/ModelViewViewModelEditorModule.cpp
@@ -35,9 +35,12 @@ void FModelViewViewModelEditorModule::StartupModule()
 
 	// Register asset types
 	{
+#if UE_MVVM_WITH_VIEWMODEL_EDITOR
+		// Only remove what is related to the viewmodel editor, not the UMG extention for view.
 		IAssetTools& AssetTools = FModuleManager::LoadModuleChecked<FAssetToolsModule>("AssetTools").Get();
 		ViewModelBlueprintActions = MakeShared<UE::MVVM::FAssetTypeActions_ViewModelBlueprint>();
 		AssetTools.RegisterAssetTypeActions(ViewModelBlueprintActions.ToSharedRef());
+#endif
 	}
 
 	PropertyBindingExtension = MakeShared<FMVVMPropertyBindingExtension>();
@@ -76,11 +79,16 @@ void FModelViewViewModelEditorModule::ShutdownModule()
 	}
 	PropertyBindingExtension.Reset();
 
+#if UE_MVVM_WITH_VIEWMODEL_EDITOR
 	// Unregister all the asset types that we registered
 	if (FAssetToolsModule* AssetTools = FModuleManager::GetModulePtr<FAssetToolsModule>("AssetTools"))
 	{
-		AssetTools->Get().UnregisterAssetTypeActions(ViewModelBlueprintActions.ToSharedRef());
+		if (ViewModelBlueprintActions)
+		{
+			AssetTools->Get().UnregisterAssetTypeActions(ViewModelBlueprintActions.ToSharedRef());
+		}
 	}
+#endif
 
 	FMVVMEditorStyle::DestroyInstance();
 
diff --git a/Engine/Plugins/Runtime/ModelViewViewModel/Source/ModelViewViewModelEditor/Private/ViewModel/AssetTypeActions_ViewModelBlueprint.cpp b/Engine/Plugins/Runtime/ModelViewViewModel/Source/ModelViewViewModelEditor/Private/ViewModel/AssetTypeActions_ViewModelBlueprint.cpp
--- a/Engine/Plugins/Runtime/ModelViewViewModel/Source/ModelViewViewModelEditor/Private/ViewModel/AssetTypeActions_ViewModelBlueprint.cpp
+++ b/Engine/Plugins/Runtime/ModelViewViewModel/Source/ModelViewViewModelEditor/Private/ViewModel/AssetTypeActions_ViewModelBlueprint.cpp
@@ -10,6 +10,7 @@
 
 #define LOCTEXT_NAMESPACE "AssetTypeActions"
 
+#if UE_MVVM_WITH_VIEWMODEL_EDITOR
 namespace UE::MVVM
 {
 
@@ -66,5 +67,6 @@ void FAssetTypeActions_ViewModelBlueprint::PerformAssetDiff(UObject* Asset1, UOb
 }
 
 } //namespace
+#endif
 
 #undef LOCTEXT_NAMESPACE
diff --git a/Engine/Plugins/Runtime/ModelViewViewModel/Source/ModelViewViewModelEditor/Private/ViewModel/AssetTypeActions_ViewModelBlueprint.h b/Engine/Plugins/Runtime/ModelViewViewModel/Source/ModelViewViewModelEditor/Private/ViewModel/AssetTypeActions_ViewModelBlueprint.h
--- a/Engine/Plugins/Runtime/ModelViewViewModel/Source/ModelViewViewModelEditor/Private/ViewModel/AssetTypeActions_ViewModelBlueprint.h
+++ b/Engine/Plugins/Runtime/ModelViewViewModel/Source/ModelViewViewModelEditor/Private/ViewModel/AssetTypeActions_ViewModelBlueprint.h
@@ -6,9 +6,10 @@
 
 struct FAssetData;
 
+#if UE_MVVM_WITH_VIEWMODEL_EDITOR
 namespace UE::MVVM
 {
-	
+
 class FAssetTypeActions_ViewModelBlueprint : public FAssetTypeActions_Base
 {
 public:
@@ -25,3 +26,4 @@ public:
 };
 
 } //namespace
+#endif
diff --git a/Engine/Plugins/Runtime/ModelViewViewModel/Source/ModelViewViewModelEditor/Private/ViewModel/MVVMViewModelBlueprintFactory.cpp b/Engine/Plugins/Runtime/ModelViewViewModel/Source/ModelViewViewModelEditor/Private/ViewModel/MVVMViewModelBlueprintFactory.cpp
--- a/Engine/Plugins/Runtime/ModelViewViewModel/Source/ModelViewViewModelEditor/Private/ViewModel/MVVMViewModelBlueprintFactory.cpp
+++ b/Engine/Plugins/Runtime/ModelViewViewModel/Source/ModelViewViewModelEditor/Private/ViewModel/MVVMViewModelBlueprintFactory.cpp
@@ -15,6 +15,8 @@
 
 #define LOCTEXT_NAMESPACE "ViewModelBlueprintFactory"
 
+#if UE_MVVM_WITH_VIEWMODEL_EDITOR
+
 /*------------------------------------------------------------------------------
 	UMVVMViewModelBlueprintFactory implementation.
 ------------------------------------------------------------------------------*/
@@ -46,14 +48,21 @@ public:
 
 }//namespace
 
+#endif
+
 UMVVMViewModelBlueprintFactory::UMVVMViewModelBlueprintFactory()
 {
+#if UE_MVVM_WITH_VIEWMODEL_EDITOR
 	bCreateNew = true;
+#else
+	bCreateNew = false;
+#endif
 	bEditAfterNew = true;
 	SupportedClass = UMVVMViewModelBlueprint::StaticClass();
 	ParentClass = UMVVMViewModelBase::StaticClass();
 }
 
+#if UE_MVVM_WITH_VIEWMODEL_EDITOR
 bool UMVVMViewModelBlueprintFactory::ConfigureProperties()
 {
 	FClassViewerModule& ClassViewerModule = FModuleManager::LoadModuleChecked<FClassViewerModule>("ClassViewer");
@@ -118,5 +127,7 @@ UObject* UMVVMViewModelBlueprintFactory::FactoryCreateNew(UClass* Class, UObject
 	return FactoryCreateNew(Class, InParent, Name, Flags, Context, Warn, NAME_None);
 }
 
+#endif
+
 #undef LOCTEXT_NAMESPACE
 
diff --git a/Engine/Plugins/Runtime/ModelViewViewModel/Source/ModelViewViewModelEditor/Private/ViewModel/MVVMViewModelBlueprintFactory.h b/Engine/Plugins/Runtime/ModelViewViewModel/Source/ModelViewViewModelEditor/Private/ViewModel/MVVMViewModelBlueprintFactory.h
--- a/Engine/Plugins/Runtime/ModelViewViewModel/Source/ModelViewViewModelEditor/Private/ViewModel/MVVMViewModelBlueprintFactory.h
+++ b/Engine/Plugins/Runtime/ModelViewViewModel/Source/ModelViewViewModelEditor/Private/ViewModel/MVVMViewModelBlueprintFactory.h
@@ -14,12 +14,14 @@ class UMVVMViewModelBlueprintFactory : public UFactory
 public:
 	UMVVMViewModelBlueprintFactory();
 
+#if UE_MVVM_WITH_VIEWMODEL_EDITOR
 	//~ Begin UFactory Interface
 	virtual bool ConfigureProperties() override;
 	virtual bool ShouldShowInNewMenu() const override;
 	virtual UObject* FactoryCreateNew(UClass* Class, UObject* InParent, FName Name, EObjectFlags Flags, UObject* Context, FFeedbackContext* Warn, FName CallingContext) override;
 	virtual UObject* FactoryCreateNew(UClass* Class, UObject* InParent, FName Name, EObjectFlags Flags, UObject* Context, FFeedbackContext* Warn) override;
-	//~ End UFactory Interface	
+	//~ End UFactory Interface
+#endif
 
 private:
 	// The parent class of the created blueprint
