MVVM: Test the function as regular conversion function and complex to get the property type.

diff --git a/Engine/Plugins/Runtime/ModelViewViewModel/Source/ModelViewViewModelEditor/Private/Widgets/SMVVMFunctionParameter.cpp b/Engine/Plugins/Runtime/ModelViewViewModel/Source/ModelViewViewModelEditor/Private/Widgets/SMVVMFunctionParameter.cpp
--- a/Engine/Plugins/Runtime/ModelViewViewModel/Source/ModelViewViewModelEditor/Private/Widgets/SMVVMFunctionParameter.cpp
+++ b/Engine/Plugins/Runtime/ModelViewViewModel/Source/ModelViewViewModelEditor/Private/Widgets/SMVVMFunctionParameter.cpp
@@ -3,11 +3,12 @@
 #include "SMVVMFunctionParameter.h" 
 #include "EdGraphSchema_K2.h"
 #include "Editor.h"
-#include "MVVMBlueprintViewBinding.h"
+#include "MVVMBlueprintView.h"
 #include "MVVMEditorSubsystem.h"
 #include "NodeFactory.h"
 #include "SGraphPin.h"
 #include "Styling/MVVMEditorStyle.h"
+#include "Types/MVVMBindingMode.h"
 #include "WidgetBlueprint.h"
 #include "Widgets/Images/SImage.h"
 #include "Widgets/SBoxPanel.h"
@@ -45,19 +46,17 @@ void SFunctionParameter::Construct(const FArguments& InArgs)
 	check(Property);
 
 	TSharedRef<SWidget> ValueWidget = SNullWidget::NullWidget;
+	bool bIsBooleanPin = CastField<const FBoolProperty>(Property) != nullptr;
 
-	bool bIsBooleanPin = false;
-	UEdGraphPin* Pin = EditorSubsystem->GetConversionFunctionArgumentPin(InArgs._WidgetBlueprint, *Binding, ParameterName, bSourceToDestination);
-	if (Pin != nullptr)
+	if (UEdGraphPin* Pin = EditorSubsystem->GetConversionFunctionArgumentPin(InArgs._WidgetBlueprint, *Binding, ParameterName, bSourceToDestination))
 	{
+		bIsBooleanPin = Pin->PinType.PinCategory == UEdGraphSchema_K2::PC_Boolean;
 		// create a new pin widget so that we can get the default value widget out of it
 		if (TSharedPtr<SGraphPin> PinWidget = FNodeFactory::CreateK2PinWidget(Pin))
 		{
 			GraphPin = PinWidget;
 			ValueWidget = PinWidget->GetDefaultValueWidget();
 		}
-
-		bIsBooleanPin = Pin->PinType.PinCategory == UEdGraphSchema_K2::PC_Boolean;
 	}
 
 	if (ValueWidget == SNullWidget::NullWidget)
