MVVM: Add customization for Viewmodel context. Add picker for the propertypath.

diff --git a/Engine/Plugins/Runtime/ModelViewViewModel/Source/ModelViewViewModelBlueprint/Public/MVVMBlueprintViewModelContext.h b/Engine/Plugins/Runtime/ModelViewViewModel/Source/ModelViewViewModelBlueprint/Public/MVVMBlueprintViewModelContext.h
--- a/Engine/Plugins/Runtime/ModelViewViewModel/Source/ModelViewViewModelBlueprint/Public/MVVMBlueprintViewModelContext.h
+++ b/Engine/Plugins/Runtime/ModelViewViewModel/Source/ModelViewViewModelBlueprint/Public/MVVMBlueprintViewModelContext.h
@@ -78,10 +78,10 @@ public:
 
 private:
 	/** When the view is spawn, create an instance of the viewmodel. */
-	UPROPERTY(VisibleAnywhere, AdvancedDisplay, Category = "MVVM")
+	UPROPERTY(VisibleAnywhere, AdvancedDisplay, Category = "Viewmodel")
 	FGuid ViewModelContextId;
 
-	UPROPERTY(VisibleAnywhere, Category = "MVVM", NoClear, meta = (AllowedClasses = "/Script/UMG.NotifyFieldValueChanged", DisallowedClasses = "/Script/UMG.Widget"))
+	UPROPERTY(VisibleAnywhere, Category = "Viewmodel", NoClear, meta = (AllowedClasses = "/Script/UMG.NotifyFieldValueChanged", DisallowedClasses = "/Script/UMG.Widget"))
 	TObjectPtr<UClass> NotifyFieldValueClass = nullptr;
 
 	UPROPERTY()
@@ -92,33 +92,33 @@ private:
 
 public:
 	/** Property name that will be generated. */
-	UPROPERTY(EditAnywhere, Category = "MVVM")
+	UPROPERTY(EditAnywhere, Category = "Viewmodel")
 	FName ViewModelName;
 
 	/** When the view is spawn, create an instance of the viewmodel. */
-	UPROPERTY(EditAnywhere, Category = "MVVM")
+	UPROPERTY(EditAnywhere, Category = "Viewmodel")
 	EMVVMBlueprintViewModelContextCreationType CreationType = EMVVMBlueprintViewModelContextCreationType::CreateInstance;
 
 	/** Identifier of an already registered viewmodel. */
-	UPROPERTY(EditAnywhere, Category = "MVVM", AdvancedDisplay, meta = (EditCondition = "CreationType == EMVVMBlueprintViewModelContextCreationType::GlobalViewModelCollection"))
+	UPROPERTY(EditAnywhere, Category = "Viewmodel", AdvancedDisplay, meta = (EditCondition = "CreationType == EMVVMBlueprintViewModelContextCreationType::GlobalViewModelCollection", EditConditionHides))
 	FName GlobalViewModelIdentifier;
 
 	/** The Path to get the viewmodel instance. */
-	UPROPERTY(EditAnywhere, Category = "MVVM", AdvancedDisplay, meta = (EditCondition = "CreationType == EMVVMBlueprintViewModelContextCreationType::PropertyPath"))
+	UPROPERTY(EditAnywhere, Category = "Viewmodel", AdvancedDisplay, meta = (EditCondition = "CreationType == EMVVMBlueprintViewModelContextCreationType::PropertyPath", EditConditionHides))
 	FString ViewModelPropertyPath;
 
 	/**
 	 * Generate a setter function for this viewmodel.
 	 * @note Always true when the Creation Type is Manual.
 	 */
-	UPROPERTY(EditAnywhere, Category = "MVVM", AdvancedDisplay, meta = (EditCondition = "CreationType != EMVVMBlueprintViewModelContextCreationType::Manual"))
+	UPROPERTY(EditAnywhere, Category = "Viewmodel", AdvancedDisplay, meta = (EditCondition = "CreationType != EMVVMBlueprintViewModelContextCreationType::Manual", EditConditionHides))
 	bool bCreateSetterFunction = false;
 
 	/**
 	 * Optional. Will not warn if the instance is not set or found.
 	 * @note Always true when the Creation Type is Manual.
 	 */
-	UPROPERTY(EditAnywhere, Category = "MVVM", AdvancedDisplay, meta = (EditCondition = "CreationType == EMVVMBlueprintViewModelContextCreationType::GlobalViewModelCollection || CreationType == EMVVMBlueprintViewModelContextCreationType::PropertyPath"))
+	UPROPERTY(EditAnywhere, Category = "Viewmodel", AdvancedDisplay, meta = (EditCondition = "CreationType == EMVVMBlueprintViewModelContextCreationType::GlobalViewModelCollection || CreationType == EMVVMBlueprintViewModelContextCreationType::PropertyPath", EditConditionHides))
 	bool bOptional = false;
 };
 
diff --git a/Engine/Plugins/Runtime/ModelViewViewModel/Source/ModelViewViewModelEditor/Private/Customizations/MVVMBlueprintViewModelContextCustomization.cpp b/Engine/Plugins/Runtime/ModelViewViewModel/Source/ModelViewViewModelEditor/Private/Customizations/MVVMBlueprintViewModelContextCustomization.cpp
new file mode 100644
--- /dev/null
+++ b/Engine/Plugins/Runtime/ModelViewViewModel/Source/ModelViewViewModelEditor/Private/Customizations/MVVMBlueprintViewModelContextCustomization.cpp
@@ -0,0 +1,217 @@
+// Copyright Epic Games, Inc. All Rights Reserved.
+
+#include "MVVMBlueprintViewModelContextCustomization.h"
+
+#include "IPropertyTypeCustomization.h"
+#include "PropertyHandle.h"
+#include "IDetailChildrenBuilder.h"
+#include "DetailWidgetRow.h"
+#include "WidgetBlueprintEditor.h"
+#include "IPropertyAccessEditor.h"
+#include "Bindings/MVVMBindingHelper.h"
+#include "Features/IModularFeatures.h"
+#include "Widgets/SBoxPanel.h"
+#include "Widgets/SMVVMViewModelPanel.h"
+
+#define LOCTEXT_NAMESPACE "BlueprintViewModelContextDetailCustomization"
+
+namespace UE::MVVM
+{
+namespace Private
+{
+	FText BindingWidgetForVM_GetName()
+	{
+		return FText::GetEmpty();
+	}
+
+	bool BindingWidgetForVM_CanBindProperty(const FProperty* Property, const UClass* ClassToLookFor)
+	{
+		const FObjectPropertyBase* ObjectProperty = CastField<const FObjectPropertyBase>(Property);
+		return (ObjectProperty != nullptr && ObjectProperty->PropertyClass->IsChildOf(ClassToLookFor));
+	}
+}
+
+/**
+ *
+ */
+bool FViewModelPropertyAccessEditor::CanBindProperty(FProperty* Property) const
+{
+	// Property == GeneratePureBindingsProperty is only to start the algo
+	return ViewModelProperty != Property && (Private::BindingWidgetForVM_CanBindProperty(Property, ClassToLookFor.Get()) || Property == GeneratePureBindingsProperty);
+}
+
+bool FViewModelPropertyAccessEditor::CanBindFunction(UFunction* Function) const
+{
+	return Private::BindingWidgetForVM_CanBindProperty(BindingHelper::GetReturnProperty(Function), ClassToLookFor.Get());
+}
+
+bool FViewModelPropertyAccessEditor::CanBindToClass(UClass* Class) const
+{
+	return true;
+}
+
+void FViewModelPropertyAccessEditor::AddBinding(FName, const TArray<FBindingChainElement>& BindingChain)
+{
+	TStringBuilder<256> Path;
+	for (const FBindingChainElement& Binding : BindingChain)
+	{
+		if (Path.Len() != 0)
+		{
+			Path << TEXT('.');
+		}
+		Path << Binding.Field.GetFName();
+	}
+
+	AssignToProperty->SetValue(Path.ToString());
+}
+
+bool FViewModelPropertyAccessEditor::HasValidClassToLookFor() const
+{
+	return ClassToLookFor.Get() != nullptr;
+}
+
+TSharedRef<SWidget> FViewModelPropertyAccessEditor::MakePropertyBindingWidget(TSharedRef<FWidgetBlueprintEditor> WidgetBlueprintEditor, FProperty* PropertyToMatch, TSharedRef<IPropertyHandle> InAssignToProperty, FName ViewModelPropertyName)
+{
+	UClass* SkeletonClass = WidgetBlueprintEditor->GetBlueprintObj()->SkeletonGeneratedClass.Get();
+	if (!SkeletonClass)
+	{
+		return SNullWidget::NullWidget;
+	}
+	ViewModelProperty = SkeletonClass->FindPropertyByName(ViewModelPropertyName);
+	AssignToProperty = InAssignToProperty;
+
+	if (!IModularFeatures::Get().IsModularFeatureAvailable("PropertyAccessEditor"))
+	{
+		return SNullWidget::NullWidget;
+	}
+
+	GeneratePureBindingsProperty = PropertyToMatch;
+	FPropertyBindingWidgetArgs Args;
+	Args.Property = GeneratePureBindingsProperty;
+	Args.bAllowArrayElementBindings = false;
+	Args.bAllowStructMemberBindings = true;
+	Args.bAllowUObjectFunctions = true;
+	Args.bAllowStructFunctions = true;
+	Args.bAllowNewBindings = true;
+	Args.bGeneratePureBindings = true;
+
+	Args.CurrentBindingText.BindStatic(&Private::BindingWidgetForVM_GetName);
+	Args.OnCanBindProperty.BindRaw(this, &FViewModelPropertyAccessEditor::CanBindProperty);
+	Args.OnCanBindFunction.BindRaw(this, &FViewModelPropertyAccessEditor::CanBindFunction);
+	Args.OnCanBindToClass.BindRaw(this, &FViewModelPropertyAccessEditor::CanBindToClass);
+	Args.OnAddBinding.BindRaw(this, &FViewModelPropertyAccessEditor::AddBinding);
+
+	IPropertyAccessEditor& PropertyAccessEditor = IModularFeatures::Get().GetModularFeature<IPropertyAccessEditor>("PropertyAccessEditor");
+	TSharedRef<SWidget> Result = PropertyAccessEditor.MakePropertyBindingWidget(WidgetBlueprintEditor->GetBlueprintObj(), Args);
+	Result->SetEnabled(MakeAttributeRaw(this, &FViewModelPropertyAccessEditor::HasValidClassToLookFor));
+	return Result;
+}
+
+/**
+ * 
+ */
+FBlueprintViewModelContextDetailCustomization::FBlueprintViewModelContextDetailCustomization(TWeakPtr<FWidgetBlueprintEditor> InEditor)
+	: WidgetBlueprintEditor(InEditor)
+{}
+
+FBlueprintViewModelContextDetailCustomization::~FBlueprintViewModelContextDetailCustomization()
+{
+}
+
+void FBlueprintViewModelContextDetailCustomization::CustomizeChildren(TSharedRef<IPropertyHandle> PropertyHandle, IDetailChildrenBuilder& ChildBuilder, IPropertyTypeCustomizationUtils& CustomizationUtils)
+{
+	uint32 NumChildren = 0;
+	PropertyHandle->GetNumChildren(NumChildren);
+
+	FName ViewModelPropertyName;
+	for (uint32 ChildIndex = 0; ChildIndex < NumChildren; ++ChildIndex)
+	{
+		TSharedPtr<IPropertyHandle> ChildHandle = PropertyHandle->GetChildHandle(ChildIndex);
+		const FName PropertyName = ChildHandle->GetProperty()->GetFName();
+
+		if (PropertyName == "NotifyFieldValueClass")
+		{
+			ensure(CastField<FClassProperty>(ChildHandle->GetProperty()));
+			NotifyFieldValueClassHandle = ChildHandle;
+			UObject* Object = nullptr;
+			if (ChildHandle->GetValue(Object) == FPropertyAccess::Success)
+			{
+				PropertyAccessEditor.ClassToLookFor = Cast<UClass>(Object);
+			}
+			ChildHandle->SetOnPropertyValueChanged(FSimpleDelegate::CreateSP(this, &FBlueprintViewModelContextDetailCustomization::HandleClassChanged));
+		}
+
+		if (PropertyName == GET_MEMBER_NAME_CHECKED(FMVVMBlueprintViewModelContext, ViewModelName))
+		{
+			ensure(CastField<FNameProperty>(ChildHandle->GetProperty()));
+			ChildHandle->GetValue(ViewModelPropertyName);
+		}
+		
+		if (PropertyName == GET_MEMBER_NAME_CHECKED(FMVVMBlueprintViewModelContext, ViewModelPropertyPath))
+		{
+			ensure(CastField<FStrProperty>(ChildHandle->GetProperty()));
+			PropertyPathHandle = ChildHandle;
+		}
+	}
+
+	if (NotifyFieldValueClassHandle)
+	{
+		for (uint32 ChildIndex = 0; ChildIndex < NumChildren; ++ChildIndex)
+		{
+			TSharedPtr<IPropertyHandle> ChildHandle = PropertyHandle->GetChildHandle(ChildIndex);
+			const FName PropertyName = ChildHandle->GetProperty()->GetFName();
+
+			if (PropertyName == GET_MEMBER_NAME_CHECKED(FMVVMBlueprintViewModelContext, ViewModelPropertyPath))
+			{
+				ensure(CastField<FStrProperty>(ChildHandle->GetProperty()));
+				if (TSharedPtr<FWidgetBlueprintEditor> SharedWidgetBlueprintEditor = WidgetBlueprintEditor.Pin())
+				{
+					IDetailPropertyRow& PropertyRow = ChildBuilder.AddProperty(ChildHandle.ToSharedRef());
+
+					TSharedPtr<SWidget> NameWidget, ValueWidget;
+					PropertyRow.GetDefaultWidgets(NameWidget, ValueWidget);
+					PropertyRow.CustomWidget()
+					.NameContent()
+					[
+						NameWidget.ToSharedRef()
+					]
+					.ValueContent()
+					[
+						SNew(SHorizontalBox)
+						+ SHorizontalBox::Slot()
+						.FillWidth(1.f)
+						[
+							ValueWidget.ToSharedRef()
+						]
+						+ SHorizontalBox::Slot()
+						.AutoWidth()
+						[
+							PropertyAccessEditor.MakePropertyBindingWidget(SharedWidgetBlueprintEditor.ToSharedRef(), NotifyFieldValueClassHandle->GetProperty(), PropertyPathHandle.ToSharedRef(), ViewModelPropertyName)
+						]
+					];
+				}
+			}
+			else
+			{
+				ChildBuilder.AddProperty(ChildHandle.ToSharedRef());
+			}
+		}
+	}
+}
+
+void FBlueprintViewModelContextDetailCustomization::HandleClassChanged()
+{
+	UObject* Object = nullptr;
+	if (NotifyFieldValueClassHandle->GetValue(Object) == FPropertyAccess::Success)
+	{
+		PropertyAccessEditor.ClassToLookFor = Cast<UClass>(Object);
+	}
+	else
+	{
+		PropertyAccessEditor.ClassToLookFor.Reset();
+	}
+}
+
+} // namespace UE::MVVM
+
+#undef LOCTEXT_NAMESPACE
diff --git a/Engine/Plugins/Runtime/ModelViewViewModel/Source/ModelViewViewModelEditor/Private/Customizations/MVVMBlueprintViewModelContextCustomization.h b/Engine/Plugins/Runtime/ModelViewViewModel/Source/ModelViewViewModelEditor/Private/Customizations/MVVMBlueprintViewModelContextCustomization.h
new file mode 100644
--- /dev/null
+++ b/Engine/Plugins/Runtime/ModelViewViewModel/Source/ModelViewViewModelEditor/Private/Customizations/MVVMBlueprintViewModelContextCustomization.h
@@ -0,0 +1,59 @@
+// Copyright Epic Games, Inc. All Rights Reserved.
+
+#pragma once
+
+#include "IPropertyTypeCustomization.h"
+#include "Layout/Visibility.h"
+
+class FWidgetBlueprintEditor;
+struct FBindingChainElement;
+class SWidget;
+
+
+namespace UE::MVVM
+{
+
+class FViewModelPropertyAccessEditor
+{
+public:
+	TWeakObjectPtr<UClass> ClassToLookFor;
+
+	TSharedRef<SWidget> MakePropertyBindingWidget(TSharedRef<FWidgetBlueprintEditor> WidgetBlueprintEditor, FProperty* PropertyToMatch, TSharedRef<IPropertyHandle> AssignToProperty, FName ViewModelPropertyName);
+
+private:
+	bool CanBindProperty(FProperty* Property) const;
+	bool CanBindFunction(UFunction* Function) const;
+	bool CanBindToClass(UClass* Class) const;
+	bool HasValidClassToLookFor() const;
+	void AddBinding(FName, const TArray<FBindingChainElement>& BindingChain);
+
+	FProperty* GeneratePureBindingsProperty = nullptr;
+	FProperty* ViewModelProperty = nullptr;
+	TSharedPtr<IPropertyHandle> AssignToProperty = nullptr;
+};
+
+class FBlueprintViewModelContextDetailCustomization : public IPropertyTypeCustomization
+{
+public:
+	FBlueprintViewModelContextDetailCustomization(TWeakPtr<FWidgetBlueprintEditor> InEditor);
+	~FBlueprintViewModelContextDetailCustomization();
+
+	//~ IDetailCustomization interface
+	virtual void CustomizeHeader(TSharedRef<IPropertyHandle> PropertyHandle, FDetailWidgetRow& HeaderRow, IPropertyTypeCustomizationUtils& CustomizationUtils) override {}
+	virtual void CustomizeChildren(TSharedRef<IPropertyHandle> PropertyHandle, IDetailChildrenBuilder& ChildBuilder, IPropertyTypeCustomizationUtils& CustomizationUtils) override;
+
+	static TSharedRef<IPropertyTypeCustomization> MakeInstance(TWeakPtr<FWidgetBlueprintEditor> InEditor)
+	{
+		return MakeShared<FBlueprintViewModelContextDetailCustomization>(InEditor);
+	}
+
+private:
+	void HandleClassChanged();
+
+	FViewModelPropertyAccessEditor PropertyAccessEditor;
+	TWeakPtr<FWidgetBlueprintEditor> WidgetBlueprintEditor;
+	TSharedPtr<IPropertyHandle> NotifyFieldValueClassHandle;
+	TSharedPtr<IPropertyHandle> PropertyPathHandle;
+};
+
+} // namespace UE::MVVM
diff --git a/Engine/Plugins/Runtime/ModelViewViewModel/Source/ModelViewViewModelEditor/Private/ModelViewViewModelEditorModule.cpp b/Engine/Plugins/Runtime/ModelViewViewModel/Source/ModelViewViewModelEditor/Private/ModelViewViewModelEditorModule.cpp
--- a/Engine/Plugins/Runtime/ModelViewViewModel/Source/ModelViewViewModelEditor/Private/ModelViewViewModelEditorModule.cpp
+++ b/Engine/Plugins/Runtime/ModelViewViewModel/Source/ModelViewViewModelEditor/Private/ModelViewViewModelEditorModule.cpp
@@ -4,6 +4,7 @@
 
 #include "BlueprintModes/WidgetBlueprintApplicationMode.h"
 #include "BlueprintModes/WidgetBlueprintApplicationModes.h"
+#include "Customizations/MVVMBlueprintViewModelContextCustomization.h"
 #include "Customizations/MVVMPropertyBindingExtension.h"
 #include "Framework/Application/SlateApplication.h"
 #include "Framework/Docking/LayoutExtender.h"
@@ -13,6 +14,7 @@
 #include "MVVMEditorCommands.h"
 #include "MVVMEditorSubsystem.h"
 #include "MVVMWidgetBlueprintExtension_View.h"
+#include "PropertyEditorModule.h"
 #include "Styling/MVVMEditorStyle.h"
 #include "Tabs/MVVMBindingSummoner.h"
 #include "Tabs/MVVMViewModelSummoner.h"
@@ -46,6 +48,11 @@ void FModelViewViewModelEditorModule::StartupModule()
 	PropertyBindingExtension = MakeShared<FMVVMPropertyBindingExtension>();
 	UMGEditorModule.GetPropertyBindingExtensibilityManager()->AddExtension(PropertyBindingExtension.ToSharedRef());
 
+	UMGEditorModule.RegisterInstancedCustomPropertyTypeLayout(
+		FMVVMBlueprintViewModelContext::StaticStruct()->GetStructPathName()
+		, IUMGEditorModule::FOnGetInstancePropertyTypeCustomizationInstance::CreateStatic(UE::MVVM::FBlueprintViewModelContextDetailCustomization::MakeInstance)
+		);
+
 	FBlueprintEditorUtils::OnRenameVariableReferencesEvent.AddRaw(this, &FModelViewViewModelEditorModule::HandleRenameVariableReferences);
 
 	{
@@ -76,6 +83,10 @@ void FModelViewViewModelEditorModule::ShutdownModule()
 	{
 		UMGEditorModule->OnRegisterTabsForEditor().RemoveAll(this);
 		UMGEditorModule->GetPropertyBindingExtensibilityManager()->RemoveExtension(PropertyBindingExtension.ToSharedRef());
+		if (UObjectInitialized())
+		{
+			UMGEditorModule->UnregisterInstancedCustomPropertyTypeLayout(FMVVMBlueprintViewModelContext::StaticStruct()->GetStructPathName());
+		}
 	}
 	PropertyBindingExtension.Reset();
 
diff --git a/Engine/Plugins/Runtime/ModelViewViewModel/Source/ModelViewViewModelEditor/Private/Widgets/SMVVMViewModelPanel.cpp b/Engine/Plugins/Runtime/ModelViewViewModel/Source/ModelViewViewModelEditor/Private/Widgets/SMVVMViewModelPanel.cpp
--- a/Engine/Plugins/Runtime/ModelViewViewModel/Source/ModelViewViewModelEditor/Private/Widgets/SMVVMViewModelPanel.cpp
+++ b/Engine/Plugins/Runtime/ModelViewViewModel/Source/ModelViewViewModelEditor/Private/Widgets/SMVVMViewModelPanel.cpp
@@ -25,6 +25,19 @@
 
 #define LOCTEXT_NAMESPACE "ViewModelPanel"
 
+
+void UMVVMBlueprintViewModelContextWrapper::PostEditChangeProperty(struct FPropertyChangedEvent& PropertyChangedEvent)
+{
+	Super::PostEditChangeProperty(PropertyChangedEvent);
+	if (UMVVMBlueprintView* BlueprintViewPtr = BlueprintView.Get())
+	{
+		if (FMVVMBlueprintViewModelContext* ViewModelContextPtr = BlueprintViewPtr->FindViewModel(ViewModelId))
+		{
+			*ViewModelContextPtr = Wrapper;
+		}
+	}
+}
+
 namespace UE::MVVM
 {
 
@@ -39,6 +52,8 @@ void SMVVMViewModelPanel::Construct(const FArguments& InArgs, TSharedPtr<FWidget
 	FieldIterator = MakeUnique<FFieldIterator_Bindable>(WidgetBlueprint, EFieldVisibility::None);
 	FieldExpander = MakeUnique<FFieldExpander_Bindable>();
 
+	WidgetBlueprintEditor->OnSelectedWidgetsChanging.AddSP(this, &SMVVMViewModelPanel::HandleEditorSelectionChanged);
+
 	if (CurrentBlueprintView)
 	{
 		// Listen to when the viewmodel are modified
@@ -68,19 +83,8 @@ void SMVVMViewModelPanel::Construct(const FArguments& InArgs, TSharedPtr<FWidget
 
 	FillViewModel();
 
+	ModelContextWrapper.Reset(NewObject<UMVVMBlueprintViewModelContextWrapper>());
 
-	FPropertyEditorModule& PropertyEditor = FModuleManager::Get().GetModuleChecked<FPropertyEditorModule>("PropertyEditor");
-	FDetailsViewArgs DetailsViewArgs;
-	DetailsViewArgs.NameAreaSettings = FDetailsViewArgs::HideNameArea;
-	DetailsViewArgs.NotifyHook = this;
-	DetailsViewArgs.bHideSelectionTip = true;
-	DetailsViewArgs.DefaultsOnlyVisibility = EEditDefaultsOnlyNodeVisibility::Hide;
-	DetailsViewArgs.bAllowSearch = false;
-	FStructureDetailsViewArgs StructDetailViewArgs;
-	StructDetailViewArgs.bShowObjects = true;
-	StructDetailViewArgs.bShowInterfaces = true;
-
-	PropertyView = PropertyEditor.CreateStructureDetailView(DetailsViewArgs, StructDetailViewArgs, TSharedPtr<class FStructOnScope>());
 
 	ChildSlot
 	[
@@ -94,11 +98,6 @@ void SMVVMViewModelPanel::Construct(const FArguments& InArgs, TSharedPtr<FWidget
 			[
 				ViewModelTreeView.ToSharedRef()
 			]
-			+ SSplitter::Slot()
-			.SizeRule(SSplitter::SizeToContent)
-			[
-				PropertyView->GetWidget().ToSharedRef()
-			]
 		]
 	];
 }
@@ -468,28 +467,57 @@ TSharedPtr<SWidget> SMVVMViewModelPanel::HandleContextMenuOpening(UE::PropertyVi
 void SMVVMViewModelPanel::HandleSelectionChanged(UE::PropertyViewer::SPropertyViewer::FHandle ContainerHandle, TArrayView<const FFieldVariant> Fields, ESelectInfo::Type SelectionType)
 {
 	SelectedViewModelGuid = FGuid();
+	ModelContextWrapper->ViewModelId = FGuid();
+	ModelContextWrapper->BlueprintView.Reset();
 
-	bool bSet = false;
-	if (Fields.Num() == 0)
+	if (bIsViewModelSelecting)
 	{
-		if (const FGuid* VMGuidPtr = PropertyViewerHandles.Find(ContainerHandle))
-		{
-			SelectedViewModelGuid = *VMGuidPtr;
+		return;
+	}
+
+	TGuardValue Tmp = TGuardValue(bIsViewModelSelecting, true);
 
-			if (UMVVMBlueprintView* BlueprintView = WeakBlueprintView.Get())
+	if (TSharedPtr<FWidgetBlueprintEditor> BlueprintEditor = WeakBlueprintEditor.Pin())
+	{
+		bool bSet = false;
+		if (Fields.Num() == 0)
+		{
+			if (const FGuid* VMGuidPtr = PropertyViewerHandles.Find(ContainerHandle))
 			{
-				if (FMVVMBlueprintViewModelContext* ViewModelContext = BlueprintView->FindViewModel(SelectedViewModelGuid))
+				SelectedViewModelGuid = *VMGuidPtr;
+
+				if (UMVVMBlueprintView* BlueprintView = WeakBlueprintView.Get())
 				{
-					PropertyView->SetStructureData(MakeShared<FStructOnScope>(FMVVMBlueprintViewModelContext::StaticStruct(), reinterpret_cast<uint8*>(ViewModelContext)));
-					bSet = true;
+					if (FMVVMBlueprintViewModelContext* ViewModelContext = BlueprintView->FindViewModel(SelectedViewModelGuid))
+					{
+						ModelContextWrapper->Wrapper = *ViewModelContext;
+						ModelContextWrapper->ViewModelId = SelectedViewModelGuid;
+						ModelContextWrapper->BlueprintView = WeakBlueprintView;
+						
+						BlueprintEditor->CleanSelection();
+
+						TSet<UObject*> Selections;
+						Selections.Add(ModelContextWrapper.Get());
+						BlueprintEditor->SelectObjects(Selections);
+						bSet = true;
+					}
 				}
 			}
 		}
+
+		if (!bSet && BlueprintEditor->GetSelectedObjects().Contains(ModelContextWrapper.Get()))
+		{
+			BlueprintEditor->SelectObjects(TSet<UObject*>());
+		}
 	}
-	
-	if (!bSet)
+}
+
+
+void SMVVMViewModelPanel::HandleEditorSelectionChanged()
+{
+	if (!bIsViewModelSelecting && SelectedViewModelGuid.IsValid())
 	{
-		PropertyView->SetStructureData(TSharedPtr<FStructOnScope>());
+		ViewModelTreeView->SetSelection(UE::PropertyViewer::SPropertyViewer::FHandle(), TArrayView<const FFieldVariant>());
 	}
 }
 
diff --git a/Engine/Plugins/Runtime/ModelViewViewModel/Source/ModelViewViewModelEditor/Private/Widgets/SMVVMViewModelPanel.h b/Engine/Plugins/Runtime/ModelViewViewModel/Source/ModelViewViewModelEditor/Private/Widgets/SMVVMViewModelPanel.h
--- a/Engine/Plugins/Runtime/ModelViewViewModel/Source/ModelViewViewModelEditor/Private/Widgets/SMVVMViewModelPanel.h
+++ b/Engine/Plugins/Runtime/ModelViewViewModel/Source/ModelViewViewModelEditor/Private/Widgets/SMVVMViewModelPanel.h
@@ -6,17 +6,37 @@
 #include "Misc/NotifyHook.h"
 #include "Widgets/PropertyViewer/SPropertyViewer.h"
 
+#include "UObject/StrongObjectPtr.h"
+#include "MVVMBlueprintViewModelContext.h"
+#include "SMVVMViewModelPanel.generated.h"
+
 namespace ETextCommit { enum Type : int; }
 namespace UE::MVVM { class FFieldIterator_Bindable; }
 namespace UE::PropertyViewer { class FFieldExpander_Default; }
 
+class INotifyFieldValueChanged;
 class FWidgetBlueprintEditor;
 class SInlineEditableTextBlock;
 class SPositiveActionButton;
 class UBlueprintExtension;
 class UMVVMBlueprintView;
 class FUICommandList;
-class IStructureDetailsView;
+class IDetailsView;
+
+UCLASS()
+class UMVVMBlueprintViewModelContextWrapper : public UObject
+{
+	GENERATED_BODY()
+
+public:
+	virtual void PostEditChangeProperty(struct FPropertyChangedEvent& PropertyChangedEvent) override;
+
+	UPROPERTY(EditAnywhere, Category="Viewmodel", meta = (ShowOnlyInnerProperties))
+	FMVVMBlueprintViewModelContext Wrapper;
+
+	TWeakObjectPtr<UMVVMBlueprintView> BlueprintView;
+	FGuid ViewModelId;
+};
 
 namespace UE::MVVM
 {
@@ -48,6 +68,7 @@ private:
 	TSharedRef<SWidget> HandleGenerateContainer(UE::PropertyViewer::SPropertyViewer::FHandle ContainerHandle, TOptional<FText> DisplayName);
 	TSharedPtr<SWidget> HandleContextMenuOpening(UE::PropertyViewer::SPropertyViewer::FHandle ContainerHandle, TArrayView<const FFieldVariant> Field) const;
 	void HandleSelectionChanged(UE::PropertyViewer::SPropertyViewer::FHandle, TArrayView<const FFieldVariant>, ESelectInfo::Type);
+	void HandleEditorSelectionChanged();
 	bool HandleVerifyNameTextChanged(const FText& InText, FText& OutErrorMessage, FGuid ViewModelGuid);
 	void HandleNameTextCommited(const FText& InText, ETextCommit::Type CommitInfo, FGuid ViewModelGuid);
 
@@ -71,7 +92,6 @@ private:
 	TUniquePtr<FFieldIterator_Bindable> FieldIterator;
 	TUniquePtr<UE::PropertyViewer::FFieldExpander_Default> FieldExpander;
 	TSharedPtr<FUICommandList> CommandList;
-	TSharedPtr<IStructureDetailsView> PropertyView;
 
 	TMap<UE::PropertyViewer::SPropertyViewer::FHandle, FGuid> PropertyViewerHandles;
 	TMap<FGuid, TSharedPtr<SInlineEditableTextBlock>> EditableTextBlocks;
@@ -80,8 +100,11 @@ private:
 	TWeakObjectPtr<UMVVMBlueprintView> WeakBlueprintView;
 	FDelegateHandle ViewModelsUpdatedHandle;
 
+	TStrongObjectPtr<UMVVMBlueprintViewModelContextWrapper> ModelContextWrapper;
+
 	FGuid SelectedViewModelGuid;
 	FName PreviousViewModelPropertyName;
+	bool bIsViewModelSelecting = false;
 };
 
 } // namespace
diff --git a/Engine/Source/Runtime/AdvancedWidgets/Private/Widgets/PropertyViewer/SPropertyViewerImpl.cpp b/Engine/Source/Runtime/AdvancedWidgets/Private/Widgets/PropertyViewer/SPropertyViewerImpl.cpp
--- a/Engine/Source/Runtime/AdvancedWidgets/Private/Widgets/PropertyViewer/SPropertyViewerImpl.cpp
+++ b/Engine/Source/Runtime/AdvancedWidgets/Private/Widgets/PropertyViewer/SPropertyViewerImpl.cpp
@@ -773,17 +773,17 @@ TSharedPtr<FTreeNode> FPropertyViewerImpl::FindExistingChild(const TSharedPtr<FT
 
 void FPropertyViewerImpl::SetSelection(SPropertyViewer::FHandle Identifier, TArrayView<const FFieldVariant> FieldPath)
 {
+	TreeWidget->ClearSelection();
 	for (const TSharedPtr<FTreeNode>& Node : TreeSource)
 	{
 		if (TSharedPtr<FContainer> Container = Node->GetContainer())
 		{
 			if (Container->GetIdentifier() == Identifier)
 			{
-				TreeWidget->ClearSelection();
-
 				if (TSharedPtr<FTreeNode> FoundNode = FindExistingChild(Node, FieldPath))
 				{
 					TreeWidget->SetItemSelection(FoundNode, true);
+					break;
 				}
 			}
 		}
