MVVM: Add setting to find conversion functions.

diff --git a/Engine/Plugins/Runtime/ModelViewViewModel/Source/ModelViewViewModelBlueprint/Private/MVVMDeveloperProjectSettings.cpp b/Engine/Plugins/Runtime/ModelViewViewModel/Source/ModelViewViewModelBlueprint/Private/MVVMDeveloperProjectSettings.cpp
--- a/Engine/Plugins/Runtime/ModelViewViewModel/Source/ModelViewViewModelBlueprint/Private/MVVMDeveloperProjectSettings.cpp
+++ b/Engine/Plugins/Runtime/ModelViewViewModel/Source/ModelViewViewModelBlueprint/Private/MVVMDeveloperProjectSettings.cpp
@@ -1,6 +1,7 @@
 // Copyright Epic Games, Inc. All Rights Reserved.
 
 #include "MVVMDeveloperProjectSettings.h"
+#include "Engine/Blueprint.h"
 
 #define LOCTEXT_NAMESPACE "MVVMDeveloperProjectSettings"
 
@@ -49,4 +50,18 @@ bool UMVVMDeveloperProjectSettings::IsFunctionAllowed(const UFunction* Function)
 	return true;
 }
 
+TArray<const UClass*> UMVVMDeveloperProjectSettings::GetAllowedConversionFunctionClasses() const
+{
+	TArray<const UClass*> Result;
+	for (const FSoftClassPath& SoftClass : AllowedClassForConversionFunctions)
+	{
+		if (UClass* Class = SoftClass.ResolveClass())
+		{
+			Result.Add(Class);
+		}
+	}
+
+	return Result;
+}
+
 #undef LOCTEXT_NAMESPACE
diff --git a/Engine/Plugins/Runtime/ModelViewViewModel/Source/ModelViewViewModelBlueprint/Public/MVVMDeveloperProjectSettings.h b/Engine/Plugins/Runtime/ModelViewViewModel/Source/ModelViewViewModelBlueprint/Public/MVVMDeveloperProjectSettings.h
--- a/Engine/Plugins/Runtime/ModelViewViewModel/Source/ModelViewViewModelBlueprint/Public/MVVMDeveloperProjectSettings.h
+++ b/Engine/Plugins/Runtime/ModelViewViewModel/Source/ModelViewViewModelBlueprint/Public/MVVMDeveloperProjectSettings.h
@@ -24,6 +24,16 @@ struct FMVVMDeveloperProjectWidgetSettings
 	TSet<FName> AdvancedFieldNames;
 };
 
+/**
+ *
+ */
+UENUM()
+enum class EMVVMDeveloperConversionFunctionFilterType : uint8
+{
+	BlueprintActionRegistry,
+	AllowedList,
+};
+
 
 /**
  * Implements the settings for the MVVM Editor
@@ -47,6 +57,13 @@ public:
 		return AllowedExecutionMode.Contains(ExecutionMode);
 	}
 
+	EMVVMDeveloperConversionFunctionFilterType GetConversionFunctionFilter() const
+	{
+		return ConversionFunctionFilter;
+	}
+
+	TArray<const UClass*> GetAllowedConversionFunctionClasses() const;
+
 private:
 	/** Permission list for filtering which properties are visible in UI. */
 	UPROPERTY(EditAnywhere, config, Category = "Viewmodel")
@@ -55,6 +72,15 @@ private:
 	/** Permission list for filtering which execution mode is allowed. */
 	UPROPERTY(EditAnywhere, config, Category = "Viewmodel")
 	TSet<EMVVMExecutionMode> AllowedExecutionMode;
+
+public:
+	/** Permission list for filtering which execution mode is allowed. */
+	UPROPERTY(EditAnywhere, config, Category = "Viewmodel")
+	EMVVMDeveloperConversionFunctionFilterType ConversionFunctionFilter = EMVVMDeveloperConversionFunctionFilterType::BlueprintActionRegistry;
+
+	/** Individual class that are allowed to be uses as conversion functions. */
+	UPROPERTY(EditAnywhere, config, Category = "Viewmodel", meta = (EditCondition = "ConversionFunctionFilter == EMVVMDeveloperConversionFunctionFilterType::AllowedList"))
+	TSet<FSoftClassPath> AllowedClassForConversionFunctions;
 };
 
 #if UE_ENABLE_INCLUDE_ORDER_DEPRECATED_IN_5_2
diff --git a/Engine/Plugins/Runtime/ModelViewViewModel/Source/ModelViewViewModelEditor/Private/MVVMEditorSubsystem.cpp b/Engine/Plugins/Runtime/ModelViewViewModel/Source/ModelViewViewModelEditor/Private/MVVMEditorSubsystem.cpp
--- a/Engine/Plugins/Runtime/ModelViewViewModel/Source/ModelViewViewModelEditor/Private/MVVMEditorSubsystem.cpp
+++ b/Engine/Plugins/Runtime/ModelViewViewModel/Source/ModelViewViewModelEditor/Private/MVVMEditorSubsystem.cpp
@@ -17,6 +17,7 @@
 #include "K2Node_FunctionResult.h"
 #include "K2Node_VariableGet.h"
 #include "MVVMBlueprintView.h"
+#include "MVVMDeveloperProjectSettings.h"
 #include "MVVMSubsystem.h"
 #include "MVVMWidgetBlueprintExtension_View.h"
 #include "ScopedTransaction.h"
@@ -673,34 +674,81 @@ TArray<UFunction*> UMVVMEditorSubsystem::GetAvailableConversionFunctions(const U
 {
 	TArray<UFunction*> ConversionFunctions;
 
-	FBlueprintActionDatabase& ActionDatabase = FBlueprintActionDatabase::Get();
-	const FBlueprintActionDatabase::FActionRegistry& ActionRegistry = ActionDatabase.GetAllActions();
+	auto AddFunction = [this, &Source, &Destination, WidgetBlueprint, &ConversionFunctions](const UFunction* Function)
+	{
+		// functions in the widget blueprint can do anything they want, other functions have to be static functions in a BlueprintFunctionLibrary
+		const UClass* FunctionClass = Function->GetOuterUClass();
 
-	for (auto It = ActionRegistry.CreateConstIterator(); It; ++It)
+		bool bIsFromWidgetBlueprint = WidgetBlueprint->GeneratedClass->IsChildOf(FunctionClass) && Function->HasAllFunctionFlags(FUNC_BlueprintPure);
+		bool bFromBlueprintFunctionLibrary = FunctionClass->IsChildOf<UBlueprintFunctionLibrary>() && Function->HasAllFunctionFlags(FUNC_Static | FUNC_BlueprintPure);
+		if (bIsFromWidgetBlueprint || bFromBlueprintFunctionLibrary)
+		{
+			if (IsValidConversionFunction(Function, Source, Destination))
+			{
+				ConversionFunctions.Add(const_cast<UFunction*>(Function));
+			}
+		}
+	};
+
+	EMVVMDeveloperConversionFunctionFilterType FilterType = GetDefault<UMVVMDeveloperProjectSettings>()->GetConversionFunctionFilter();
+	if (FilterType == EMVVMDeveloperConversionFunctionFilterType::BlueprintActionRegistry)
 	{
-		if (UObject* ActionObject = It->Key.ResolveObjectPtr())
+		const FBlueprintActionDatabase::FActionRegistry& ActionRegistry = FBlueprintActionDatabase::Get().GetAllActions();
+		for (auto It = ActionRegistry.CreateConstIterator(); It; ++It)
 		{
-			for (const UBlueprintNodeSpawner* NodeSpawner : It->Value)
+			if (UObject* ActionObject = It->Key.ResolveObjectPtr())
 			{
-				FBlueprintActionInfo BlueprintAction(ActionObject, NodeSpawner);
-				const UFunction* Function = BlueprintAction.GetAssociatedFunction();
-				if (Function != nullptr)
+				for (const UBlueprintNodeSpawner* NodeSpawner : It->Value)
 				{
-					// functions in the widget blueprint can do anything they want, other functions have to be static functions in a BlueprintFunctionLibrary
-					const UClass* FunctionClass = Function->GetOuterUClass();
-					if ((FunctionClass->ClassGeneratedBy == WidgetBlueprint) ||
-						(WidgetBlueprint->ParentClass && WidgetBlueprint->ParentClass->IsChildOf(FunctionClass) && Function->HasAllFunctionFlags(FUNC_Static | FUNC_BlueprintPure)) ||
-						(FunctionClass->IsChildOf<UBlueprintFunctionLibrary>() && Function->HasAllFunctionFlags(FUNC_Static | FUNC_BlueprintPure)))
+					FBlueprintActionInfo BlueprintAction(ActionObject, NodeSpawner);
+					const UFunction* Function = BlueprintAction.GetAssociatedFunction();
+					if (Function != nullptr)
 					{
-						if (IsValidConversionFunction(Function, Source, Destination))
-						{
-							ConversionFunctions.Add(const_cast<UFunction*>(Function));
-						}
+						AddFunction(Function);
 					}
 				}
 			}
 		}
 	}
+	else if (FilterType == EMVVMDeveloperConversionFunctionFilterType::AllowedList)
+	{
+		auto IsInheritedBlueprintFunction = [](const UFunction* Function)
+		{
+			bool bIsBpInheritedFunc = false;
+			if (UClass* FuncClass = Function->GetOwnerClass())
+			{
+				if (UBlueprint* BpOwner = Cast<UBlueprint>(FuncClass->ClassGeneratedBy))
+				{
+					FName FuncName = Function->GetFName();
+					if (UClass* ParentClass = BpOwner->ParentClass)
+					{
+						bIsBpInheritedFunc = (ParentClass->FindFunctionByName(FuncName, EIncludeSuperFlag::IncludeSuper) != nullptr);
+					}
+				}
+			}
+			return bIsBpInheritedFunc;
+		};
+
+		TArray<const UClass*> Classes = GetDefault<UMVVMDeveloperProjectSettings>()->GetAllowedConversionFunctionClasses();
+		Classes.Add(WidgetBlueprint->GeneratedClass);
+		for (const UClass* Class : Classes)
+		{
+			for (TFieldIterator<UFunction> FunctionIt(Class, EFieldIteratorFlags::ExcludeSuper); FunctionIt; ++FunctionIt)
+			{
+				UFunction* Function = *FunctionIt;
+				if (IsInheritedBlueprintFunction(Function))
+				{
+					continue;
+				}
+				if (!Function->HasAllFunctionFlags(FUNC_BlueprintCallable))
+				{
+					continue;
+				}
+
+				AddFunction(Function);
+			}
+		}
+	}
 
 	ConversionFunctions.Sort([](const UFunction& A, const UFunction& B) -> bool
 		{
diff --git a/Engine/Source/Editor/UMGEditor/Private/UMGEditorModule.cpp b/Engine/Source/Editor/UMGEditor/Private/UMGEditorModule.cpp
--- a/Engine/Source/Editor/UMGEditor/Private/UMGEditorModule.cpp
+++ b/Engine/Source/Editor/UMGEditor/Private/UMGEditorModule.cpp
@@ -299,13 +299,12 @@ public:
 
 	virtual void RegisterInstancedCustomPropertyTypeLayout(FTopLevelAssetPath Type, FOnGetInstancePropertyTypeCustomizationInstance Delegate) override
 	{
-		int32 IndexOf = CustomPropertyTypeLayout.IndexOfByPredicate([Type](const FCustomPropertyTypeLayout& Element) { return Element.Type == Type; });
-		if (ensure(!CustomPropertyTypeLayout.IsValidIndex(IndexOf)))
+		bool bIncluded = CustomPropertyTypeLayout.ContainsByPredicate([Type](const FCustomPropertyTypeLayout& Element) { return Element.Type == Type; });
+		if (ensure(!bIncluded))
 		{
-			FCustomPropertyTypeLayout Layout;
+			FCustomPropertyTypeLayout& Layout = CustomPropertyTypeLayout.AddDefaulted_GetRef();
 			Layout.Type = Type;
 			Layout.Delegate = MoveTemp(Delegate);
-			CustomPropertyTypeLayout.Add(MoveTemp(Layout));
 		}
 	}
 
