MVVM: Allow picking selecting an abstract class in as the viewmodal in UMG View but you are not able to create an instance.
Add a new way to select the viewmodel uses by the UMG View. User can extend the Resolver and can implement their own logic.

diff --git a/Engine/Plugins/Runtime/ModelViewViewModel/Source/ModelViewViewModel/Private/View/MVVMViewClass.cpp b/Engine/Plugins/Runtime/ModelViewViewModel/Source/ModelViewViewModel/Private/View/MVVMViewClass.cpp
--- a/Engine/Plugins/Runtime/ModelViewViewModel/Source/ModelViewViewModel/Private/View/MVVMViewClass.cpp
+++ b/Engine/Plugins/Runtime/ModelViewViewModel/Source/ModelViewViewModel/Private/View/MVVMViewClass.cpp
@@ -3,6 +3,7 @@
 #include "View/MVVMViewClass.h"
 #include "Types/MVVMFieldContext.h"
 #include "View/MVVMView.h"
+#include "View/MVVMViewModelResolver.h"
 
 #include "Bindings/MVVMFieldPathHelper.h"
 #include "Blueprint/UserWidget.h"
@@ -79,6 +80,19 @@ FMVVMViewClass_SourceCreator FMVVMViewClass_SourceCreator::MakeGlobalContext(FNa
 	return Result;
 }
 
+FMVVMViewClass_SourceCreator FMVVMViewClass_SourceCreator::MakeResolver(FName InName, UClass* InNotifyFieldValueChangedClass, UMVVMViewModelResolver* InResolver, bool bOptional)
+{
+	FMVVMViewClass_SourceCreator Result;
+	if (ensure(InResolver && InNotifyFieldValueChangedClass && InNotifyFieldValueChangedClass->ImplementsInterface(UNotifyFieldValueChanged::StaticClass())))
+	{
+		Result.PropertyName = InName;
+		Result.ExpectedSourceType = InNotifyFieldValueChangedClass;
+		Result.Resolver = InResolver;
+		Result.bOptional = bOptional;
+		ensure(InResolver->GetPackage() != GetTransientPackage() && !InResolver->HasAnyFlags(RF_Transient));
+	}
+	return Result;
+}
 
 UObject* FMVVMViewClass_SourceCreator::CreateInstance(const UMVVMViewClass* InViewClass, UMVVMView* InView, UUserWidget* InUserWidget) const
 {
@@ -93,8 +107,7 @@ UObject* FMVVMViewClass_SourceCreator::CreateInstance(const UMVVMViewClass* InVi
 	{
 		auto AssignProperty = [FoundObjectProperty, InUserWidget](UObject* NewObject)
 		{
-			check(NewObject);
-			if (ensure(NewObject->GetClass()->IsChildOf(FoundObjectProperty->PropertyClass)))
+			if (ensure(NewObject == nullptr || NewObject->GetClass()->IsChildOf(FoundObjectProperty->PropertyClass)))
 			{
 				FoundObjectProperty->SetObjectPropertyValue_InContainer(InUserWidget, NewObject);
 			}
@@ -105,7 +118,6 @@ UObject* FMVVMViewClass_SourceCreator::CreateInstance(const UMVVMViewClass* InVi
 			if (ExpectedSourceType.Get() != nullptr)
 			{
 				Result = NewObject<UObject>(InUserWidget, ExpectedSourceType.Get(), NAME_None, RF_Transient);
-				AssignProperty(Result);
 			}
 			else if (!bOptional)
 			{
@@ -113,6 +125,22 @@ UObject* FMVVMViewClass_SourceCreator::CreateInstance(const UMVVMViewClass* InVi
 				Log.Error(FText::Format(LOCTEXT("CreateInstanceCreateInstance", "The source '{0}' could not be created. The class is not loaded."), FText::FromName(PropertyName)));
 			}
 		}
+		else if (Resolver)
+		{
+			Result = Resolver->CreateInstance(ExpectedSourceType.Get(), InUserWidget, InView);
+			if (!Result && !bOptional)
+			{
+				UE::MVVM::FMessageLog Log(InUserWidget);
+				Log.Error(FText::Format(LOCTEXT("CreateInstanceFailResolver", "The source '{0}' could not be created. Resolver returned an invalid value."), FText::FromName(PropertyName)));
+			}
+			if (Result && !Result->IsA(ExpectedSourceType.Get()))
+			{
+				Result = nullptr;
+
+				UE::MVVM::FMessageLog Log(InUserWidget);
+				Log.Error(FText::Format(LOCTEXT("CreateInstanceFailResolverExpected", "The source '{0}' could not be created. Resolver returned viewodel of an unexpected type."), FText::FromName(PropertyName)));
+			}
+		}
 		else if (GlobalViewModelInstance.IsValid())
 		{
 			UMVVMViewModelCollectionObject* Collection = nullptr;
@@ -133,7 +161,6 @@ UObject* FMVVMViewClass_SourceCreator::CreateInstance(const UMVVMViewClass* InVi
 			{
 				ensureMsgf(FoundViewModelInstance->IsA(GlobalViewModelInstance.ContextClass), TEXT("The Global View Model Instance is not of the expected type."));
 				Result = FoundViewModelInstance;
-				AssignProperty(Result);
 			}
 			else if (!bOptional)
 			{
@@ -157,7 +184,6 @@ UObject* FMVVMViewClass_SourceCreator::CreateInstance(const UMVVMViewClass* InVi
 				if (ObjectResult.HasValue() && ObjectResult.GetValue() != nullptr)
 				{
 					Result = ObjectResult.GetValue();
-					AssignProperty(ObjectResult.GetValue());
 				}
 			}
 
@@ -167,6 +193,8 @@ UObject* FMVVMViewClass_SourceCreator::CreateInstance(const UMVVMViewClass* InVi
 				Log.Error(FText::Format(LOCTEXT("CreateInstanceInvalidBiding", "The source '{0}' was evaluated to be invalid at initialization."), FText::FromName(PropertyName)));
 			}
 		}
+
+		AssignProperty(Result);
 	}
 
 	return Result;
diff --git a/Engine/Plugins/Runtime/ModelViewViewModel/Source/ModelViewViewModel/Public/View/MVVMViewClass.h b/Engine/Plugins/Runtime/ModelViewViewModel/Source/ModelViewViewModel/Public/View/MVVMViewClass.h
--- a/Engine/Plugins/Runtime/ModelViewViewModel/Source/ModelViewViewModel/Public/View/MVVMViewClass.h
+++ b/Engine/Plugins/Runtime/ModelViewViewModel/Source/ModelViewViewModel/Public/View/MVVMViewClass.h
@@ -16,6 +16,7 @@ class UMVVMUserWidgetBinding;
 class UMVVMView;
 class UMVVMViewClass;
 class UMVVMViewModelBlueprintExtension;
+class UMVVMViewModelResolver;
 class UUserWidget;
 
 namespace UE::MVVM::Private
@@ -38,6 +39,7 @@ public:
 	MODELVIEWVIEWMODEL_API static FMVVMViewClass_SourceCreator MakeInstance(FName Name, UClass* NotifyFieldValueChangedClass);
 	MODELVIEWVIEWMODEL_API static FMVVMViewClass_SourceCreator MakeFieldPath(FName Name, UClass* NotifyFieldValueChangedClass, FMVVMVCompiledFieldPath FieldPath, bool bOptional);
 	MODELVIEWVIEWMODEL_API static FMVVMViewClass_SourceCreator MakeGlobalContext(FName Name, FMVVMViewModelContext Context, bool bOptional);
+	MODELVIEWVIEWMODEL_API static FMVVMViewClass_SourceCreator MakeResolver(FName Name, UClass* InNotifyFieldValueChangedClass, UMVVMViewModelResolver* Source, bool bOptional);
 
 	UObject* CreateInstance(const UMVVMViewClass* ViewClass, UMVVMView* View, UUserWidget* UserWidget) const;
 
@@ -56,6 +58,10 @@ private:
 	UPROPERTY()
 	TSubclassOf<UObject> ExpectedSourceType;
 
+	/** The resolver to fetch the source at runtime. */
+	UPROPERTY(Instanced)
+	TObjectPtr<UMVVMViewModelResolver> Resolver = nullptr;
+
 	/** Info to find the ViewModel instance at runtime. */
 	UPROPERTY()
 	FMVVMViewModelContext GlobalViewModelInstance;
diff --git a/Engine/Plugins/Runtime/ModelViewViewModel/Source/ModelViewViewModel/Public/View/MVVMViewModelResolver.h b/Engine/Plugins/Runtime/ModelViewViewModel/Source/ModelViewViewModel/Public/View/MVVMViewModelResolver.h
new file mode 100644
--- /dev/null
+++ b/Engine/Plugins/Runtime/ModelViewViewModel/Source/ModelViewViewModel/Public/View/MVVMViewModelResolver.h
@@ -0,0 +1,30 @@
+// Copyright Epic Games, Inc. All Rights Reserved.
+
+#pragma once
+
+#include "UObject/Object.h"
+#include "FieldNotification/IFieldValueChanged.h"
+
+#include "MVVMViewModelResolver.generated.h"
+
+class UMVVMView;
+class UMVVMViewClass;
+class UUserWidget;
+
+/**
+ * Shared data to find or create a ViewModel at runtime.
+ */
+UCLASS(Abstract, Blueprintable, EditInlineNew, DisplayName = "Viewmode Resolver")
+class MODELVIEWVIEWMODEL_API UMVVMViewModelResolver : public UObject
+{
+	GENERATED_BODY()
+
+public:
+	virtual UObject* CreateInstance(const UClass* ExpectedType, const UUserWidget* UserWidget, const UMVVMView* View) const
+	{
+		return K2_CreateInstance(ExpectedType, UserWidget).GetObject();
+	}
+
+	UFUNCTION(BlueprintImplementableEvent, Category="Viewmodel", DisplayName="Create Instance")
+	TScriptInterface<INotifyFieldValueChanged> K2_CreateInstance(const UClass* ExpectedType, const UUserWidget* UserWidget) const;
+};
diff --git a/Engine/Plugins/Runtime/ModelViewViewModel/Source/ModelViewViewModelBlueprint/Private/MVVMViewBlueprintCompiler.cpp b/Engine/Plugins/Runtime/ModelViewViewModel/Source/ModelViewViewModelBlueprint/Private/MVVMViewBlueprintCompiler.cpp
--- a/Engine/Plugins/Runtime/ModelViewViewModel/Source/ModelViewViewModelBlueprint/Private/MVVMViewBlueprintCompiler.cpp
+++ b/Engine/Plugins/Runtime/ModelViewViewModel/Source/ModelViewViewModelBlueprint/Private/MVVMViewBlueprintCompiler.cpp
@@ -10,8 +10,9 @@
 #include "HAL/IConsoleManager.h"
 #include "MVVMFunctionGraphHelper.h"
 #include "Templates/ValueOrError.h"
-#include "View/MVVMViewClass.h"
 #include "Types/MVVMBindingName.h"
+#include "View/MVVMViewClass.h"
+#include "View/MVVMViewModelResolver.h"
 #include "UObject/LinkerLoad.h"
 
 #define LOCTEXT_NAMESPACE "MVVMViewBlueprintCompiler"
@@ -806,6 +807,15 @@ bool FMVVMViewBlueprintCompiler::PreCompileSourceCreators(UWidgetBlueprintGenera
 					continue;
 				}
 			}
+			else if (ViewModelContext.CreationType == EMVVMBlueprintViewModelContextCreationType::Resolver)
+			{
+				if (!ViewModelContext.Resolver)
+				{
+					AddErrorForViewModel(ViewModelContext, LOCTEXT("ViewmodelInvalidResolver", "Viewmodel doesn't have a valid Resolver. You can specify a new one in the Viewmodels panel."));
+					bAreSourcesCreatorValid = false;
+					continue;
+				}
+			}
 			else
 			{
 				AddErrorForViewModel(ViewModelContext, LOCTEXT("ViewmodelInvalidCreationType", "Viewmodel doesn't have a valid creation type. You can select one in the Viewmodels panel."));
@@ -874,6 +884,18 @@ bool FMVVMViewBlueprintCompiler::CompileSourceCreators(const FCompiledBindingLib
 
 				CompiledSourceCreator = FMVVMViewClass_SourceCreator::MakeGlobalContext(ViewModelContext.GetViewModelName(), MoveTemp(GlobalViewModelInstance), ViewModelContext.bOptional);
 			}
+			else if (ViewModelContext.CreationType == EMVVMBlueprintViewModelContextCreationType::Resolver)
+			{
+				UMVVMViewModelResolver* Resolver = DuplicateObject(ViewModelContext.Resolver.Get(), ViewExtension);
+				if (!Resolver)
+				{
+					AddErrorForViewModel(ViewModelContext, LOCTEXT("ViewmodelFailedResolverDuplicate", "Internal error. The resolver could not be dupliated."));
+					bAreSourcesCreatorValid = false;
+					continue;
+				}
+
+				CompiledSourceCreator = FMVVMViewClass_SourceCreator::MakeResolver(ViewModelContext.GetViewModelName(), ViewModelContext.GetViewModelClass(), Resolver, ViewModelContext.bOptional);
+			}
 			else
 			{
 				AddErrorForViewModel(ViewModelContext, LOCTEXT("ViewModelWithoutValidCreationType", "The viewmodel doesn't have a valid creation type."));
diff --git a/Engine/Plugins/Runtime/ModelViewViewModel/Source/ModelViewViewModelBlueprint/Public/MVVMBlueprintViewModelContext.h b/Engine/Plugins/Runtime/ModelViewViewModel/Source/ModelViewViewModelBlueprint/Public/MVVMBlueprintViewModelContext.h
--- a/Engine/Plugins/Runtime/ModelViewViewModel/Source/ModelViewViewModelBlueprint/Public/MVVMBlueprintViewModelContext.h
+++ b/Engine/Plugins/Runtime/ModelViewViewModel/Source/ModelViewViewModelBlueprint/Public/MVVMBlueprintViewModelContext.h
@@ -7,6 +7,8 @@
 
 #include "MVVMBlueprintViewModelContext.generated.h"
 
+class UMVVMViewModelResolver;
+
 /**
  *
  */
@@ -21,12 +23,14 @@ enum class EMVVMBlueprintViewModelContextCreationType : uint8
 	GlobalViewModelCollection,
 	// The viewmodel will be fetched by evaluating a function or a property path.
 	PropertyPath,
+	// The viewmodel will be fetched by evaluating the resolver object.
+	Resolver,
 };
 
 /**
  *
  */
-USTRUCT(BlueprintType)
+USTRUCT()
 struct MODELVIEWVIEWMODELBLUEPRINT_API FMVVMBlueprintViewModelContext
 {
 	GENERATED_BODY()
@@ -107,6 +111,9 @@ public:
 	UPROPERTY(EditAnywhere, Category = "Viewmodel", AdvancedDisplay, meta = (EditCondition = "CreationType == EMVVMBlueprintViewModelContextCreationType::PropertyPath", EditConditionHides))
 	FString ViewModelPropertyPath;
 
+	UPROPERTY(EditAnywhere, Category = "Viewmodel", Instanced, AdvancedDisplay, meta = (EditInline, EditCondition = "CreationType == EMVVMBlueprintViewModelContextCreationType::Resolver", EditConditionHides))
+	TObjectPtr<UMVVMViewModelResolver> Resolver = nullptr;
+
 	/**
 	 * Generate a setter function for this viewmodel.
 	 * @note Always true when the Creation Type is Manual.
diff --git a/Engine/Plugins/Runtime/ModelViewViewModel/Source/ModelViewViewModelEditor/Private/Widgets/SMVVMSelectViewModel.cpp b/Engine/Plugins/Runtime/ModelViewViewModel/Source/ModelViewViewModelEditor/Private/Widgets/SMVVMSelectViewModel.cpp
--- a/Engine/Plugins/Runtime/ModelViewViewModel/Source/ModelViewViewModelEditor/Private/Widgets/SMVVMSelectViewModel.cpp
+++ b/Engine/Plugins/Runtime/ModelViewViewModel/Source/ModelViewViewModelEditor/Private/Widgets/SMVVMSelectViewModel.cpp
@@ -19,9 +19,7 @@ namespace UE::MVVM
 
 namespace Private
 {
-	static const EClassFlags DisallowedClassFlags = CLASS_HideDropDown | CLASS_Hidden | CLASS_Deprecated | CLASS_Abstract | CLASS_NotPlaceable;
-
-	bool IsValidViewModel(const UClass* InClass)
+	bool IsValidViewModel(const UClass* InClass, EClassFlags DisallowedClassFlags)
 	{
 		if (InClass == nullptr || InClass->IsChildOf(UWidget::StaticClass()))
 		{
@@ -33,7 +31,7 @@ namespace Private
 
 bool FViewModelClassFilter::IsClassAllowed(const FClassViewerInitializationOptions& InInitOptions, const UClass* InClass, TSharedRef<FClassViewerFilterFuncs> InFilterFuncs)
 {
-	return Private::IsValidViewModel(InClass);
+	return Private::IsValidViewModel(InClass, DisallowedClassFlags);
 }
 
 bool FViewModelClassFilter::IsUnloadedClassAllowed(const FClassViewerInitializationOptions& InInitOptions, const TSharedRef<const IUnloadedBlueprintData> InUnloadedClassData, TSharedRef<FClassViewerFilterFuncs> InFilterFuncs)
@@ -42,7 +40,7 @@ bool FViewModelClassFilter::IsUnloadedClassAllowed(const FClassViewerInitializat
 	{
 		return false;
 	}
-	return InUnloadedClassData->ImplementsInterface(UNotifyFieldValueChanged::StaticClass()) && !InUnloadedClassData->HasAnyClassFlags(Private::DisallowedClassFlags);
+	return InUnloadedClassData->ImplementsInterface(UNotifyFieldValueChanged::StaticClass()) && !InUnloadedClassData->HasAnyClassFlags(DisallowedClassFlags);
 }
 
 
@@ -51,11 +49,15 @@ void SMVVMSelectViewModel::Construct(const FArguments& InArgs, const UWidgetBlue
 	OnCancel = InArgs._OnCancel;
 	OnViewModelCommitted = InArgs._OnViewModelCommitted;
 
+	TSharedRef<FViewModelClassFilter> ClassFilter = MakeShared<FViewModelClassFilter>();
+	DisallowedClassFlags = InArgs._DisallowedClassFlags.Get(ClassFilter->DisallowedClassFlags);
+	ClassFilter->DisallowedClassFlags = InArgs._DisallowedClassFlags.GetValue();
+
 	FClassViewerInitializationOptions ClassViewerOptions;
 	ClassViewerOptions.DisplayMode = EClassViewerDisplayMode::TreeView;
 	ClassViewerOptions.Mode = EClassViewerMode::ClassPicker;
 	ClassViewerOptions.NameTypeToDisplay = EClassViewerNameTypeToDisplay::DisplayName;
-	ClassViewerOptions.ClassFilters.Add(MakeShared<FViewModelClassFilter>());
+	ClassViewerOptions.ClassFilters.Add(ClassFilter);
 
 	ClassViewer = FModuleManager::LoadModuleChecked<FClassViewerModule>("ClassViewer")
 		.CreateClassViewer(ClassViewerOptions, FOnClassPicked::CreateSP(this, &SMVVMSelectViewModel::HandleClassPicked));
@@ -132,7 +134,7 @@ void SMVVMSelectViewModel::HandleClassPicked(UClass* ClassPicked)
 	BindingListWidget->ClearSources();
 	SelectedClass.Reset();
 
-	if (Private::IsValidViewModel(ClassPicked))
+	if (Private::IsValidViewModel(ClassPicked, DisallowedClassFlags))
 	{
 		SelectedClass = ClassPicked;
 		BindingListWidget->AddSource(ClassPicked, FName(), FGuid());
diff --git a/Engine/Plugins/Runtime/ModelViewViewModel/Source/ModelViewViewModelEditor/Private/Widgets/SMVVMSelectViewModel.h b/Engine/Plugins/Runtime/ModelViewViewModel/Source/ModelViewViewModelEditor/Private/Widgets/SMVVMSelectViewModel.h
--- a/Engine/Plugins/Runtime/ModelViewViewModel/Source/ModelViewViewModelEditor/Private/Widgets/SMVVMSelectViewModel.h
+++ b/Engine/Plugins/Runtime/ModelViewViewModel/Source/ModelViewViewModelEditor/Private/Widgets/SMVVMSelectViewModel.h
@@ -4,6 +4,7 @@
 
 #include "ClassViewerFilter.h"
 #include "Widgets/SCompoundWidget.h"
+#include "UObject/ObjectMacros.h"
 
 class SClassViewer;
 class UWidgetBlueprint;
@@ -19,6 +20,8 @@ public:
 
 	virtual bool IsClassAllowed(const FClassViewerInitializationOptions& InInitOptions, const UClass* InClass, TSharedRef<FClassViewerFilterFuncs> InFilterFuncs) override;
 	virtual bool IsUnloadedClassAllowed(const FClassViewerInitializationOptions& InInitOptions, const TSharedRef<const IUnloadedBlueprintData> InUnloadedClassData, TSharedRef<FClassViewerFilterFuncs> InFilterFuncs) override;
+
+	EClassFlags DisallowedClassFlags = CLASS_HideDropDown | CLASS_Hidden | CLASS_Deprecated | CLASS_Abstract | CLASS_NotPlaceable;
 };
 
 class SMVVMSelectViewModel : public SCompoundWidget
@@ -29,6 +32,7 @@ public:
 	SLATE_BEGIN_ARGS(SMVVMSelectViewModel) {}
 		SLATE_EVENT(FSimpleDelegate, OnCancel)
 		SLATE_EVENT(FOnValueChanged, OnViewModelCommitted)
+		SLATE_ARGUMENT(TOptional<EClassFlags>, DisallowedClassFlags)
 	SLATE_END_ARGS()
 
 public:
@@ -47,6 +51,7 @@ private:
 
 	FSimpleDelegate OnCancel;
 	FOnValueChanged OnViewModelCommitted;
+	EClassFlags DisallowedClassFlags;
 };
 
 } //namespace
diff --git a/Engine/Plugins/Runtime/ModelViewViewModel/Source/ModelViewViewModelEditor/Private/Widgets/SMVVMViewModelPanel.cpp b/Engine/Plugins/Runtime/ModelViewViewModel/Source/ModelViewViewModelEditor/Private/Widgets/SMVVMViewModelPanel.cpp
--- a/Engine/Plugins/Runtime/ModelViewViewModel/Source/ModelViewViewModelEditor/Private/Widgets/SMVVMViewModelPanel.cpp
+++ b/Engine/Plugins/Runtime/ModelViewViewModel/Source/ModelViewViewModelEditor/Private/Widgets/SMVVMViewModelPanel.cpp
@@ -7,6 +7,7 @@
 #include "MVVMBlueprintView.h"
 #include "MVVMEditorSubsystem.h"
 #include "PropertyEditorModule.h"
+#include "View/MVVMViewModelResolver.h"
 #include "WidgetBlueprint.h"
 #include "WidgetBlueprintEditor.h"
 
@@ -33,6 +34,10 @@ void UMVVMBlueprintViewModelContextWrapper::PostEditChangeProperty(struct FPrope
 	{
 		if (FMVVMBlueprintViewModelContext* ViewModelContextPtr = BlueprintViewPtr->FindViewModel(ViewModelId))
 		{
+			if (Wrapper.Resolver && (Wrapper.Resolver->GetOuter() == this || Wrapper.Resolver->GetOuter() == GetTransientPackage()))
+			{
+				Wrapper.Resolver->Rename(nullptr, BlueprintViewPtr);
+			}
 			*ViewModelContextPtr = Wrapper;
 		}
 	}
@@ -212,6 +217,7 @@ TSharedRef<SWidget> SMVVMViewModelPanel::MakeAddMenu()
 			SNew(SMVVMSelectViewModel, WidgetBlueprint)
 			.OnCancel(this, &SMVVMViewModelPanel::HandleCancelAddMenu)
 			.OnViewModelCommitted(this, &SMVVMViewModelPanel::HandleAddMenuViewModel)
+			.DisallowedClassFlags(CLASS_HideDropDown | CLASS_Hidden | CLASS_Deprecated | CLASS_NotPlaceable)
 		];
 }
 
