MVVM: Filter Viewmodel context creation type per class. Use MVVMAllowedContextCreationType or MVVMDisallowedContextCreationType to allow or disallow types.

UCLASS(meta=(MVVMAllowedContextCreationType="Manual|CreateInstance"))

diff --git a/Engine/Plugins/Runtime/ModelViewViewModel/Source/ModelViewViewModel/Public/MVVMViewModelBase.h b/Engine/Plugins/Runtime/ModelViewViewModel/Source/ModelViewViewModel/Public/MVVMViewModelBase.h
--- a/Engine/Plugins/Runtime/ModelViewViewModel/Source/ModelViewViewModel/Public/MVVMViewModelBase.h
+++ b/Engine/Plugins/Runtime/ModelViewViewModel/Source/ModelViewViewModel/Public/MVVMViewModelBase.h
@@ -17,6 +17,18 @@
 #define UE_MVVM_SET_PROPERTY_VALUE(MemberName, NewValue) \
 	SetPropertyValue(MemberName, NewValue, ThisClass::FFieldNotificationClassDescriptor::MemberName)
 
+namespace UMMVVMViewModelBase
+{
+	// valid keywords for the UCLASS macro
+	enum
+	{
+		// [ClassMetadata] Specifies which ContextCreationType is allowed for that viewmodel. ex: UCLASS(meta=(MVVMAllowedContextCreationType="Manual|CreateInstance"))
+		MVVMAllowedContextCreationType,
+		// [ClassMetadata] Specifies which ContextCreationType is disallowed for that viewmodel, all other type are allowed. ex: UCLASS(meta=(MVVMDisallowedContextCreationType="Manual|CreateInstance"))
+		MVVMDisallowedContextCreationType
+	};
+}
+
 
 /** Base class for MVVM viewmodel. */
 UCLASS(Blueprintable, Abstract, DisplayName="MVVM ViewModel")
diff --git a/Engine/Plugins/Runtime/ModelViewViewModel/Source/ModelViewViewModelBlueprint/Private/MVVMBlueprintViewModelContext.cpp b/Engine/Plugins/Runtime/ModelViewViewModel/Source/ModelViewViewModelBlueprint/Private/MVVMBlueprintViewModelContext.cpp
--- a/Engine/Plugins/Runtime/ModelViewViewModel/Source/ModelViewViewModelBlueprint/Private/MVVMBlueprintViewModelContext.cpp
+++ b/Engine/Plugins/Runtime/ModelViewViewModel/Source/ModelViewViewModelBlueprint/Private/MVVMBlueprintViewModelContext.cpp
@@ -1,6 +1,7 @@
 // Copyright Epic Games, Inc. All Rights Reserved.
 
 #include "MVVMBlueprintViewModelContext.h"
+#include "MVVMDeveloperProjectSettings.h"
 
 
 #include UE_INLINE_GENERATED_CPP_BY_NAME(MVVMBlueprintViewModelContext)
@@ -9,9 +10,16 @@ FMVVMBlueprintViewModelContext::FMVVMBlueprintViewModelContext(const UClass* InC
 {
 	if (InClass && InClass->ImplementsInterface(UNotifyFieldValueChanged::StaticClass()))
 	{
-		ViewModelContextId = FGuid::NewGuid();
-		NotifyFieldValueClass = const_cast<UClass*>(InClass);
-		ViewModelName = InViewModelName;
+#if WITH_EDITOR
+		TArray<EMVVMBlueprintViewModelContextCreationType> CreationTypes = UE::MVVM::GetAllowedContextCreationType(InClass);
+		if (CreationTypes.Num() > 0)
+#endif
+		{
+			CreationType = CreationTypes[0];
+			ViewModelContextId = FGuid::NewGuid();
+			NotifyFieldValueClass = const_cast<UClass*>(InClass);
+			ViewModelName = InViewModelName;
+		}
 	}
 }
 
@@ -21,3 +29,133 @@ FText FMVVMBlueprintViewModelContext::GetDisplayName() const
 	return FText::FromName(ViewModelName);
 }
 
+#if WITH_EDITOR
+namespace UE::MVVM
+{
+namespace Private
+{
+static const FName NAME_MVVMAllowedContextCreationType = TEXT("MVVMAllowedContextCreationType");
+static const FName NAME_MVVMDisallowedContextCreationType = TEXT("MVVMDisallowedContextCreationType");
+
+// Find the Allowed and Disallowed from first calss in the hierarchy that defines it.
+void GenerateAllowedList(const UClass* Class, TArray<FString>& AllowedStrings, TArray<FString>& DisallowedStrings)
+{
+	if (AllowedStrings.Num() == 0 && Class->HasMetaData(NAME_MVVMAllowedContextCreationType))
+	{
+		Class->GetMetaData(NAME_MVVMAllowedContextCreationType).ParseIntoArray(AllowedStrings, TEXT("|"));
+	}
+	if (AllowedStrings.Num() == 0 && Class->HasMetaData(NAME_MVVMDisallowedContextCreationType))
+	{
+		Class->GetMetaData(NAME_MVVMDisallowedContextCreationType).ParseIntoArray(DisallowedStrings, TEXT("|"));
+	}
+
+	UClass* ParentClass = Class->GetSuperClass();
+	if (ParentClass && AllowedStrings.Num() == 0 && DisallowedStrings.Num() == 0)
+	{
+		GenerateAllowedList(ParentClass, AllowedStrings, DisallowedStrings);
+	}
+}
+
+TArray<EMVVMBlueprintViewModelContextCreationType> GenerateAllowedList(const UClass* Class)
+{
+	const UEnum* EnumCreationType = StaticEnum<EMVVMBlueprintViewModelContextCreationType>();
+	TArray<FString> AllowedStrings;
+	TArray<FString> DisallowedStrings;
+	GenerateAllowedList(Class, AllowedStrings, DisallowedStrings);
+
+	if (AllowedStrings.Num() != 0)
+	{
+		for (FString& Trim : AllowedStrings)
+		{
+			Trim.TrimStartAndEndInline();
+		}
+		AllowedStrings.RemoveSwap(FString()); // Remove empty lines
+	}
+	if (DisallowedStrings.Num() != 0)
+	{
+		for (FString& Trim : DisallowedStrings)
+		{
+			Trim.TrimStartAndEndInline();
+		}
+		DisallowedStrings.RemoveSwap(FString()); // Remove empty lines
+	}
+
+	TArray<EMVVMBlueprintViewModelContextCreationType, TInlineAllocator<8>> FullList;
+	ensure(EnumCreationType->NumEnums() < 8); // just upgrade the FullList InlineAllocator to a bigger number.
+	for (int32 Index = 0; Index < EnumCreationType->NumEnums() - 1; ++Index)
+	{
+		EMVVMBlueprintViewModelContextCreationType CreationType = static_cast<EMVVMBlueprintViewModelContextCreationType>(EnumCreationType->GetValueByIndex(Index));
+
+		const bool bIsHidden = EnumCreationType->HasMetaData(TEXT("Hidden"), Index);
+		const bool bAllowed = GetDefault<UMVVMDeveloperProjectSettings>()->IsContextCreationTypeAllowed(CreationType);
+
+		if (!bIsHidden && bAllowed)
+		{
+			FullList.Add(CreationType);
+		}
+	}
+
+	TArray<EMVVMBlueprintViewModelContextCreationType> Result;
+	Result.Reserve(EnumCreationType->NumEnums());
+	if (AllowedStrings.Num() != 0 || DisallowedStrings.Num() != 0)
+	{
+		auto RunAllowed = [&Result, &AllowedStrings, &FullList, EnumCreationType]()
+		{
+			for (const FString& Allowed : AllowedStrings)
+			{
+				if (Allowed == TEXT("all"))
+				{
+					Result = FullList;
+					break;
+				}
+
+				int32 FoundIndex = EnumCreationType->GetIndexByName(*Allowed);
+				if (ensure(FoundIndex != INDEX_NONE))
+				{
+					Result.AddUnique(static_cast<EMVVMBlueprintViewModelContextCreationType>(EnumCreationType->GetValueByIndex(FoundIndex)));
+				}
+			}
+		};
+		auto RunDisallowed = [&Result, &DisallowedStrings, EnumCreationType]()
+		{
+			for (const FString& Disallowed : DisallowedStrings)
+			{
+				int32 FoundIndex = EnumCreationType->GetIndexByName(*Disallowed);
+				if (ensure(FoundIndex != INDEX_NONE))
+				{
+					Result.RemoveSingleSwap(static_cast<EMVVMBlueprintViewModelContextCreationType>(EnumCreationType->GetValueByIndex(FoundIndex)));
+				}
+			}
+		};
+
+		if (AllowedStrings.Num() != 0 && DisallowedStrings.Num() == 0)
+		{
+			RunAllowed();
+		}
+		if (AllowedStrings.Num() == 0 && DisallowedStrings.Num() != 0)
+		{
+			Result = FullList;
+			RunDisallowed();
+		}
+		else
+		{
+			ensureMsgf(false, TEXT("Disallowed and Allowed list found for class '%s'. Only provide one of the options."), *Class->GetFullName());
+			RunAllowed();
+			RunDisallowed();
+		}
+	}
+	else
+	{
+		Result = FullList;
+	}
+
+	return Result;
+}
+} //namespace Private
+
+TArray<EMVVMBlueprintViewModelContextCreationType> GetAllowedContextCreationType(const UClass* Class)
+{
+	return Private::GenerateAllowedList(Class);
+}
+}//namespace UE::MVVM
+#endif
\ No newline at end of file
diff --git a/Engine/Plugins/Runtime/ModelViewViewModel/Source/ModelViewViewModelBlueprint/Private/MVVMViewBlueprintCompiler.cpp b/Engine/Plugins/Runtime/ModelViewViewModel/Source/ModelViewViewModelBlueprint/Private/MVVMViewBlueprintCompiler.cpp
--- a/Engine/Plugins/Runtime/ModelViewViewModel/Source/ModelViewViewModelBlueprint/Private/MVVMViewBlueprintCompiler.cpp
+++ b/Engine/Plugins/Runtime/ModelViewViewModel/Source/ModelViewViewModelBlueprint/Private/MVVMViewBlueprintCompiler.cpp
@@ -749,11 +749,20 @@ bool FMVVMViewBlueprintCompiler::PreCompileSourceCreators(UWidgetBlueprintGenera
 
 			if (ViewModelContext.GetViewModelClass()->HasAllClassFlags(CLASS_Deprecated))
 			{
-				AddErrorForViewModel(ViewModelContext, FText::Format(LOCTEXT("ViewModelTypeDeprecated", "Viewmodel type '{0}' is deprecated and should not be used. Please update it in the View Models panel."),
+				AddErrorForViewModel(ViewModelContext, FText::Format(LOCTEXT("ViewModelTypeDeprecated", "Viewmodel class '{0}' is deprecated and should not be used. Please update it in the View Models panel."),
 					ViewModelContext.GetViewModelClass()->GetDisplayNameText()
 				));
 			}
 
+			if (!GetAllowedContextCreationType(ViewModelContext.GetViewModelClass()).Contains(ViewModelContext.CreationType))
+			{
+				AddErrorForViewModel(ViewModelContext, FText::Format(LOCTEXT("ViewModelContextCreationTypeInvalid", "Viewmodel '{0}' has an invalidate creation type. You can change it in the View Models panel."),
+					ViewModelContext.GetViewModelClass()->GetDisplayNameText()
+				));
+				bAreSourcesCreatorValid = false;
+				continue;
+			}
+
 			if (ViewModelContext.CreationType == EMVVMBlueprintViewModelContextCreationType::Manual)
 			{
 			}
@@ -761,7 +770,7 @@ bool FMVVMViewBlueprintCompiler::PreCompileSourceCreators(UWidgetBlueprintGenera
 			{
 				if (ViewModelContext.GetViewModelClass()->HasAllClassFlags(CLASS_Abstract))
 				{
-					AddErrorForViewModel(ViewModelContext, FText::Format(LOCTEXT("ViewModelTypeAbstract", "Viewmodel type '{0}' is abstract and can't be created. You can change it in the View Models panel."),
+					AddErrorForViewModel(ViewModelContext, FText::Format(LOCTEXT("ViewModelTypeAbstract", "Viewmodel class '{0}' is abstract and can't be created. You can change it in the View Models panel."),
 						ViewModelContext.GetViewModelClass()->GetDisplayNameText()
 					));
 					bAreSourcesCreatorValid = false;
diff --git a/Engine/Plugins/Runtime/ModelViewViewModel/Source/ModelViewViewModelBlueprint/Public/MVVMBlueprintViewModelContext.h b/Engine/Plugins/Runtime/ModelViewViewModel/Source/ModelViewViewModelBlueprint/Public/MVVMBlueprintViewModelContext.h
--- a/Engine/Plugins/Runtime/ModelViewViewModel/Source/ModelViewViewModelBlueprint/Public/MVVMBlueprintViewModelContext.h
+++ b/Engine/Plugins/Runtime/ModelViewViewModel/Source/ModelViewViewModelBlueprint/Public/MVVMBlueprintViewModelContext.h
@@ -27,6 +27,13 @@ enum class EMVVMBlueprintViewModelContextCreationType : uint8
 	Resolver,
 };
 
+namespace UE::MVVM
+{
+#if WITH_EDITOR
+	UE_NODISCARD MODELVIEWVIEWMODELBLUEPRINT_API TArray<EMVVMBlueprintViewModelContextCreationType> GetAllowedContextCreationType(const UClass* Class);
+#endif
+}
+
 /**
  *
  */
diff --git a/Engine/Plugins/Runtime/ModelViewViewModel/Source/ModelViewViewModelEditor/Private/Customizations/MVVMBlueprintViewModelContextCustomization.cpp b/Engine/Plugins/Runtime/ModelViewViewModel/Source/ModelViewViewModelEditor/Private/Customizations/MVVMBlueprintViewModelContextCustomization.cpp
--- a/Engine/Plugins/Runtime/ModelViewViewModel/Source/ModelViewViewModelEditor/Private/Customizations/MVVMBlueprintViewModelContextCustomization.cpp
+++ b/Engine/Plugins/Runtime/ModelViewViewModel/Source/ModelViewViewModelEditor/Private/Customizations/MVVMBlueprintViewModelContextCustomization.cpp
@@ -2,6 +2,7 @@
 
 #include "MVVMBlueprintViewModelContextCustomization.h"
 
+#include "Framework/MultiBox/MultiBoxBuilder.h"
 #include "IPropertyTypeCustomization.h"
 #include "PropertyHandle.h"
 #include "IDetailChildrenBuilder.h"
@@ -11,6 +12,7 @@
 #include "Bindings/MVVMBindingHelper.h"
 #include "Features/IModularFeatures.h"
 #include "Widgets/SBoxPanel.h"
+#include "Widgets/Input/SComboButton.h"
 #include "Widgets/SMVVMViewModelPanel.h"
 
 #define LOCTEXT_NAMESPACE "BlueprintViewModelContextDetailCustomization"
@@ -114,9 +116,6 @@ FBlueprintViewModelContextDetailCustomization::FBlueprintViewModelContextDetailC
 	: WidgetBlueprintEditor(InEditor)
 {}
 
-FBlueprintViewModelContextDetailCustomization::~FBlueprintViewModelContextDetailCustomization()
-{
-}
 
 void FBlueprintViewModelContextDetailCustomization::CustomizeChildren(TSharedRef<IPropertyHandle> PropertyHandle, IDetailChildrenBuilder& ChildBuilder, IPropertyTypeCustomizationUtils& CustomizationUtils)
 {
@@ -136,7 +135,12 @@ void FBlueprintViewModelContextDetailCustomization::CustomizeChildren(TSharedRef
 			UObject* Object = nullptr;
 			if (ChildHandle->GetValue(Object) == FPropertyAccess::Success)
 			{
-				PropertyAccessEditor.ClassToLookFor = Cast<UClass>(Object);
+				UClass* ViewModelClass = Cast<UClass>(Object);
+				if (ViewModelClass)
+				{
+					AllowedCreationTypes = GetAllowedContextCreationType(ViewModelClass);
+				}
+				PropertyAccessEditor.ClassToLookFor = ViewModelClass;
 			}
 			ChildHandle->SetOnPropertyValueChanged(FSimpleDelegate::CreateSP(this, &FBlueprintViewModelContextDetailCustomization::HandleClassChanged));
 		}
@@ -152,6 +156,12 @@ void FBlueprintViewModelContextDetailCustomization::CustomizeChildren(TSharedRef
 			ensure(CastField<FStrProperty>(ChildHandle->GetProperty()));
 			PropertyPathHandle = ChildHandle;
 		}
+
+		if (PropertyName == GET_MEMBER_NAME_CHECKED(FMVVMBlueprintViewModelContext, CreationType))
+		{
+			ensure(CastField<FEnumProperty>(ChildHandle->GetProperty()));
+			CreationTypeHandle = ChildHandle;
+		}
 	}
 
 	if (NotifyFieldValueClassHandle)
@@ -191,6 +201,30 @@ void FBlueprintViewModelContextDetailCustomization::CustomizeChildren(TSharedRef
 					];
 				}
 			}
+			else if (PropertyName == GET_MEMBER_NAME_CHECKED(FMVVMBlueprintViewModelContext, CreationType))
+			{
+				IDetailPropertyRow& PropertyRow = ChildBuilder.AddProperty(ChildHandle.ToSharedRef());
+
+				TSharedPtr<SWidget> NameWidget, ValueWidget;
+				PropertyRow.GetDefaultWidgets(NameWidget, ValueWidget);
+				PropertyRow.CustomWidget()
+					.NameContent()
+					[
+						NameWidget.ToSharedRef()
+					]
+					.ValueContent()
+					[
+						SNew(SComboButton)
+						.ContentPadding(FMargin(4.f, 0.f))
+						.OnGetMenuContent(this, &FBlueprintViewModelContextDetailCustomization::CreateExecutionTypeMenuContent)
+						.ButtonContent()
+						[
+							SNew(STextBlock)
+							.Text(this, &FBlueprintViewModelContextDetailCustomization::GetCreationTypeValue)
+							.ToolTipText(this, &FBlueprintViewModelContextDetailCustomization::GetExecutionTypeValueToolTip)
+						]
+					];
+			}
 			else
 			{
 				ChildBuilder.AddProperty(ChildHandle.ToSharedRef());
@@ -202,14 +236,64 @@ void FBlueprintViewModelContextDetailCustomization::CustomizeChildren(TSharedRef
 void FBlueprintViewModelContextDetailCustomization::HandleClassChanged()
 {
 	UObject* Object = nullptr;
+	AllowedCreationTypes.Reset();
+	PropertyAccessEditor.ClassToLookFor.Reset();
 	if (NotifyFieldValueClassHandle->GetValue(Object) == FPropertyAccess::Success)
 	{
-		PropertyAccessEditor.ClassToLookFor = Cast<UClass>(Object);
+		if (UClass* ViewModelClass = Cast<UClass>(Object))
+		{
+			PropertyAccessEditor.ClassToLookFor = ViewModelClass;
+			AllowedCreationTypes = GetAllowedContextCreationType(ViewModelClass);
+		}
+	}
+}
+
+TSharedRef<SWidget> FBlueprintViewModelContextDetailCustomization::CreateExecutionTypeMenuContent()
+{
+	const bool bCloseAfterSelection = true;
+	FMenuBuilder MenuBuilder(bCloseAfterSelection, nullptr, nullptr, true);
+
+	const UEnum* EnumCreationType = StaticEnum<EMVVMBlueprintViewModelContextCreationType>();
+	for (EMVVMBlueprintViewModelContextCreationType Type : AllowedCreationTypes)
+	{
+		int32 Index = EnumCreationType->GetIndexByValue((int64)Type);
+		MenuBuilder.AddMenuEntry(
+			EnumCreationType->GetDisplayNameTextByIndex(Index),
+			EnumCreationType->GetToolTipTextByIndex(Index),
+			FSlateIcon(),
+			FUIAction
+			(
+				FExecuteAction::CreateLambda([this, Type]()
+				{
+					uint8 Value = static_cast<uint8>(Type);
+					CreationTypeHandle->SetValue(Value);
+				})
+			)
+		);
 	}
-	else
+
+	return MenuBuilder.MakeWidget();
+}
+
+FText FBlueprintViewModelContextDetailCustomization::GetCreationTypeValue() const
+{
+	uint8 Value = 0;
+	if (CreationTypeHandle->GetValue(Value) == FPropertyAccess::Result::Success)
+	{
+		return StaticEnum<EMVVMBlueprintViewModelContextCreationType>()->GetDisplayNameTextByValue((int64)Value);
+	}
+	return FText::GetEmpty();
+}
+
+FText FBlueprintViewModelContextDetailCustomization::GetExecutionTypeValueToolTip() const
+{
+	uint8 Value = 0;
+	if (CreationTypeHandle->GetValue(Value) == FPropertyAccess::Result::Success)
 	{
-		PropertyAccessEditor.ClassToLookFor.Reset();
+		UEnum* EnumCreationType = StaticEnum<EMVVMBlueprintViewModelContextCreationType>();
+		return EnumCreationType->GetToolTipTextByIndex(EnumCreationType->GetIndexByValue((int64)Value));
 	}
+	return FText::GetEmpty();
 }
 
 } // namespace UE::MVVM
diff --git a/Engine/Plugins/Runtime/ModelViewViewModel/Source/ModelViewViewModelEditor/Private/Customizations/MVVMBlueprintViewModelContextCustomization.h b/Engine/Plugins/Runtime/ModelViewViewModel/Source/ModelViewViewModelEditor/Private/Customizations/MVVMBlueprintViewModelContextCustomization.h
--- a/Engine/Plugins/Runtime/ModelViewViewModel/Source/ModelViewViewModelEditor/Private/Customizations/MVVMBlueprintViewModelContextCustomization.h
+++ b/Engine/Plugins/Runtime/ModelViewViewModel/Source/ModelViewViewModelEditor/Private/Customizations/MVVMBlueprintViewModelContextCustomization.h
@@ -8,7 +8,7 @@
 class FWidgetBlueprintEditor;
 struct FBindingChainElement;
 class SWidget;
-
+enum class EMVVMBlueprintViewModelContextCreationType : uint8;
 
 namespace UE::MVVM
 {
@@ -36,7 +36,6 @@ class FBlueprintViewModelContextDetailCustomization : public IPropertyTypeCustom
 {
 public:
 	FBlueprintViewModelContextDetailCustomization(TWeakPtr<FWidgetBlueprintEditor> InEditor);
-	~FBlueprintViewModelContextDetailCustomization();
 
 	//~ IDetailCustomization interface
 	virtual void CustomizeHeader(TSharedRef<IPropertyHandle> PropertyHandle, FDetailWidgetRow& HeaderRow, IPropertyTypeCustomizationUtils& CustomizationUtils) override {}
@@ -48,12 +47,17 @@ public:
 	}
 
 private:
+	TSharedRef<SWidget> CreateExecutionTypeMenuContent();
+	FText GetCreationTypeValue() const;
+	FText GetExecutionTypeValueToolTip() const;
 	void HandleClassChanged();
 
 	FViewModelPropertyAccessEditor PropertyAccessEditor;
 	TWeakPtr<FWidgetBlueprintEditor> WidgetBlueprintEditor;
 	TSharedPtr<IPropertyHandle> NotifyFieldValueClassHandle;
 	TSharedPtr<IPropertyHandle> PropertyPathHandle;
+	TSharedPtr<IPropertyHandle> CreationTypeHandle;
+	TArray<EMVVMBlueprintViewModelContextCreationType> AllowedCreationTypes;
 };
 
 } // namespace UE::MVVM
