MVVM: Rename ViewModelResolver to ViewModelContextResolver

diff --git a/Engine/Plugins/Runtime/ModelViewViewModel/Source/ModelViewViewModel/Private/View/MVVMViewClass.cpp b/Engine/Plugins/Runtime/ModelViewViewModel/Source/ModelViewViewModel/Private/View/MVVMViewClass.cpp
--- a/Engine/Plugins/Runtime/ModelViewViewModel/Source/ModelViewViewModel/Private/View/MVVMViewClass.cpp
+++ b/Engine/Plugins/Runtime/ModelViewViewModel/Source/ModelViewViewModel/Private/View/MVVMViewClass.cpp
@@ -3,7 +3,7 @@
 #include "View/MVVMViewClass.h"
 #include "Types/MVVMFieldContext.h"
 #include "View/MVVMView.h"
-#include "View/MVVMViewModelResolver.h"
+#include "View/MVVMViewModelContextResolver.h"
 
 #include "Bindings/MVVMFieldPathHelper.h"
 #include "Blueprint/UserWidget.h"
@@ -80,7 +80,7 @@ FMVVMViewClass_SourceCreator FMVVMViewClass_SourceCreator::MakeGlobalContext(FNa
 	return Result;
 }
 
-FMVVMViewClass_SourceCreator FMVVMViewClass_SourceCreator::MakeResolver(FName InName, UClass* InNotifyFieldValueChangedClass, UMVVMViewModelResolver* InResolver, bool bOptional)
+FMVVMViewClass_SourceCreator FMVVMViewClass_SourceCreator::MakeResolver(FName InName, UClass* InNotifyFieldValueChangedClass, UMVVMViewModelContextResolver* InResolver, bool bOptional)
 {
 	FMVVMViewClass_SourceCreator Result;
 	if (ensure(InResolver && InNotifyFieldValueChangedClass && InNotifyFieldValueChangedClass->ImplementsInterface(UNotifyFieldValueChanged::StaticClass())))
diff --git a/Engine/Plugins/Runtime/ModelViewViewModel/Source/ModelViewViewModel/Public/View/MVVMViewClass.h b/Engine/Plugins/Runtime/ModelViewViewModel/Source/ModelViewViewModel/Public/View/MVVMViewClass.h
--- a/Engine/Plugins/Runtime/ModelViewViewModel/Source/ModelViewViewModel/Public/View/MVVMViewClass.h
+++ b/Engine/Plugins/Runtime/ModelViewViewModel/Source/ModelViewViewModel/Public/View/MVVMViewClass.h
@@ -16,7 +16,7 @@ class UMVVMUserWidgetBinding;
 class UMVVMView;
 class UMVVMViewClass;
 class UMVVMViewModelBlueprintExtension;
-class UMVVMViewModelResolver;
+class UMVVMViewModelContextResolver;
 class UUserWidget;
 
 namespace UE::MVVM::Private
@@ -39,7 +39,7 @@ public:
 	MODELVIEWVIEWMODEL_API static FMVVMViewClass_SourceCreator MakeInstance(FName Name, UClass* NotifyFieldValueChangedClass);
 	MODELVIEWVIEWMODEL_API static FMVVMViewClass_SourceCreator MakeFieldPath(FName Name, UClass* NotifyFieldValueChangedClass, FMVVMVCompiledFieldPath FieldPath, bool bOptional);
 	MODELVIEWVIEWMODEL_API static FMVVMViewClass_SourceCreator MakeGlobalContext(FName Name, FMVVMViewModelContext Context, bool bOptional);
-	MODELVIEWVIEWMODEL_API static FMVVMViewClass_SourceCreator MakeResolver(FName Name, UClass* InNotifyFieldValueChangedClass, UMVVMViewModelResolver* Source, bool bOptional);
+	MODELVIEWVIEWMODEL_API static FMVVMViewClass_SourceCreator MakeResolver(FName Name, UClass* InNotifyFieldValueChangedClass, UMVVMViewModelContextResolver* Source, bool bOptional);
 
 	UObject* CreateInstance(const UMVVMViewClass* ViewClass, UMVVMView* View, UUserWidget* UserWidget) const;
 
@@ -60,7 +60,7 @@ private:
 
 	/** The resolver to fetch the source at runtime. */
 	UPROPERTY(Instanced)
-	TObjectPtr<UMVVMViewModelResolver> Resolver = nullptr;
+	TObjectPtr<UMVVMViewModelContextResolver> Resolver = nullptr;
 
 	/** Info to find the ViewModel instance at runtime. */
 	UPROPERTY()
diff --git a/Engine/Plugins/Runtime/ModelViewViewModel/Source/ModelViewViewModel/Public/View/MVVMViewModelContextResolver.h b/Engine/Plugins/Runtime/ModelViewViewModel/Source/ModelViewViewModel/Public/View/MVVMViewModelContextResolver.h
new file mode 100644
--- /dev/null
+++ b/Engine/Plugins/Runtime/ModelViewViewModel/Source/ModelViewViewModel/Public/View/MVVMViewModelContextResolver.h
@@ -0,0 +1,30 @@
+// Copyright Epic Games, Inc. All Rights Reserved.
+
+#pragma once
+
+#include "UObject/Object.h"
+#include "FieldNotification/IFieldValueChanged.h"
+
+#include "MVVMViewModelContextResolver.generated.h"
+
+class UMVVMView;
+class UMVVMViewClass;
+class UUserWidget;
+
+/**
+ * Shared data to find or create a ViewModel at runtime.
+ */
+UCLASS(Abstract, Blueprintable, EditInlineNew, DisplayName = "Viewmode Resolver")
+class MODELVIEWVIEWMODEL_API UMVVMViewModelContextResolver : public UObject
+{
+	GENERATED_BODY()
+
+public:
+	virtual UObject* CreateInstance(const UClass* ExpectedType, const UUserWidget* UserWidget, const UMVVMView* View) const
+	{
+		return K2_CreateInstance(ExpectedType, UserWidget).GetObject();
+	}
+
+	UFUNCTION(BlueprintImplementableEvent, Category="Viewmodel", DisplayName="Create Instance")
+	TScriptInterface<INotifyFieldValueChanged> K2_CreateInstance(const UClass* ExpectedType, const UUserWidget* UserWidget) const;
+};
diff --git a/Engine/Plugins/Runtime/ModelViewViewModel/Source/ModelViewViewModel/Public/View/MVVMViewModelResolver.h b/Engine/Plugins/Runtime/ModelViewViewModel/Source/ModelViewViewModel/Public/View/MVVMViewModelResolver.h
deleted file mode 100644
diff --git a/Engine/Plugins/Runtime/ModelViewViewModel/Source/ModelViewViewModelBlueprint/Private/MVVMViewBlueprintCompiler.cpp b/Engine/Plugins/Runtime/ModelViewViewModel/Source/ModelViewViewModelBlueprint/Private/MVVMViewBlueprintCompiler.cpp
--- a/Engine/Plugins/Runtime/ModelViewViewModel/Source/ModelViewViewModelBlueprint/Private/MVVMViewBlueprintCompiler.cpp
+++ b/Engine/Plugins/Runtime/ModelViewViewModel/Source/ModelViewViewModelBlueprint/Private/MVVMViewBlueprintCompiler.cpp
@@ -12,7 +12,7 @@
 #include "Templates/ValueOrError.h"
 #include "Types/MVVMBindingName.h"
 #include "View/MVVMViewClass.h"
-#include "View/MVVMViewModelResolver.h"
+#include "View/MVVMViewModelContextResolver.h"
 #include "UObject/LinkerLoad.h"
 
 #define LOCTEXT_NAMESPACE "MVVMViewBlueprintCompiler"
@@ -895,7 +895,7 @@ bool FMVVMViewBlueprintCompiler::CompileSourceCreators(const FCompiledBindingLib
 			}
 			else if (ViewModelContext.CreationType == EMVVMBlueprintViewModelContextCreationType::Resolver)
 			{
-				UMVVMViewModelResolver* Resolver = DuplicateObject(ViewModelContext.Resolver.Get(), ViewExtension);
+				UMVVMViewModelContextResolver* Resolver = DuplicateObject(ViewModelContext.Resolver.Get(), ViewExtension);
 				if (!Resolver)
 				{
 					AddErrorForViewModel(ViewModelContext, LOCTEXT("ViewmodelFailedResolverDuplicate", "Internal error. The resolver could not be dupliated."));
diff --git a/Engine/Plugins/Runtime/ModelViewViewModel/Source/ModelViewViewModelBlueprint/Public/MVVMBlueprintViewModelContext.h b/Engine/Plugins/Runtime/ModelViewViewModel/Source/ModelViewViewModelBlueprint/Public/MVVMBlueprintViewModelContext.h
--- a/Engine/Plugins/Runtime/ModelViewViewModel/Source/ModelViewViewModelBlueprint/Public/MVVMBlueprintViewModelContext.h
+++ b/Engine/Plugins/Runtime/ModelViewViewModel/Source/ModelViewViewModelBlueprint/Public/MVVMBlueprintViewModelContext.h
@@ -7,7 +7,7 @@
 
 #include "MVVMBlueprintViewModelContext.generated.h"
 
-class UMVVMViewModelResolver;
+class UMVVMViewModelContextResolver;
 
 /**
  *
@@ -119,7 +119,7 @@ public:
 	FString ViewModelPropertyPath;
 
 	UPROPERTY(EditAnywhere, Category = "Viewmodel", Instanced, AdvancedDisplay, meta = (EditInline, EditCondition = "CreationType == EMVVMBlueprintViewModelContextCreationType::Resolver", EditConditionHides))
-	TObjectPtr<UMVVMViewModelResolver> Resolver = nullptr;
+	TObjectPtr<UMVVMViewModelContextResolver> Resolver = nullptr;
 
 	/**
 	 * Generate a setter function for this viewmodel.
diff --git a/Engine/Plugins/Runtime/ModelViewViewModel/Source/ModelViewViewModelEditor/Private/Widgets/SMVVMViewModelPanel.cpp b/Engine/Plugins/Runtime/ModelViewViewModel/Source/ModelViewViewModelEditor/Private/Widgets/SMVVMViewModelPanel.cpp
--- a/Engine/Plugins/Runtime/ModelViewViewModel/Source/ModelViewViewModelEditor/Private/Widgets/SMVVMViewModelPanel.cpp
+++ b/Engine/Plugins/Runtime/ModelViewViewModel/Source/ModelViewViewModelEditor/Private/Widgets/SMVVMViewModelPanel.cpp
@@ -7,7 +7,7 @@
 #include "MVVMBlueprintView.h"
 #include "MVVMEditorSubsystem.h"
 #include "PropertyEditorModule.h"
-#include "View/MVVMViewModelResolver.h"
+#include "View/MVVMViewModelContextResolver.h"
 #include "WidgetBlueprint.h"
 #include "WidgetBlueprintEditor.h"
 
