MVVM: Fix bad else case when the allowed and disallowed list is used.

diff --git a/Engine/Plugins/Runtime/ModelViewViewModel/Source/ModelViewViewModelBlueprint/Private/MVVMBlueprintViewModelContext.cpp b/Engine/Plugins/Runtime/ModelViewViewModel/Source/ModelViewViewModelBlueprint/Private/MVVMBlueprintViewModelContext.cpp
--- a/Engine/Plugins/Runtime/ModelViewViewModel/Source/ModelViewViewModelBlueprint/Private/MVVMBlueprintViewModelContext.cpp
+++ b/Engine/Plugins/Runtime/ModelViewViewModel/Source/ModelViewViewModelBlueprint/Private/MVVMBlueprintViewModelContext.cpp
@@ -132,7 +132,7 @@ TArray<EMVVMBlueprintViewModelContextCreationType> GenerateAllowedList(const UCl
 		{
 			RunAllowed();
 		}
-		if (AllowedStrings.Num() == 0 && DisallowedStrings.Num() != 0)
+		else if (AllowedStrings.Num() == 0 && DisallowedStrings.Num() != 0)
 		{
 			Result = FullList;
 			RunDisallowed();
