MVVM: Veriry the allowed list when picking a property/function from the field selector.

diff --git a/Engine/Plugins/Runtime/ModelViewViewModel/Source/ModelViewViewModelBlueprint/ModelViewViewModelBlueprint.Build.cs b/Engine/Plugins/Runtime/ModelViewViewModel/Source/ModelViewViewModelBlueprint/ModelViewViewModelBlueprint.Build.cs
--- a/Engine/Plugins/Runtime/ModelViewViewModel/Source/ModelViewViewModelBlueprint/ModelViewViewModelBlueprint.Build.cs
+++ b/Engine/Plugins/Runtime/ModelViewViewModel/Source/ModelViewViewModelBlueprint/ModelViewViewModelBlueprint.Build.cs
@@ -23,6 +23,7 @@ public class ModelViewViewModelBlueprint : ModuleRules
 			{
 				"BlueprintGraph",
 				"KismetCompiler",
+				"PropertyEditor",
 				"SlateCore",
 				"Slate",
 				"UMG",
diff --git a/Engine/Plugins/Runtime/ModelViewViewModel/Source/ModelViewViewModelBlueprint/Private/MVVMDeveloperProjectSettings.cpp b/Engine/Plugins/Runtime/ModelViewViewModel/Source/ModelViewViewModelBlueprint/Private/MVVMDeveloperProjectSettings.cpp
--- a/Engine/Plugins/Runtime/ModelViewViewModel/Source/ModelViewViewModelBlueprint/Private/MVVMDeveloperProjectSettings.cpp
+++ b/Engine/Plugins/Runtime/ModelViewViewModel/Source/ModelViewViewModelBlueprint/Private/MVVMDeveloperProjectSettings.cpp
@@ -6,6 +6,9 @@
 #include "MVVMBlueprintViewModelContext.h"
 #include "Types/MVVMExecutionMode.h"
 
+#include "BlueprintEditorSettings.h"
+#include "PropertyEditorPermissionList.h"
+
 #define LOCTEXT_NAMESPACE "MVVMDeveloperProjectSettings"
 
 UMVVMDeveloperProjectSettings::UMVVMDeveloperProjectSettings()
@@ -35,6 +38,11 @@ FText UMVVMDeveloperProjectSettings::GetSectionText() const
 bool UMVVMDeveloperProjectSettings::IsPropertyAllowed(const FProperty* Property) const
 {
 	check(Property);
+	if (!FPropertyEditorPermissionList::Get().DoesPropertyPassFilter(Property->GetOwnerStruct(), Property->GetFName()))
+	{
+		return false;
+	}
+
 	TStringBuilder<256> StringBuilder;
 	Property->GetOwnerClass()->GetPathName(nullptr, StringBuilder);
 	FSoftClassPath StructPath;
@@ -49,6 +57,11 @@ bool UMVVMDeveloperProjectSettings::IsPropertyAllowed(const FProperty* Property)
 bool UMVVMDeveloperProjectSettings::IsFunctionAllowed(const UFunction* Function) const
 {
 	check(Function);
+	if (!GetMutableDefault<UBlueprintEditorSettings>()->GetFunctionPermissions().PassesFilter(Function->GetFullName()))
+	{
+		return false;
+	}
+
 	TStringBuilder<256> StringBuilder;
 	Function->GetOwnerClass()->GetPathName(nullptr, StringBuilder);
 	FSoftClassPath StructPath;
diff --git a/Engine/Plugins/Runtime/ModelViewViewModel/Source/ModelViewViewModelEditor/Private/Widgets/SMVVMViewModelBindingListWidget.cpp b/Engine/Plugins/Runtime/ModelViewViewModel/Source/ModelViewViewModelEditor/Private/Widgets/SMVVMViewModelBindingListWidget.cpp
--- a/Engine/Plugins/Runtime/ModelViewViewModel/Source/ModelViewViewModelEditor/Private/Widgets/SMVVMViewModelBindingListWidget.cpp
+++ b/Engine/Plugins/Runtime/ModelViewViewModel/Source/ModelViewViewModelEditor/Private/Widgets/SMVVMViewModelBindingListWidget.cpp
@@ -4,9 +4,11 @@
 
 #include "Bindings/MVVMBindingHelper.h"
 #include "Blueprint/WidgetTree.h"
+#include "BlueprintEditorSettings.h"
 #include "MVVMBlueprintViewModelContext.h"
 #include "MVVMDeveloperProjectSettings.h"
 #include "MVVMSubsystem.h"
+#include "PropertyEditorPermissionList.h"
 #include "Styling/MVVMEditorStyle.h"
 #include "Types/MVVMAvailableBinding.h"
 #include "Types/MVVMBindingSource.h"
@@ -233,12 +235,37 @@ FFieldExpander_Bindable::FFieldExpander_Bindable()
 	SetExpandFunction(UE::PropertyViewer::FFieldExpander_Default::EFunctionExpand::FunctionProperties);
 }
 
+TOptional<const UClass*> FFieldExpander_Bindable::CanExpandObject(const FObjectPropertyBase* Property, const UObject* Instance) const
+{
+	TOptional<const UClass*> Result = UE::PropertyViewer::FFieldExpander_Default::CanExpandObject(Property, Instance);
+	if (Result.IsSet() && Result.GetValue())
+	{
+		if (GetDefault<UBlueprintEditorSettings>()->IsClassAllowedOnPin(Result.GetValue()))
+		{
+			return Result;
+		}
+	}
+	return TOptional<const UClass*>();
+}
+
+bool FFieldExpander_Bindable::CanExpandScriptStruct(const FStructProperty* StructProperty) const
+{
+	if (UE::PropertyViewer::FFieldExpander_Default::CanExpandScriptStruct(StructProperty))
+	{
+		return GetMutableDefault<UBlueprintEditorSettings>()->GetStructPermissions().PassesFilter(StructProperty->GetFullName());
+	}
+	return false;
+}
+
 TOptional<const UStruct*> FFieldExpander_Bindable::GetExpandedFunction(const UFunction* Function) const
 {
 	const FProperty* ReturnProperty = Function ? BindingHelper::GetReturnProperty(Function) : nullptr;
 	if (const FObjectPropertyBase* ObjectProperty = CastField<const FObjectPropertyBase>(ReturnProperty))
 	{
-		return ObjectProperty->PropertyClass;
+		if (GetDefault<UBlueprintEditorSettings>()->IsClassAllowedOnPin(ObjectProperty->PropertyClass))
+		{
+			return ObjectProperty->PropertyClass;
+		}
 	}
 	return TOptional<const UStruct*>();
 }
diff --git a/Engine/Plugins/Runtime/ModelViewViewModel/Source/ModelViewViewModelEditor/Private/Widgets/SMVVMViewModelBindingListWidget.h b/Engine/Plugins/Runtime/ModelViewViewModel/Source/ModelViewViewModelEditor/Private/Widgets/SMVVMViewModelBindingListWidget.h
--- a/Engine/Plugins/Runtime/ModelViewViewModel/Source/ModelViewViewModelEditor/Private/Widgets/SMVVMViewModelBindingListWidget.h
+++ b/Engine/Plugins/Runtime/ModelViewViewModel/Source/ModelViewViewModelEditor/Private/Widgets/SMVVMViewModelBindingListWidget.h
@@ -57,6 +57,8 @@ namespace UE::MVVM
 	{
 	public:
 		FFieldExpander_Bindable();
+		virtual TOptional<const UClass*> CanExpandObject(const FObjectPropertyBase* Property, const UObject* Instance) const override;
+		virtual bool CanExpandScriptStruct(const FStructProperty*) const override;
 		virtual TOptional<const UStruct*> GetExpandedFunction(const UFunction* Function) const override;
 	};
 
