MVVM: Delayed bindnig were ticked by a world. If the widget changed world, the binding might not be ticked anymore.

diff --git a/Engine/Plugins/Runtime/ModelViewViewModel/Source/ModelViewViewModel/Private/View/MVVMBindingSubsystem.cpp b/Engine/Plugins/Runtime/ModelViewViewModel/Source/ModelViewViewModel/Private/View/MVVMBindingSubsystem.cpp
new file mode 100644
--- /dev/null
+++ b/Engine/Plugins/Runtime/ModelViewViewModel/Source/ModelViewViewModel/Private/View/MVVMBindingSubsystem.cpp
@@ -0,0 +1,123 @@
+// Copyright Epic Games, Inc. All Rights Reserved.
+
+#include "View/MVVMBindingSubsystem.h"
+
+#include "Framework/Application/SlateApplication.h"
+#include "Misc/MemStack.h"
+#include "SlateGlobals.h"
+#include "Stats/Stats2.h"
+#include "View/MVVMView.h"
+
+#include UE_INLINE_GENERATED_CPP_BY_NAME(MVVMBindingSubsystem)
+
+
+DECLARE_CYCLE_STAT(TEXT("MVVM Bindings"), STAT_MVVMBindingTick, STATGROUP_Slate);
+
+void UMVVMBindingSubsystem::Initialize(FSubsystemCollectionBase& Collection)
+{
+	Super::Initialize(Collection);
+
+	if (FSlateApplication::IsInitialized())
+	{
+		FSlateApplication::Get().OnPreTick().AddUObject(this, &UMVVMBindingSubsystem::HandlePreTick);
+	}
+}
+
+void UMVVMBindingSubsystem::Deinitialize()
+{
+	if (FSlateApplication::IsInitialized())
+	{
+		FSlateApplication::Get().OnPreTick().RemoveAll(this);
+	}
+	Super::Deinitialize();
+}
+
+void UMVVMBindingSubsystem::HandlePreTick(float DeltaTime)
+{
+	SCOPE_CYCLE_COUNTER(STAT_MVVMBindingTick);
+
+	if (EveryTickBindings.Num() > 0)
+	{
+		FMemMark Mark(FMemStack::Get());
+		TArray<const UMVVMView*, TMemStackAllocator<>> ToTick;
+		ToTick.Reserve(EveryTickBindings.Num());
+		for (TWeakObjectPtr<const UMVVMView> View : EveryTickBindings)
+		{
+			if (const UMVVMView* ViewPtr = View.Get())
+			{
+				ToTick.Add(ViewPtr);
+			}
+		}
+
+		for (const UMVVMView* ViewPtr : ToTick)
+		{
+			ViewPtr->ExecuteEveryTickBindings();
+		}
+	}
+
+	if (DelayedBindings.Num() > 0)
+	{
+		FDelayedMap AllDelayedBindingsExecutedThisFrame;
+		AllDelayedBindingsExecutedThisFrame.Reserve(DelayedBindings.Num());
+
+		FDelayedMap DelayedBindingsWhileTicking = MoveTemp(DelayedBindings);
+		DelayedBindings = FDelayedMap();
+
+		do 
+		{
+			for (const auto& DelayedBindingsPair : DelayedBindingsWhileTicking)
+			{
+				if (const UMVVMView* View = DelayedBindingsPair.Key.Get())
+				{
+					ensure(DelayedBindingsPair.Value.Num() > 0);
+					FDelayedBindingList& AllBindingList = AllDelayedBindingsExecutedThisFrame.FindOrAdd(DelayedBindingsPair.Key);
+					for (const FMVVMViewDelayedBinding& DelayedBinding : DelayedBindingsPair.Value)
+					{
+						View->ExecuteDelayedBinding(DelayedBinding);
+						AllBindingList.AddUnique(DelayedBinding);
+					}
+				}
+			}
+
+			DelayedBindingsWhileTicking.Reset();
+
+			// Test new binding added while executing the current binding list
+			for (const auto& DelayedBindingPair : DelayedBindings)
+			{
+				if (DelayedBindingPair.Key.Get())
+				{
+					if (FDelayedBindingList* AllDelayedBindingsListPtr = AllDelayedBindingsExecutedThisFrame.Find(DelayedBindingPair.Key))
+					{
+						for (const FMVVMViewDelayedBinding& NewCompiledBinding : DelayedBindingPair.Value)
+						{
+							if (!AllDelayedBindingsListPtr->Find(NewCompiledBinding))
+							{
+								// It was not executed. Add it to be executed this frame.
+								DelayedBindingsWhileTicking.FindOrAdd(DelayedBindingPair.Key).AddUnique(NewCompiledBinding);
+							}
+						}
+					}
+				}
+			}
+		} while (DelayedBindingsWhileTicking.Num() > 0);
+	}
+}
+
+void UMVVMBindingSubsystem::AddViewWithEveryTickBinding(const UMVVMView* InView)
+{
+	check(!EveryTickBindings.Contains(InView));
+	ensureMsgf(FSlateApplication::IsInitialized(), TEXT("The Slate Application is not initialized. This is probably because you are running a server. The Delayed and Tick binding will not execute."));
+
+	EveryTickBindings.Add(InView);
+}
+
+void UMVVMBindingSubsystem::RemoveViewWithEveryTickBinding(const UMVVMView* InView)
+{
+	EveryTickBindings.RemoveSingleSwap(InView);
+}
+
+void UMVVMBindingSubsystem::AddDelayedBinding(const UMVVMView* View, FMVVMViewDelayedBinding InCompiledBinding)
+{
+	ensureMsgf(FSlateApplication::IsInitialized(), TEXT("The Slate Application is not initialized. This is probably because you are running a server. The Delayed and Tick binding will not execute."));
+	DelayedBindings.FindOrAdd(View).AddUnique(InCompiledBinding);
+}
diff --git a/Engine/Plugins/Runtime/ModelViewViewModel/Source/ModelViewViewModel/Private/View/MVVMBindingSubsystem.h b/Engine/Plugins/Runtime/ModelViewViewModel/Source/ModelViewViewModel/Private/View/MVVMBindingSubsystem.h
new file mode 100644
--- /dev/null
+++ b/Engine/Plugins/Runtime/ModelViewViewModel/Source/ModelViewViewModel/Private/View/MVVMBindingSubsystem.h
@@ -0,0 +1,60 @@
+// Copyright Epic Games, Inc. All Rights Reserved.
+
+#pragma once
+
+#include "Subsystems/EngineSubsystem.h"
+
+#include "MVVMBindingSubsystem.generated.h"
+
+class UMVVMView;
+
+/** */
+struct FMVVMViewDelayedBinding
+{
+	FMVVMViewDelayedBinding(int32 InCompiledBindingIndex)
+		: CompiledBindingIndex(InCompiledBindingIndex)
+	{
+	}
+
+	int32 GetCompiledBindingIndex() const
+	{
+		return CompiledBindingIndex;
+	}
+
+	bool operator==(const FMVVMViewDelayedBinding& Other) const
+	{
+		return CompiledBindingIndex == Other.CompiledBindingIndex;
+	}
+
+	bool operator!=(const FMVVMViewDelayedBinding& Other) const
+	{
+		return CompiledBindingIndex != Other.CompiledBindingIndex;
+	}
+
+private:
+	int32 CompiledBindingIndex = INDEX_NONE;
+};
+
+/** */
+UCLASS(NotBlueprintable, Hidden)
+class UMVVMBindingSubsystem : public UEngineSubsystem
+{
+	GENERATED_BODY()
+
+public:
+	virtual void Initialize(FSubsystemCollectionBase& Collection) override;
+	virtual void Deinitialize() override;
+
+	void AddViewWithEveryTickBinding(const UMVVMView* View);
+	void RemoveViewWithEveryTickBinding(const UMVVMView* View);
+
+	void AddDelayedBinding(const UMVVMView* View, FMVVMViewDelayedBinding CompiledBinding);
+
+private:
+	void HandlePreTick(float DeltaTIme);
+
+	using FDelayedBindingList = TArray<FMVVMViewDelayedBinding, TInlineAllocator<8>>;
+	using FDelayedMap = TMap<TWeakObjectPtr<const UMVVMView>, FDelayedBindingList>;
+	FDelayedMap DelayedBindings;
+	TArray<TWeakObjectPtr<const UMVVMView>> EveryTickBindings;
+};
diff --git a/Engine/Plugins/Runtime/ModelViewViewModel/Source/ModelViewViewModel/Private/View/MVVMView.cpp b/Engine/Plugins/Runtime/ModelViewViewModel/Source/ModelViewViewModel/Private/View/MVVMView.cpp
--- a/Engine/Plugins/Runtime/ModelViewViewModel/Source/ModelViewViewModel/Private/View/MVVMView.cpp
+++ b/Engine/Plugins/Runtime/ModelViewViewModel/Source/ModelViewViewModel/Private/View/MVVMView.cpp
@@ -4,10 +4,11 @@
 #include "FieldNotification/FieldMulticastDelegate.h"
 #include "View/MVVMViewClass.h"
 #include "Misc/MemStack.h"
-#include "View/MVVMViewWorldSubsystem.h"
+#include "View/MVVMBindingSubsystem.h"
 
-#include "Debugging/MVVMDebugging.h"
 #include "Blueprint/UserWidget.h"
+#include "Debugging/MVVMDebugging.h"
+#include "Engine/Engine.h"
 #include "MVVMMessageLog.h"
 #include "Templates/ValueOrError.h"
 #include "Types/MVVMFieldContext.h"
@@ -59,7 +60,7 @@ void UMVVMView::Construct()
 
 	if (bHasEveryTickBinding)
 	{
-		GetUserWidget()->GetWorld()->GetSubsystem<UMVVMViewWorldSubsystem>()->AddViewWithEveryTickBinding(this);
+		GEngine->GetEngineSubsystem<UMVVMBindingSubsystem>()->AddViewWithEveryTickBinding(this);
 	}
 }
 
@@ -88,7 +89,7 @@ void UMVVMView::Destruct()
 
 	if (bHasEveryTickBinding)
 	{
-		GetUserWidget()->GetWorld()->GetSubsystem<UMVVMViewWorldSubsystem>()->RemoveViewWithEveryTickBinding(this);
+		GEngine->GetEngineSubsystem<UMVVMBindingSubsystem>()->RemoveViewWithEveryTickBinding(this);
 	}
 	bHasEveryTickBinding = false;
 }
@@ -187,11 +188,11 @@ bool UMVVMView::SetViewModel(FName ViewModelName, TScriptInterface<INotifyFieldV
 			{
 				if (bHasEveryTickBinding)
 				{
-					GetUserWidget()->GetWorld()->GetSubsystem<UMVVMViewWorldSubsystem>()->AddViewWithEveryTickBinding(this);
+					GEngine->GetEngineSubsystem<UMVVMBindingSubsystem>()->AddViewWithEveryTickBinding(this);
 				}
 				else
 				{
-					GetUserWidget()->GetWorld()->GetSubsystem<UMVVMViewWorldSubsystem>()->RemoveViewWithEveryTickBinding(this);
+					GEngine->GetEngineSubsystem<UMVVMBindingSubsystem>()->RemoveViewWithEveryTickBinding(this);
 				}
 			}
 		}
@@ -220,7 +221,7 @@ void UMVVMView::HandledLibraryBindingValueChanged(UObject* InViewModelOrWidget,
 		EMVVMExecutionMode ExecutionMode = Binding.GetExecuteMode();
 		if (ExecutionMode == EMVVMExecutionMode::Delayed)
 		{
-			GetUserWidget()->GetWorld()->GetSubsystem<UMVVMViewWorldSubsystem>()->AddDelayedBinding(this, InCompiledBindingIndex);
+			GEngine->GetEngineSubsystem<UMVVMBindingSubsystem>()->AddDelayedBinding(this, InCompiledBindingIndex);
 		}
 		else if (ExecutionMode != EMVVMExecutionMode::Tick)
 		{
diff --git a/Engine/Plugins/Runtime/ModelViewViewModel/Source/ModelViewViewModel/Private/View/MVVMViewWorldSubsystem.cpp b/Engine/Plugins/Runtime/ModelViewViewModel/Source/ModelViewViewModel/Private/View/MVVMViewWorldSubsystem.cpp
deleted file mode 100644
diff --git a/Engine/Plugins/Runtime/ModelViewViewModel/Source/ModelViewViewModel/Private/View/MVVMViewWorldSubsystem.h b/Engine/Plugins/Runtime/ModelViewViewModel/Source/ModelViewViewModel/Private/View/MVVMViewWorldSubsystem.h
deleted file mode 100644
