MVVM: Do not reset viewmodel pointer when not needed.

diff --git a/Engine/Plugins/Runtime/ModelViewViewModel/Source/ModelViewViewModel/Private/View/MVVMViewClass.cpp b/Engine/Plugins/Runtime/ModelViewViewModel/Source/ModelViewViewModel/Private/View/MVVMViewClass.cpp
--- a/Engine/Plugins/Runtime/ModelViewViewModel/Source/ModelViewViewModel/Private/View/MVVMViewClass.cpp
+++ b/Engine/Plugins/Runtime/ModelViewViewModel/Source/ModelViewViewModel/Private/View/MVVMViewClass.cpp
@@ -107,7 +107,7 @@ UObject* FMVVMViewClass_SourceCreator::CreateInstance(const UMVVMViewClass* InVi
 	{
 		auto AssignProperty = [FoundObjectProperty, InUserWidget](UObject* NewObject)
 		{
-			if (ensure(NewObject == nullptr || NewObject->GetClass()->IsChildOf(FoundObjectProperty->PropertyClass)))
+			if (NewObject && NewObject->GetClass()->IsChildOf(FoundObjectProperty->PropertyClass))
 			{
 				FoundObjectProperty->SetObjectPropertyValue_InContainer(InUserWidget, NewObject);
 			}
