MVVM: Test if the properties or functions are allowed at compile time.

diff --git a/Engine/Plugins/Runtime/ModelViewViewModel/Source/ModelViewViewModelBlueprint/Private/Bindings/MVVMCompiledBindingLibraryCompiler.cpp b/Engine/Plugins/Runtime/ModelViewViewModel/Source/ModelViewViewModelBlueprint/Private/Bindings/MVVMCompiledBindingLibraryCompiler.cpp
--- a/Engine/Plugins/Runtime/ModelViewViewModel/Source/ModelViewViewModelBlueprint/Private/Bindings/MVVMCompiledBindingLibraryCompiler.cpp
+++ b/Engine/Plugins/Runtime/ModelViewViewModel/Source/ModelViewViewModelBlueprint/Private/Bindings/MVVMCompiledBindingLibraryCompiler.cpp
@@ -4,6 +4,7 @@
 
 #include "Bindings/MVVMBindingHelper.h"
 #include "Engine/Engine.h"
+#include "MVVMDeveloperProjectSettings.h"
 #include "MVVMSubsystem.h"
 #include "Misc/TVariantMeta.h"
 #include <limits> // IWYU pragma: keep
@@ -258,6 +259,11 @@ TValueOrError<FCompiledBindingLibraryCompiler::FFieldPathHandle, FText> FCompile
 				return MakeError(ValidatedStr);
 			}
 
+			if (!GetDefault<UMVVMDeveloperProjectSettings>()->IsPropertyAllowed(FieldVariant.GetProperty()))
+			{
+				return MakeError(LOCTEXT("PropertyNotAllow", "A property is not allowed."));
+			}
+
 			RawFieldIndexes.Add(Impl->AddUniqueField(FieldVariant));
 		}
 		else if (FieldVariant.IsFunction())
@@ -274,6 +280,11 @@ TValueOrError<FCompiledBindingLibraryCompiler::FFieldPathHandle, FText> FCompile
 				return MakeError(FText::Format(LOCTEXT("FunctionNotReadableAtRuntime", "Function '{0}' is not readable at runtime."), FieldVariant.GetFunction()->GetDisplayNameText()));
 			}
 
+			if (!GetDefault<UMVVMDeveloperProjectSettings>()->IsFunctionAllowed(FieldVariant.GetFunction()))
+			{
+				return MakeError(LOCTEXT("FunctionNotAllow", "A function is not allowed."));
+			}
+
 			if (bIsLast && !bInRead)
 			{
 				const FProperty* FirstProperty = BindingHelper::GetFirstArgumentProperty(FieldVariant.GetFunction());
@@ -377,6 +388,11 @@ TValueOrError<FCompiledBindingLibraryCompiler::FFieldPathHandle, FText> FCompile
 		}
 	}
 
+	if (!GetDefault<UMVVMDeveloperProjectSettings>()->IsConversionFunctionAllowed(Function))
+	{
+		return MakeError(LOCTEXT("ConversionFunctionNotAllow", "The conversion function is not allowed."));
+	}
+
 	TArray<int32> RawFieldIndexes;
 	RawFieldIndexes.Add(Impl->AddUniqueField(FMVVMConstFieldVariant(Function)));
 	const int32 FoundFieldPath = Impl->FieldPaths.IndexOfByPredicate([&RawFieldIndexes](const Private::FRawFieldPath& Other)
diff --git a/Engine/Plugins/Runtime/ModelViewViewModel/Source/ModelViewViewModelBlueprint/Private/MVVMDeveloperProjectSettings.cpp b/Engine/Plugins/Runtime/ModelViewViewModel/Source/ModelViewViewModelBlueprint/Private/MVVMDeveloperProjectSettings.cpp
--- a/Engine/Plugins/Runtime/ModelViewViewModel/Source/ModelViewViewModelBlueprint/Private/MVVMDeveloperProjectSettings.cpp
+++ b/Engine/Plugins/Runtime/ModelViewViewModel/Source/ModelViewViewModelBlueprint/Private/MVVMDeveloperProjectSettings.cpp
@@ -3,11 +3,10 @@
 #include "MVVMDeveloperProjectSettings.h"
 #include "Engine/Blueprint.h"
 
-#include "MVVMBlueprintViewModelContext.h"
-#include "Types/MVVMExecutionMode.h"
-
 #include "BlueprintEditorSettings.h"
+#include "MVVMBlueprintViewModelContext.h"
 #include "PropertyEditorPermissionList.h"
+#include "Types/MVVMExecutionMode.h"
 
 #define LOCTEXT_NAMESPACE "MVVMDeveloperProjectSettings"
 
@@ -43,7 +42,7 @@ bool UMVVMDeveloperProjectSettings::IsPropertyAllowed(const FProperty* Property)
 		return false;
 	}
 
-	TStringBuilder<256> StringBuilder;
+	TStringBuilder<512> StringBuilder;
 	Property->GetOwnerClass()->GetPathName(nullptr, StringBuilder);
 	FSoftClassPath StructPath;
 	StructPath.SetPath(StringBuilder);
@@ -57,12 +56,15 @@ bool UMVVMDeveloperProjectSettings::IsPropertyAllowed(const FProperty* Property)
 bool UMVVMDeveloperProjectSettings::IsFunctionAllowed(const UFunction* Function) const
 {
 	check(Function);
-	if (!GetMutableDefault<UBlueprintEditorSettings>()->GetFunctionPermissions().PassesFilter(Function->GetFullName()))
+
+	TStringBuilder<512> StringBuilder;
+	Function->GetPathName(nullptr, StringBuilder);
+	if (!GetMutableDefault<UBlueprintEditorSettings>()->GetFunctionPermissions().PassesFilter(StringBuilder.ToView()))
 	{
 		return false;
 	}
 
-	TStringBuilder<256> StringBuilder;
+	StringBuilder.Reset();
 	Function->GetOwnerClass()->GetPathName(nullptr, StringBuilder);
 	FSoftClassPath StructPath;
 	StructPath.SetPath(StringBuilder);
@@ -73,6 +75,30 @@ bool UMVVMDeveloperProjectSettings::IsFunctionAllowed(const UFunction* Function)
 	return true;
 }
 
+bool UMVVMDeveloperProjectSettings::IsConversionFunctionAllowed(const UFunction* Function) const
+{
+	if (ConversionFunctionFilter == EMVVMDeveloperConversionFunctionFilterType::AllowedList)
+	{
+		return IsFunctionAllowed(Function);
+	}
+	else
+	{
+		TStringBuilder<512> FunctionClassPath;
+		Function->GetOwnerClass()->GetPathName(nullptr, FunctionClassPath);
+		TStringBuilder<512> AllowedClassPath;
+		for (const FSoftClassPath& SoftClass : AllowedClassForConversionFunctions)
+		{
+			SoftClass.ToString(AllowedClassPath);
+			if (AllowedClassPath.ToView() == FunctionClassPath.ToView())
+			{
+				return true;
+			}
+			AllowedClassPath.Reset();
+		}
+		return false;
+	}
+}
+
 TArray<const UClass*> UMVVMDeveloperProjectSettings::GetAllowedConversionFunctionClasses() const
 {
 	TArray<const UClass*> Result;
diff --git a/Engine/Plugins/Runtime/ModelViewViewModel/Source/ModelViewViewModelBlueprint/Public/MVVMDeveloperProjectSettings.h b/Engine/Plugins/Runtime/ModelViewViewModel/Source/ModelViewViewModelBlueprint/Public/MVVMDeveloperProjectSettings.h
--- a/Engine/Plugins/Runtime/ModelViewViewModel/Source/ModelViewViewModelBlueprint/Public/MVVMDeveloperProjectSettings.h
+++ b/Engine/Plugins/Runtime/ModelViewViewModel/Source/ModelViewViewModelBlueprint/Public/MVVMDeveloperProjectSettings.h
@@ -53,6 +53,7 @@ public:
 
 	bool IsPropertyAllowed(const FProperty* Property) const;
 	bool IsFunctionAllowed(const UFunction* Function) const;
+	bool IsConversionFunctionAllowed(const UFunction* Function) const;
 
 	bool IsExecutionModeAllowed(EMVVMExecutionMode ExecutionMode) const
 	{
