Revert properties or functions compile time test introdued with CL 24208984

diff --git a/Engine/Plugins/Runtime/ModelViewViewModel/Source/ModelViewViewModelBlueprint/Private/Bindings/MVVMCompiledBindingLibraryCompiler.cpp b/Engine/Plugins/Runtime/ModelViewViewModel/Source/ModelViewViewModelBlueprint/Private/Bindings/MVVMCompiledBindingLibraryCompiler.cpp
--- a/Engine/Plugins/Runtime/ModelViewViewModel/Source/ModelViewViewModelBlueprint/Private/Bindings/MVVMCompiledBindingLibraryCompiler.cpp
+++ b/Engine/Plugins/Runtime/ModelViewViewModel/Source/ModelViewViewModelBlueprint/Private/Bindings/MVVMCompiledBindingLibraryCompiler.cpp
@@ -259,10 +259,10 @@ TValueOrError<FCompiledBindingLibraryCompiler::FFieldPathHandle, FText> FCompile
 				return MakeError(ValidatedStr);
 			}
 
-			if (!GetDefault<UMVVMDeveloperProjectSettings>()->IsPropertyAllowed(FieldVariant.GetProperty()))
-			{
-				return MakeError(LOCTEXT("PropertyNotAllow", "A property is not allowed."));
-			}
+			//if (!GetDefault<UMVVMDeveloperProjectSettings>()->IsPropertyAllowed(FieldVariant.GetProperty()))
+			//{
+			//	return MakeError(LOCTEXT("PropertyNotAllow", "A property is not allowed."));
+			//}
 
 			RawFieldIndexes.Add(Impl->AddUniqueField(FieldVariant));
 		}
@@ -280,10 +280,10 @@ TValueOrError<FCompiledBindingLibraryCompiler::FFieldPathHandle, FText> FCompile
 				return MakeError(FText::Format(LOCTEXT("FunctionNotReadableAtRuntime", "Function '{0}' is not readable at runtime."), FieldVariant.GetFunction()->GetDisplayNameText()));
 			}
 
-			if (!GetDefault<UMVVMDeveloperProjectSettings>()->IsFunctionAllowed(FieldVariant.GetFunction()))
-			{
-				return MakeError(LOCTEXT("FunctionNotAllow", "A function is not allowed."));
-			}
+			//if (!GetDefault<UMVVMDeveloperProjectSettings>()->IsFunctionAllowed(FieldVariant.GetFunction()))
+			//{
+			//	return MakeError(LOCTEXT("FunctionNotAllow", "A function is not allowed."));
+			//}
 
 			if (bIsLast && !bInRead)
 			{
@@ -388,10 +388,10 @@ TValueOrError<FCompiledBindingLibraryCompiler::FFieldPathHandle, FText> FCompile
 		}
 	}
 
-	if (!GetDefault<UMVVMDeveloperProjectSettings>()->IsConversionFunctionAllowed(Function))
-	{
-		return MakeError(LOCTEXT("ConversionFunctionNotAllow", "The conversion function is not allowed."));
-	}
+	//if (!GetDefault<UMVVMDeveloperProjectSettings>()->IsConversionFunctionAllowed(Function))
+	//{
+	//	return MakeError(LOCTEXT("ConversionFunctionNotAllow", "The conversion function is not allowed."));
+	//}
 
 	TArray<int32> RawFieldIndexes;
 	RawFieldIndexes.Add(Impl->AddUniqueField(FMVVMConstFieldVariant(Function)));
