MVVM: Fix issue with local conversion function introduced with CL24074931

diff --git a/Engine/Plugins/Runtime/ModelViewViewModel/Source/ModelViewViewModelEditor/Private/MVVMEditorSubsystem.cpp b/Engine/Plugins/Runtime/ModelViewViewModel/Source/ModelViewViewModelEditor/Private/MVVMEditorSubsystem.cpp
--- a/Engine/Plugins/Runtime/ModelViewViewModel/Source/ModelViewViewModelEditor/Private/MVVMEditorSubsystem.cpp
+++ b/Engine/Plugins/Runtime/ModelViewViewModel/Source/ModelViewViewModelEditor/Private/MVVMEditorSubsystem.cpp
@@ -679,9 +679,10 @@ TArray<UFunction*> UMVVMEditorSubsystem::GetAvailableConversionFunctions(const U
 		// functions in the widget blueprint can do anything they want, other functions have to be static functions in a BlueprintFunctionLibrary
 		const UClass* FunctionClass = Function->GetOuterUClass();
 
-		bool bIsFromWidgetBlueprint = WidgetBlueprint->GeneratedClass->IsChildOf(FunctionClass) && Function->HasAllFunctionFlags(FUNC_BlueprintPure);
+		bool bIsFromWidgetBlueprint = WidgetBlueprint->GeneratedClass && WidgetBlueprint->GeneratedClass->IsChildOf(FunctionClass) && Function->HasAllFunctionFlags(FUNC_BlueprintPure);
+		bool bIsFromKeletonWidgetBlueprint = WidgetBlueprint->SkeletonGeneratedClass && WidgetBlueprint->SkeletonGeneratedClass->IsChildOf(FunctionClass) && Function->HasAllFunctionFlags(FUNC_BlueprintPure);
 		bool bFromBlueprintFunctionLibrary = FunctionClass->IsChildOf<UBlueprintFunctionLibrary>() && Function->HasAllFunctionFlags(FUNC_Static | FUNC_BlueprintPure);
-		if (bIsFromWidgetBlueprint || bFromBlueprintFunctionLibrary)
+		if (bIsFromWidgetBlueprint || bIsFromKeletonWidgetBlueprint || bFromBlueprintFunctionLibrary)
 		{
 			if (IsValidConversionFunction(Function, Source, Destination))
 			{
