MVVM: Test if the properties or functions are allowed at compile time.

diff --git a/Engine/Plugins/Runtime/ModelViewViewModel/Source/ModelViewViewModelBlueprint/Private/Bindings/MVVMCompiledBindingLibraryCompiler.cpp b/Engine/Plugins/Runtime/ModelViewViewModel/Source/ModelViewViewModelBlueprint/Private/Bindings/MVVMCompiledBindingLibraryCompiler.cpp
--- a/Engine/Plugins/Runtime/ModelViewViewModel/Source/ModelViewViewModelBlueprint/Private/Bindings/MVVMCompiledBindingLibraryCompiler.cpp
+++ b/Engine/Plugins/Runtime/ModelViewViewModel/Source/ModelViewViewModelBlueprint/Private/Bindings/MVVMCompiledBindingLibraryCompiler.cpp
@@ -259,10 +259,10 @@ TValueOrError<FCompiledBindingLibraryCompiler::FFieldPathHandle, FText> FCompile
 				return MakeError(ValidatedStr);
 			}
 
-			//if (!GetDefault<UMVVMDeveloperProjectSettings>()->IsPropertyAllowed(FieldVariant.GetProperty()))
-			//{
-			//	return MakeError(LOCTEXT("PropertyNotAllow", "A property is not allowed."));
-			//}
+			if (!GetDefault<UMVVMDeveloperProjectSettings>()->IsPropertyAllowed(FieldVariant.GetProperty()))
+			{
+				return MakeError(LOCTEXT("PropertyNotAllow", "A property is not allowed."));
+			}
 
 			RawFieldIndexes.Add(Impl->AddUniqueField(FieldVariant));
 		}
@@ -280,10 +280,10 @@ TValueOrError<FCompiledBindingLibraryCompiler::FFieldPathHandle, FText> FCompile
 				return MakeError(FText::Format(LOCTEXT("FunctionNotReadableAtRuntime", "Function '{0}' is not readable at runtime."), FieldVariant.GetFunction()->GetDisplayNameText()));
 			}
 
-			//if (!GetDefault<UMVVMDeveloperProjectSettings>()->IsFunctionAllowed(FieldVariant.GetFunction()))
-			//{
-			//	return MakeError(LOCTEXT("FunctionNotAllow", "A function is not allowed."));
-			//}
+			if (!GetDefault<UMVVMDeveloperProjectSettings>()->IsFunctionAllowed(FieldVariant.GetFunction()))
+			{
+				return MakeError(LOCTEXT("FunctionNotAllow", "A function is not allowed."));
+			}
 
 			if (bIsLast && !bInRead)
 			{
@@ -388,10 +388,10 @@ TValueOrError<FCompiledBindingLibraryCompiler::FFieldPathHandle, FText> FCompile
 		}
 	}
 
-	//if (!GetDefault<UMVVMDeveloperProjectSettings>()->IsConversionFunctionAllowed(Function))
-	//{
-	//	return MakeError(LOCTEXT("ConversionFunctionNotAllow", "The conversion function is not allowed."));
-	//}
+	if (!GetDefault<UMVVMDeveloperProjectSettings>()->IsConversionFunctionAllowed(Function))
+	{
+		return MakeError(LOCTEXT("ConversionFunctionNotAllow", "The conversion function is not allowed."));
+	}
 
 	TArray<int32> RawFieldIndexes;
 	RawFieldIndexes.Add(Impl->AddUniqueField(FMVVMConstFieldVariant(Function)));
diff --git a/Engine/Plugins/Runtime/ModelViewViewModel/Source/ModelViewViewModelBlueprint/Private/MVVMDeveloperProjectSettings.cpp b/Engine/Plugins/Runtime/ModelViewViewModel/Source/ModelViewViewModelBlueprint/Private/MVVMDeveloperProjectSettings.cpp
--- a/Engine/Plugins/Runtime/ModelViewViewModel/Source/ModelViewViewModelBlueprint/Private/MVVMDeveloperProjectSettings.cpp
+++ b/Engine/Plugins/Runtime/ModelViewViewModel/Source/ModelViewViewModelBlueprint/Private/MVVMDeveloperProjectSettings.cpp
@@ -58,10 +58,14 @@ bool UMVVMDeveloperProjectSettings::IsFunctionAllowed(const UFunction* Function)
 	check(Function);
 
 	TStringBuilder<512> StringBuilder;
-	Function->GetPathName(nullptr, StringBuilder);
-	if (!GetMutableDefault<UBlueprintEditorSettings>()->GetFunctionPermissions().PassesFilter(StringBuilder.ToView()))
+	const FPathPermissionList& FunctionPermissions = GetMutableDefault<UBlueprintEditorSettings>()->GetFunctionPermissions();
+	if (FunctionPermissions.HasFiltering())
 	{
-		return false;
+		Function->GetPathName(nullptr, StringBuilder);
+		if (!FunctionPermissions.PassesFilter(StringBuilder.ToView()))
+		{
+			return false;
+		}
 	}
 
 	StringBuilder.Reset();
@@ -77,25 +81,42 @@ bool UMVVMDeveloperProjectSettings::IsFunctionAllowed(const UFunction* Function)
 
 bool UMVVMDeveloperProjectSettings::IsConversionFunctionAllowed(const UFunction* Function) const
 {
-	if (ConversionFunctionFilter == EMVVMDeveloperConversionFunctionFilterType::AllowedList)
+	static FName NAME_ComplexConversionFunction = TEXT("MVVMComplexConversionFunction");
+	if (Function->HasMetaData(NAME_ComplexConversionFunction))
+	{
+		return true;
+	}
+
+	if (ConversionFunctionFilter == EMVVMDeveloperConversionFunctionFilterType::BlueprintActionRegistry)
 	{
 		return IsFunctionAllowed(Function);
 	}
 	else
 	{
-		TStringBuilder<512> FunctionClassPath;
-		Function->GetOwnerClass()->GetPathName(nullptr, FunctionClassPath);
-		TStringBuilder<512> AllowedClassPath;
-		for (const FSoftClassPath& SoftClass : AllowedClassForConversionFunctions)
+		check(ConversionFunctionFilter == EMVVMDeveloperConversionFunctionFilterType::AllowedList);
+
+		if (Function->HasAllFunctionFlags(FUNC_Static))
 		{
-			SoftClass.ToString(AllowedClassPath);
-			if (AllowedClassPath.ToView() == FunctionClassPath.ToView())
+			TStringBuilder<512> FunctionClassPath;
+			Function->GetOwnerClass()->GetPathName(nullptr, FunctionClassPath);
+			TStringBuilder<512> AllowedClassPath;
+			for (const FSoftClassPath& SoftClass : AllowedClassForConversionFunctions)
 			{
-				return true;
+				SoftClass.ToString(AllowedClassPath);
+				if (AllowedClassPath.ToView() == FunctionClassPath.ToView())
+				{
+					return true;
+				}
+				AllowedClassPath.Reset();
 			}
-			AllowedClassPath.Reset();
+
+			return false;
+		}
+		else
+		{
+			// The function is on self and may have been filtered.
+			return IsFunctionAllowed(Function);
 		}
-		return false;
 	}
 }
 
diff --git a/Engine/Plugins/Runtime/ModelViewViewModel/Source/ModelViewViewModelBlueprint/Private/MVVMViewBlueprintCompiler.cpp b/Engine/Plugins/Runtime/ModelViewViewModel/Source/ModelViewViewModelBlueprint/Private/MVVMViewBlueprintCompiler.cpp
--- a/Engine/Plugins/Runtime/ModelViewViewModel/Source/ModelViewViewModelBlueprint/Private/MVVMViewBlueprintCompiler.cpp
+++ b/Engine/Plugins/Runtime/ModelViewViewModel/Source/ModelViewViewModelBlueprint/Private/MVVMViewBlueprintCompiler.cpp
@@ -5,15 +5,16 @@
 #include "Bindings/MVVMConversionFunctionHelper.h"
 #include "Bindings/MVVMFieldPathHelper.h"
 #include "Components/Widget.h"
+#include "HAL/IConsoleManager.h"
 #include "MVVMBlueprintView.h"
 #include "MVVMDeveloperProjectSettings.h"
-#include "HAL/IConsoleManager.h"
+#include "PropertyEditorPermissionList.h"
 #include "MVVMFunctionGraphHelper.h"
 #include "Templates/ValueOrError.h"
 #include "Types/MVVMBindingName.h"
+#include "UObject/LinkerLoad.h"
 #include "View/MVVMViewClass.h"
 #include "View/MVVMViewModelContextResolver.h"
-#include "UObject/LinkerLoad.h"
 
 #define LOCTEXT_NAMESPACE "MVVMViewBlueprintCompiler"
 
@@ -590,6 +591,8 @@ bool FMVVMViewBlueprintCompiler::PreCompile(UWidgetBlueprintGeneratedClass* Clas
 	CompilerBindings.Reset(NumBindings*2);
 	BindingSourceContexts.Reset(NumBindings*2);
 
+	FPropertyEditorPermissionList::Get().AddPermissionList(Class, FNamePermissionList(), EPropertyEditorPermissionListRules::AllowListAllProperties);
+
 	PreCompileBindingSources(Class, BlueprintView);
 	PreCompileSourceCreators(Class, BlueprintView);
 	PreCompileBindings(Class, BlueprintView);
diff --git a/Engine/Plugins/Runtime/ModelViewViewModel/Source/ModelViewViewModelEditor/Private/MVVMEditorSubsystem.cpp b/Engine/Plugins/Runtime/ModelViewViewModel/Source/ModelViewViewModelEditor/Private/MVVMEditorSubsystem.cpp
--- a/Engine/Plugins/Runtime/ModelViewViewModel/Source/ModelViewViewModelEditor/Private/MVVMEditorSubsystem.cpp
+++ b/Engine/Plugins/Runtime/ModelViewViewModel/Source/ModelViewViewModelEditor/Private/MVVMEditorSubsystem.cpp
@@ -644,7 +644,7 @@ bool UMVVMEditorSubsystem::IsValidConversionFunction(const UFunction* Function,
 		return false;
 	}
 
-	return true;
+	return GetDefault<UMVVMDeveloperProjectSettings>()->IsConversionFunctionAllowed(Function);
 }
 
 bool UMVVMEditorSubsystem::IsSimpleConversionFunction(const UFunction* Function) const
