MVVM: Fix viewmodel rename from the detail panel.

diff --git a/Engine/Plugins/Runtime/ModelViewViewModel/Source/ModelViewViewModelEditor/Private/Customizations/MVVMBlueprintViewModelContextCustomization.cpp b/Engine/Plugins/Runtime/ModelViewViewModel/Source/ModelViewViewModelEditor/Private/Customizations/MVVMBlueprintViewModelContextCustomization.cpp
--- a/Engine/Plugins/Runtime/ModelViewViewModel/Source/ModelViewViewModelEditor/Private/Customizations/MVVMBlueprintViewModelContextCustomization.cpp
+++ b/Engine/Plugins/Runtime/ModelViewViewModel/Source/ModelViewViewModelEditor/Private/Customizations/MVVMBlueprintViewModelContextCustomization.cpp
@@ -2,17 +2,20 @@
 
 #include "MVVMBlueprintViewModelContextCustomization.h"
 
-#include "Framework/MultiBox/MultiBoxBuilder.h"
+#include "Bindings/MVVMBindingHelper.h"
+#include "DetailWidgetRow.h"
+#include "Features/IModularFeatures.h"
+#include "IDetailChildrenBuilder.h"
+#include "IPropertyAccessEditor.h"
 #include "IPropertyTypeCustomization.h"
+#include "MVVMEditorSubsystem.h"
 #include "PropertyHandle.h"
-#include "IDetailChildrenBuilder.h"
-#include "DetailWidgetRow.h"
+
+#include "Framework/MultiBox/MultiBoxBuilder.h"
 #include "WidgetBlueprintEditor.h"
-#include "IPropertyAccessEditor.h"
-#include "Bindings/MVVMBindingHelper.h"
-#include "Features/IModularFeatures.h"
 #include "Widgets/SBoxPanel.h"
 #include "Widgets/Input/SComboButton.h"
+#include "Widgets/Input/SEditableTextBox.h"
 #include "Widgets/SMVVMViewModelPanel.h"
 
 #define LOCTEXT_NAMESPACE "BlueprintViewModelContextDetailCustomization"
@@ -122,13 +125,20 @@ void FBlueprintViewModelContextDetailCustomization::CustomizeChildren(TSharedRef
 	uint32 NumChildren = 0;
 	PropertyHandle->GetNumChildren(NumChildren);
 
+	FName Name_NotifyFieldValueClass = TEXT("NotifyFieldValueClass");
+	FName Name_ViewModelContextId = TEXT("ViewModelContextId");
+	FName Name_ViewModelName = GET_MEMBER_NAME_CHECKED(FMVVMBlueprintViewModelContext, ViewModelName);
+	FName Name_ViewModelPropertyPath = GET_MEMBER_NAME_CHECKED(FMVVMBlueprintViewModelContext, ViewModelPropertyPath);
+	FName Name_CreationType = GET_MEMBER_NAME_CHECKED(FMVVMBlueprintViewModelContext, CreationType);
+
+	FGuid ViewModelContextId;
 	FName ViewModelPropertyName;
 	for (uint32 ChildIndex = 0; ChildIndex < NumChildren; ++ChildIndex)
 	{
 		TSharedPtr<IPropertyHandle> ChildHandle = PropertyHandle->GetChildHandle(ChildIndex);
 		const FName PropertyName = ChildHandle->GetProperty()->GetFName();
 
-		if (PropertyName == "NotifyFieldValueClass")
+		if (PropertyName == Name_NotifyFieldValueClass)
 		{
 			ensure(CastField<FClassProperty>(ChildHandle->GetProperty()));
 			NotifyFieldValueClassHandle = ChildHandle;
@@ -145,23 +155,31 @@ void FBlueprintViewModelContextDetailCustomization::CustomizeChildren(TSharedRef
 			ChildHandle->SetOnPropertyValueChanged(FSimpleDelegate::CreateSP(this, &FBlueprintViewModelContextDetailCustomization::HandleClassChanged));
 		}
 
-		if (PropertyName == GET_MEMBER_NAME_CHECKED(FMVVMBlueprintViewModelContext, ViewModelName))
+		else if (PropertyName == Name_ViewModelName)
 		{
 			ensure(CastField<FNameProperty>(ChildHandle->GetProperty()));
+			ViewModelNameHandle = ChildHandle;
 			ChildHandle->GetValue(ViewModelPropertyName);
 		}
-		
-		if (PropertyName == GET_MEMBER_NAME_CHECKED(FMVVMBlueprintViewModelContext, ViewModelPropertyPath))
+		else if (PropertyName == Name_ViewModelPropertyPath)
 		{
 			ensure(CastField<FStrProperty>(ChildHandle->GetProperty()));
 			PropertyPathHandle = ChildHandle;
 		}
-
-		if (PropertyName == GET_MEMBER_NAME_CHECKED(FMVVMBlueprintViewModelContext, CreationType))
+		else if (PropertyName == Name_CreationType)
 		{
 			ensure(CastField<FEnumProperty>(ChildHandle->GetProperty()));
 			CreationTypeHandle = ChildHandle;
 		}
+		else if (PropertyName == Name_ViewModelContextId)
+		{
+			ensure(CastField<FStructProperty>(ChildHandle->GetProperty()) && CastField<FStructProperty>(ChildHandle->GetProperty())->Struct->GetFName() == "Guid");
+			void* Buffer = nullptr;
+			if (ChildHandle->GetValueData(Buffer) == FPropertyAccess::Success)
+			{
+				ViewModelContextId = *reinterpret_cast<FGuid*>(Buffer);
+			}
+		}
 	}
 
 	if (NotifyFieldValueClassHandle)
@@ -171,7 +189,40 @@ void FBlueprintViewModelContextDetailCustomization::CustomizeChildren(TSharedRef
 			TSharedPtr<IPropertyHandle> ChildHandle = PropertyHandle->GetChildHandle(ChildIndex);
 			const FName PropertyName = ChildHandle->GetProperty()->GetFName();
 
-			if (PropertyName == GET_MEMBER_NAME_CHECKED(FMVVMBlueprintViewModelContext, ViewModelPropertyPath))
+			if (PropertyName == Name_ViewModelName)
+			{
+				ensure(CastField<FNameProperty>(ChildHandle->GetProperty()));
+				if (TSharedPtr<FWidgetBlueprintEditor> SharedWidgetBlueprintEditor = WidgetBlueprintEditor.Pin())
+				{
+					IDetailPropertyRow& PropertyRow = ChildBuilder.AddProperty(ChildHandle.ToSharedRef());
+
+					TSharedPtr<SWidget> NameWidget, ValueWidget;
+					PropertyRow.GetDefaultWidgets(NameWidget, ValueWidget);
+					PropertyRow.CustomWidget()
+					.NameContent()
+					[
+						NameWidget.ToSharedRef()
+					]
+					.ValueContent()
+					[
+						SNew(SHorizontalBox)
+						+ SHorizontalBox::Slot()
+						.FillWidth(1.f)
+						[
+							SNew(SEditableTextBox)
+							.Text(this, &FBlueprintViewModelContextDetailCustomization::GetViewModalNameValueAsText)
+							.Font(CustomizationUtils.GetRegularFont())
+							.SelectAllTextWhenFocused(true)
+							.ClearKeyboardFocusOnCommit(false)
+							.OnTextCommitted(this, &FBlueprintViewModelContextDetailCustomization::HandleNameTextCommitted)
+							.OnVerifyTextChanged(this, &FBlueprintViewModelContextDetailCustomization::HandleNameVerifyTextChanged)
+							.SelectAllTextOnCommit(true)
+							.IsEnabled(ViewModelContextId.IsValid())
+						]
+					];
+				}
+			}
+			else if (PropertyName == Name_ViewModelPropertyPath)
 			{
 				ensure(CastField<FStrProperty>(ChildHandle->GetProperty()));
 				if (TSharedPtr<FWidgetBlueprintEditor> SharedWidgetBlueprintEditor = WidgetBlueprintEditor.Pin())
@@ -201,7 +252,7 @@ void FBlueprintViewModelContextDetailCustomization::CustomizeChildren(TSharedRef
 					];
 				}
 			}
-			else if (PropertyName == GET_MEMBER_NAME_CHECKED(FMVVMBlueprintViewModelContext, CreationType))
+			else if (PropertyName == Name_CreationType)
 			{
 				IDetailPropertyRow& PropertyRow = ChildBuilder.AddProperty(ChildHandle.ToSharedRef());
 
@@ -296,6 +347,82 @@ FText FBlueprintViewModelContextDetailCustomization::GetExecutionTypeValueToolTi
 	return FText::GetEmpty();
 }
 
+FText FBlueprintViewModelContextDetailCustomization::GetViewModalNameValueAsText() const
+{
+	check(ViewModelNameHandle.IsValid());
+	FText Result;
+	ViewModelNameHandle->GetValueAsFormattedText(Result);
+	return Result;
+}
+
+namespace Private
+{
+bool VerifyViewModelName(TSharedPtr<FWidgetBlueprintEditor> WidgetBlueprintEditor, TSharedPtr<IPropertyHandle> ViewModelNameHandle, const FText& RenameTo, bool bCommit, FText& OutErrorMessage)
+{
+	if (!WidgetBlueprintEditor || !ViewModelNameHandle)
+	{
+		return false;
+	}
+
+	if (RenameTo.IsEmptyOrWhitespace())
+	{
+		OutErrorMessage = LOCTEXT("EmptyViewModelName", "Empty viewmodel name.");
+		return false;
+	}
+
+	const FString& NewNameString = RenameTo.ToString();
+	if (NewNameString.Len() >= NAME_SIZE)
+	{
+		OutErrorMessage = LOCTEXT("ViewModelNameTooLong", "Viewmodel name is too long.");
+		return false;
+	}
+
+	FString GeneratedName = SlugStringForValidName(NewNameString);
+	if (NewNameString != GeneratedName)
+	{
+		OutErrorMessage = LOCTEXT("ViewModelHasInvalidChar", "ViewModel name has an invalid character.");
+		return false;
+	}
+
+	FName CurrentViewModelName;
+	if (ViewModelNameHandle->GetValue(CurrentViewModelName) != FPropertyAccess::Success)
+	{
+		OutErrorMessage = LOCTEXT("MultipleViewModel", "Can't edit multiple viewmodel name.");
+		return false;
+	}
+
+	const FName GeneratedFName(*GeneratedName);
+	check(GeneratedFName.IsValidXName(INVALID_OBJECTNAME_CHARACTERS));
+
+	if (UWidgetBlueprint* WidgetBP = WidgetBlueprintEditor->GetWidgetBlueprintObj())
+	{
+		if (bCommit)
+		{
+			return GEditor->GetEditorSubsystem<UMVVMEditorSubsystem>()->RenameViewModel(WidgetBP, CurrentViewModelName, *NewNameString, OutErrorMessage);
+		}
+		else
+		{
+			return GEditor->GetEditorSubsystem<UMVVMEditorSubsystem>()->VerifyViewModelRename(WidgetBP, CurrentViewModelName, *NewNameString, OutErrorMessage);
+		}
+	}
+	return false;
+}
+}
+
+void FBlueprintViewModelContextDetailCustomization::HandleNameTextCommitted(const FText& NewText, ETextCommit::Type CommitInfo)
+{
+	if (CommitInfo == ETextCommit::OnEnter)
+	{
+		FText OutErrorMessage;
+		Private::VerifyViewModelName(WidgetBlueprintEditor.Pin(), ViewModelNameHandle, NewText, true, OutErrorMessage);
+	}
+}
+
+bool FBlueprintViewModelContextDetailCustomization::HandleNameVerifyTextChanged(const FText& NewText, FText& OutError) const
+{
+	return Private::VerifyViewModelName(WidgetBlueprintEditor.Pin(), ViewModelNameHandle, NewText, false, OutError);
+}
+
 } // namespace UE::MVVM
 
 #undef LOCTEXT_NAMESPACE
diff --git a/Engine/Plugins/Runtime/ModelViewViewModel/Source/ModelViewViewModelEditor/Private/Customizations/MVVMBlueprintViewModelContextCustomization.h b/Engine/Plugins/Runtime/ModelViewViewModel/Source/ModelViewViewModelEditor/Private/Customizations/MVVMBlueprintViewModelContextCustomization.h
--- a/Engine/Plugins/Runtime/ModelViewViewModel/Source/ModelViewViewModelEditor/Private/Customizations/MVVMBlueprintViewModelContextCustomization.h
+++ b/Engine/Plugins/Runtime/ModelViewViewModel/Source/ModelViewViewModelEditor/Private/Customizations/MVVMBlueprintViewModelContextCustomization.h
@@ -51,12 +51,16 @@ private:
 	FText GetCreationTypeValue() const;
 	FText GetExecutionTypeValueToolTip() const;
 	void HandleClassChanged();
+	FText GetViewModalNameValueAsText() const;
+	void HandleNameTextCommitted(const FText& NewText, ETextCommit::Type /*CommitInfo*/);
+	bool HandleNameVerifyTextChanged(const FText& Text, FText& OutError) const;
 
 	FViewModelPropertyAccessEditor PropertyAccessEditor;
 	TWeakPtr<FWidgetBlueprintEditor> WidgetBlueprintEditor;
 	TSharedPtr<IPropertyHandle> NotifyFieldValueClassHandle;
 	TSharedPtr<IPropertyHandle> PropertyPathHandle;
 	TSharedPtr<IPropertyHandle> CreationTypeHandle;
+	TSharedPtr<IPropertyHandle> ViewModelNameHandle;
 	TArray<EMVVMBlueprintViewModelContextCreationType> AllowedCreationTypes;
 };
 
