MVVM: Fix access to private UFunction with AllowPrivateAccess.
Fix the IsBindingValid error message that do not include the full property type.
Fix field picker that do not show function when there is no filter.
Support array of object binding (the same way BP does).

diff --git a/Engine/Plugins/Runtime/ModelViewViewModel/Source/ModelViewViewModel/Private/Bindings/MVVMBindingHelper.cpp b/Engine/Plugins/Runtime/ModelViewViewModel/Source/ModelViewViewModel/Private/Bindings/MVVMBindingHelper.cpp
--- a/Engine/Plugins/Runtime/ModelViewViewModel/Source/ModelViewViewModel/Private/Bindings/MVVMBindingHelper.cpp
+++ b/Engine/Plugins/Runtime/ModelViewViewModel/Source/ModelViewViewModel/Private/Bindings/MVVMBindingHelper.cpp
@@ -464,6 +464,15 @@ namespace UE::MVVM::BindingHelper
 				&& SourceObjectProperty->PropertyClass->IsChildOf(DestinationObjectProperty->PropertyClass);
 		}
 
+		bool IsObjectArrayPropertyCompatible(const FProperty* Source, const FProperty* Destination)
+		{
+			const FArrayProperty* SourceArrayObjectProperty = CastField<const FArrayProperty>(Source);
+			const FArrayProperty* DestinationArrayObjectProperty = CastField<const FArrayProperty>(Destination);
+			return SourceArrayObjectProperty
+				&& DestinationArrayObjectProperty
+				&& IsObjectPropertyCompatible(SourceArrayObjectProperty->Inner, DestinationArrayObjectProperty->Inner);
+		}
+
 
 		bool IsNumericConversionRequired(const FProperty* Source, const FProperty* Destination)
 		{
@@ -528,7 +537,8 @@ namespace UE::MVVM::BindingHelper
 
 		return Source->SameType(Destination)
 				|| Private::IsNumericConversionRequired(Source, Destination)
-				|| Private::IsObjectPropertyCompatible(Source, Destination);
+				|| Private::IsObjectPropertyCompatible(Source, Destination)
+				|| Private::IsObjectArrayPropertyCompatible(Source, Destination);
 	}
 
 
diff --git a/Engine/Plugins/Runtime/ModelViewViewModel/Source/ModelViewViewModel/Private/MVVMSubsystem.cpp b/Engine/Plugins/Runtime/ModelViewViewModel/Source/ModelViewViewModel/Private/MVVMSubsystem.cpp
--- a/Engine/Plugins/Runtime/ModelViewViewModel/Source/ModelViewViewModel/Private/MVVMSubsystem.cpp
+++ b/Engine/Plugins/Runtime/ModelViewViewModel/Source/ModelViewViewModel/Private/MVVMSubsystem.cpp
@@ -56,14 +56,23 @@ namespace UE::MVVM::Private
 	{
 		// N.B. A function can be private/protected and can still be use when it's a BlueprintGetter/BlueprintSetter.
 		//But we bind to the property not the function.
+		static FName Name_AllowPrivateAccess = TEXT("AllowPrivateAccess");
+#if WITH_EDITOR
 		if (Function->HasAnyFunctionFlags(FUNC_Private) && !bCanAccessPrivate)
 		{
-			return FMVVMAvailableBinding();
+			if (!Function->GetBoolMetaData(Name_AllowPrivateAccess))
+			{
+				return FMVVMAvailableBinding();
+			}
 		}
 		if (Function->HasAnyFunctionFlags(FUNC_Protected) && !bCanAccessProtected)
 		{
-			return FMVVMAvailableBinding();
+			if (!Function->GetBoolMetaData(Name_AllowPrivateAccess))
+			{
+				return FMVVMAvailableBinding();
+			}
 		}
+#endif
 
 		const bool bIsReadable = BindingHelper::IsValidForSourceBinding(Function);
 		const bool bIsWritable = BindingHelper::IsValidForDestinationBinding(Function);
@@ -352,6 +361,17 @@ TValueOrError<bool, FText> UMVVMSubsystem::IsBindingValid(FConstDirectionalBindi
 		return MakeError(LOCTEXT("NoValueToBindAtDestination", "There is no value to bind at the destination."));
 	}
 
+	auto GetPropertyType = [](const FProperty* Property)
+	{
+		FString InnerType;
+		FString CppType = Property->GetCPPType(&InnerType);
+		if (InnerType.Len())
+		{
+			CppType.Append(InnerType);
+		}
+		return CppType;
+	};
+
 	// Test the conversion function
 	if (Args.ConversionFunction)
 	{
@@ -396,14 +416,15 @@ TValueOrError<bool, FText> UMVVMSubsystem::IsBindingValid(FConstDirectionalBindi
 			{
 				TArray<FText> ArgumentTypes;
 				ArgumentTypes.Reserve(ArgumentProperties.Num());
+				FString InnerType;
 				for (const FProperty* ArgProperty : ArgumentProperties)
 				{
-					ArgumentTypes.Add(FText::FromString(ArgProperty->GetCPPType()));
+					ArgumentTypes.Add(FText::FromString(GetPropertyType(ArgProperty)));
 				}
 
 				return MakeError(FText::Format(LOCTEXT("SourcePropertyDoesNotMatchArguments", "The source property '{0}' ({1}) does not match any of the argument types of the conversion function ({1})."), 
 					FText::FromString(SourceProperty->GetName()), 
-					FText::FromString(SourceProperty->GetCPPType()),
+					FText::FromString(GetPropertyType(SourceProperty)),
 					FText::Join(LOCTEXT("CommaDelim", ", "), ArgumentTypes))
 				);
 			}
@@ -417,7 +438,7 @@ TValueOrError<bool, FText> UMVVMSubsystem::IsBindingValid(FConstDirectionalBindi
 			{
 				return MakeError(FText::Format(LOCTEXT("DestinationPropertyDoesNotMatchReturn", "The destination property '{0}' ({1}) does not match the return type of the conversion function ({2})."), 
 					FText::FromString(DestinationProperty->GetName()), 
-					FText::FromString(DestinationProperty->GetCPPType()),
+					FText::FromString(GetPropertyType(DestinationProperty)),
 					FText::FromString(ReturnProperty->GetCPPType()))
 				);
 			}
@@ -427,9 +448,9 @@ TValueOrError<bool, FText> UMVVMSubsystem::IsBindingValid(FConstDirectionalBindi
 	{
 		return MakeError(FText::Format(LOCTEXT("SourcePropertyDoesNotMatchDestination", "The source property '{0}' ({1}) does not match the type of the destination property '{2}' ({3}). A conversion function is required."), 
 			FText::FromString(SourceProperty->GetName()),
-			FText::FromString(SourceProperty->GetCPPType()), 
+			FText::FromString(GetPropertyType(SourceProperty)),
 			FText::FromString(DestinationProperty->GetName()), 
-			FText::FromString(DestinationProperty->GetCPPType()))
+			FText::FromString(GetPropertyType(DestinationProperty)))
 		);
 	}
 
diff --git a/Engine/Plugins/Runtime/ModelViewViewModel/Source/ModelViewViewModelBlueprint/Public/MVVMBlueprintViewModelContext.h b/Engine/Plugins/Runtime/ModelViewViewModel/Source/ModelViewViewModelBlueprint/Public/MVVMBlueprintViewModelContext.h
--- a/Engine/Plugins/Runtime/ModelViewViewModel/Source/ModelViewViewModelBlueprint/Public/MVVMBlueprintViewModelContext.h
+++ b/Engine/Plugins/Runtime/ModelViewViewModel/Source/ModelViewViewModelBlueprint/Public/MVVMBlueprintViewModelContext.h
@@ -122,10 +122,10 @@ public:
 	TObjectPtr<UMVVMViewModelContextResolver> Resolver = nullptr;
 
 	/**
-	 * Generate a setter function for this viewmodel.
+	 * Generate a public setter for this viewmodel.
 	 * @note Always true when the Creation Type is Manual.
 	 */
-	UPROPERTY(EditAnywhere, Category = "Viewmodel", AdvancedDisplay, meta = (EditCondition = "CreationType != EMVVMBlueprintViewModelContextCreationType::Manual", EditConditionHides))
+	UPROPERTY(EditAnywhere, Category = "Viewmodel", AdvancedDisplay, meta = (DisplanName="Create Public Setter", EditCondition = "CreationType != EMVVMBlueprintViewModelContextCreationType::Manual", EditConditionHides))
 	bool bCreateSetterFunction = false;
 
 	/**
diff --git a/Engine/Plugins/Runtime/ModelViewViewModel/Source/ModelViewViewModelEditor/Private/MVVMEditorSubsystem.cpp b/Engine/Plugins/Runtime/ModelViewViewModel/Source/ModelViewViewModelEditor/Private/MVVMEditorSubsystem.cpp
--- a/Engine/Plugins/Runtime/ModelViewViewModel/Source/ModelViewViewModelEditor/Private/MVVMEditorSubsystem.cpp
+++ b/Engine/Plugins/Runtime/ModelViewViewModel/Source/ModelViewViewModelEditor/Private/MVVMEditorSubsystem.cpp
@@ -680,9 +680,9 @@ TArray<UFunction*> UMVVMEditorSubsystem::GetAvailableConversionFunctions(const U
 		const UClass* FunctionClass = Function->GetOuterUClass();
 
 		bool bIsFromWidgetBlueprint = WidgetBlueprint->GeneratedClass && WidgetBlueprint->GeneratedClass->IsChildOf(FunctionClass) && Function->HasAllFunctionFlags(FUNC_BlueprintPure);
-		bool bIsFromKeletonWidgetBlueprint = WidgetBlueprint->SkeletonGeneratedClass && WidgetBlueprint->SkeletonGeneratedClass->IsChildOf(FunctionClass) && Function->HasAllFunctionFlags(FUNC_BlueprintPure);
+		bool bIsFromSkeletonWidgetBlueprint = WidgetBlueprint->SkeletonGeneratedClass && WidgetBlueprint->SkeletonGeneratedClass->IsChildOf(FunctionClass) && Function->HasAllFunctionFlags(FUNC_BlueprintPure);
 		bool bFromBlueprintFunctionLibrary = FunctionClass->IsChildOf<UBlueprintFunctionLibrary>() && Function->HasAllFunctionFlags(FUNC_Static | FUNC_BlueprintPure);
-		if (bIsFromWidgetBlueprint || bIsFromKeletonWidgetBlueprint || bFromBlueprintFunctionLibrary)
+		if (bIsFromWidgetBlueprint || bIsFromSkeletonWidgetBlueprint || bFromBlueprintFunctionLibrary)
 		{
 			if (IsValidConversionFunction(Function, Source, Destination))
 			{
diff --git a/Engine/Plugins/Runtime/ModelViewViewModel/Source/ModelViewViewModelEditor/Private/Widgets/SMVVMViewModelBindingListWidget.cpp b/Engine/Plugins/Runtime/ModelViewViewModel/Source/ModelViewViewModelEditor/Private/Widgets/SMVVMViewModelBindingListWidget.cpp
--- a/Engine/Plugins/Runtime/ModelViewViewModel/Source/ModelViewViewModelEditor/Private/Widgets/SMVVMViewModelBindingListWidget.cpp
+++ b/Engine/Plugins/Runtime/ModelViewViewModel/Source/ModelViewViewModelEditor/Private/Widgets/SMVVMViewModelBindingListWidget.cpp
@@ -78,14 +78,14 @@ namespace Private
 				{
 					if (EnumHasAllFlags(FieldVisibilityFlags, EFieldVisibility::Writable))
 					{
-						if (!BindingHelper::ArePropertiesCompatible(AssignableTo, BindingHelper::GetFirstArgumentProperty(Function)))
+						if (AssignableTo && !BindingHelper::ArePropertiesCompatible(AssignableTo, BindingHelper::GetFirstArgumentProperty(Function)))
 						{
 							return TOptional<FFieldVariant>();
 						}
 					}
 					if (EnumHasAllFlags(FieldVisibilityFlags, EFieldVisibility::Readable))
 					{
-						if (!BindingHelper::ArePropertiesCompatible(ReturnProperty, AssignableTo))
+						if (AssignableTo && !BindingHelper::ArePropertiesCompatible(ReturnProperty, AssignableTo))
 						{
 							return TOptional<FFieldVariant>();
 						}
