MVVM: First step to support viewmodel inside viewmodel. The wiew now support source that is not a UserWidget property. The compiler generates the missing path.

diff --git a/Engine/Plugins/Runtime/ModelViewViewModel/Source/ModelViewViewModel/Private/Bindings/MVVMCompiledBindingLibrary.cpp b/Engine/Plugins/Runtime/ModelViewViewModel/Source/ModelViewViewModel/Private/Bindings/MVVMCompiledBindingLibrary.cpp
--- a/Engine/Plugins/Runtime/ModelViewViewModel/Source/ModelViewViewModel/Private/Bindings/MVVMCompiledBindingLibrary.cpp
+++ b/Engine/Plugins/Runtime/ModelViewViewModel/Source/ModelViewViewModel/Private/Bindings/MVVMCompiledBindingLibrary.cpp
@@ -61,6 +61,11 @@ void FMVVMCompiledBindingLibrary::Load()
 
 	for (const FMVVMVCompiledFields& Field : CompiledFields)
 	{
+		if (const UClass* Class = Cast<UClass>(Field.GetStruct()))
+		{
+			ensureAlwaysMsgf(!Class->HasAnyClassFlags(CLASS_NewerVersionExists), TEXT("The field is invalid. Property and Functions will not be loaded correctly."));
+		}
+
 		{
 			const int32 NumberOfProperties = Field.GetPropertyNum();
 			for (int32 Index = 0; Index < NumberOfProperties; ++Index)
diff --git a/Engine/Plugins/Runtime/ModelViewViewModel/Source/ModelViewViewModel/Private/View/MVVMView.cpp b/Engine/Plugins/Runtime/ModelViewViewModel/Source/ModelViewViewModel/Private/View/MVVMView.cpp
--- a/Engine/Plugins/Runtime/ModelViewViewModel/Source/ModelViewViewModel/Private/View/MVVMView.cpp
+++ b/Engine/Plugins/Runtime/ModelViewViewModel/Source/ModelViewViewModel/Private/View/MVVMView.cpp
@@ -33,11 +33,35 @@ void UMVVMView::Construct()
 {
 	check(ClassExtension);
 	check(bConstructed == false);
+	check(Sources.Num() == 0);
 
 	// Init ViewModel instances
 	for (const FMVVMViewClass_SourceCreator& Item : ClassExtension->GetViewModelCreators())
 	{
-		Item.CreateInstance(ClassExtension, this, GetUserWidget());
+		UUserWidget* UserWidget = GetUserWidget();
+		UObject* NewSource = Item.CreateInstance(ClassExtension, this, UserWidget);
+		if (NewSource)
+		{
+			if (Item.IsSourceAUserWidgetProperty())
+			{
+				FObjectPropertyBase* FoundObjectProperty = FindFProperty<FObjectPropertyBase>(UserWidget->GetClass(), Item.GetSourceName());
+				if (ensureAlwaysMsgf(FoundObjectProperty, TEXT("The compiler should have added the property")))
+				{
+					if (ensure(NewSource->GetClass()->IsChildOf(FoundObjectProperty->PropertyClass)))
+					{
+						FoundObjectProperty->SetObjectPropertyValue_InContainer(UserWidget, NewSource);
+					}
+				}
+			}
+		}
+
+		check(Item.GetSourceName().IsValid());
+		check(!FindViewSource(Item.GetSourceName()));
+
+		FMVVMViewSource& NewCreatedSource = Sources.AddDefaulted_GetRef();
+		NewCreatedSource.Source = NewSource;
+		NewCreatedSource.SourceName = Item.GetSourceName();
+		NewCreatedSource.bCreatedSource = true;
 	}
 
 #if UE_WITH_MVVM_DEBUGGING
@@ -74,17 +98,17 @@ void UMVVMView::Destruct()
 	UE::MVVM::FDebugging::BroadcastViewBeginDestruction(this);
 #endif
 
-	for (FRegisteredSource& Source : AllSources)
+	for (const FMVVMViewSource& Source : Sources)
 	{
-		UObject* SourceObject = Source.Source.Get();
-		if (Source.Count > 0 && SourceObject)
+		if (Source.RegisteredCout > 0 && Source.Source)
 		{
-			TScriptInterface<INotifyFieldValueChanged> SourceAsInterface = SourceObject;
+			TScriptInterface<INotifyFieldValueChanged> SourceAsInterface = Source.Source;
 			checkf(SourceAsInterface.GetInterface(), TEXT("It was added as a INotifyFieldValueChanged. It should still be."));
 			SourceAsInterface->RemoveAllFieldValueChangedDelegates(this);
 		}
 	}
-	AllSources.Reset();
+	// Todo, do we need to reset the FProperty also?
+	Sources.Reset(); // For GC release any object used by the view
 	EnabledLibraryBindings.Reset();
 
 	if (bHasEveryTickBinding)
@@ -95,76 +119,84 @@ void UMVVMView::Destruct()
 }
 
 
+TScriptInterface<INotifyFieldValueChanged> UMVVMView::GetViewModel(FName ViewModelName) const
+{
+	const FMVVMViewSource* RegisteredSource = FindViewSource(ViewModelName);
+	return RegisteredSource ? RegisteredSource->Source : nullptr;
+}
+
+
 bool UMVVMView::SetViewModel(FName ViewModelName, TScriptInterface<INotifyFieldValueChanged> NewValue)
 {
 	if (!ViewModelName.IsNone() && ClassExtension != nullptr)
 	{
-		FObjectPropertyBase* FoundObjectProperty = FindFProperty<FObjectPropertyBase>(GetUserWidget()->GetClass(), ViewModelName);
-		if (FoundObjectProperty == nullptr)
-		{
-			UE_LOG(LogMVVM, Error, TEXT("There is no viewmodel named '%s' in the view '%s'"), *ViewModelName.ToString(), *GetFullName());
-			return false;
-		}
-
+		const int32 AllCreatedSourcesIndex = ClassExtension->GetViewModelCreators().IndexOfByPredicate([ViewModelName](const FMVVMViewClass_SourceCreator& Item)
+			{
+				return Item.GetSourceName() == ViewModelName;
+			});
 		{
-			bool bFound = false;
-			for (const FMVVMViewClass_SourceCreator& Item : ClassExtension->GetViewModelCreators())
+			if (AllCreatedSourcesIndex == INDEX_NONE)
 			{
-				if (Item.GetSourcePropertyName() == ViewModelName)
-				{
-					bFound = true;
-					if (NewValue.GetObject() && !NewValue.GetObject()->GetClass()->IsChildOf(Item.GetSourceClass()))
-					{
-						UE::MVVM::FMessageLog Log(GetUserWidget());
-						Log.Error(LOCTEXT("SetViewModelInvalidValueType", "The new viewmodel is not of the expected type."));
-						UE_LOG(LogMVVM, Error, TEXT("The viewmodel name '%s' is invalid for the view '%s'"), *ViewModelName.ToString(), *GetFullName());
-						return false;
-					}
-					break;
-				}
+				UE::MVVM::FMessageLog Log(GetUserWidget());
+				Log.Error(LOCTEXT("SetViewModelViewModelNameNotFound", "The viewmodel name could not be found."));
+				return false;
 			}
-
-			if (!bFound)
+			else if (NewValue.GetObject()
+				&& !NewValue.GetObject()->GetClass()->IsChildOf(ClassExtension->GetViewModelCreators()[AllCreatedSourcesIndex].GetSourceClass()))
 			{
 				UE::MVVM::FMessageLog Log(GetUserWidget());
-				Log.Error(LOCTEXT("SetViewModelViewModelNameNotFound", "The viewmodel name could not be found."));
+				Log.Error(LOCTEXT("SetViewModelInvalidValueType", "The new viewmodel is not of the expected type."));
+				UE_LOG(LogMVVM, Error, TEXT("The viewmodel name '%s' is invalid for the view '%s'"), *ViewModelName.ToString(), *GetFullName());
+				return false;
+			}
+			else if (!Sources.IsValidIndex(AllCreatedSourcesIndex) || Sources[AllCreatedSourcesIndex].SourceName != ViewModelName)
+			{
+				UE::MVVM::FMessageLog Log(GetUserWidget());
+				Log.Error(LOCTEXT("SetViewModelInvalidCreatedSource", "An internal error occurs. The new viewmodel is not found."));
+				UE_LOG(LogMVVM, Error, TEXT("The viewmodel name '%s' is invalid for the view '%s'"), *ViewModelName.ToString(), *GetFullName());
 				return false;
 			}
 		}
 
-
-		UObject* PreviousValue = FoundObjectProperty->GetObjectPropertyValue_InContainer(GetUserWidget());
+		FMVVMViewSource& ViewSource = Sources[AllCreatedSourcesIndex];
+		UObject* PreviousValue = Sources[AllCreatedSourcesIndex].Source;
 		if (PreviousValue != NewValue.GetObject())
 		{
-			FMemMark Mark(FMemStack::Get());
-			TArray<int32, TMemStackAllocator<>> BindingToReenabled;
 			const TArrayView<const FMVVMViewClass_CompiledBinding> CompiledBindings = ClassExtension->GetCompiledBindings();
 
-			// Unregister any from that source
-			{
-				BindingToReenabled.Reserve(CompiledBindings.Num());
+			FMemMark Mark(FMemStack::Get());
+			TArray<int32, TMemStackAllocator<>> BindingToReenable;
+			BindingToReenable.Reserve(ViewSource.RegisteredCout);
 
-				for (int32 Index = 0; Index < CompiledBindings.Num(); ++Index)
+			// Unregister any bindings from that source
+			for (int32 Index = 0; Index < CompiledBindings.Num(); ++Index)
+			{
+				const FMVVMViewClass_CompiledBinding& Binding = CompiledBindings[Index];
+				if (Binding.GetSourceName() == ViewModelName && IsLibraryBindingEnabled(Index))
 				{
-					const FMVVMViewClass_CompiledBinding& Binding = CompiledBindings[Index];
-					if (Binding.GetSourceObjectPropertyName() == ViewModelName && IsLibraryBindingEnabled(Index))
-					{
-						DisableLibraryBinding(Binding, Index);
-						BindingToReenabled.Add(Index);
-					}
+					DisableLibraryBinding(Binding, Index);
+					BindingToReenable.Add(Index);
 				}
 			}
 
-			FoundObjectProperty->SetObjectPropertyValue_InContainer(GetUserWidget(), NewValue.GetObject());
+			ViewSource.Source = NewValue.GetObject();
+			if (ClassExtension->GetViewModelCreators()[AllCreatedSourcesIndex].IsSourceAUserWidgetProperty())
+			{
+				FObjectPropertyBase* FoundObjectProperty = FindFProperty<FObjectPropertyBase>(GetUserWidget()->GetClass(), ViewModelName);
+				if (ensure(FoundObjectProperty))
+				{
+					FoundObjectProperty->SetObjectPropertyValue_InContainer(GetUserWidget(), NewValue.GetObject());
+				}
+			}
 
 			bool bPreviousEveryTickBinding = bHasEveryTickBinding;
 			bHasEveryTickBinding = false;
 			if (NewValue.GetObject())
 			{
 				// Register back any binding that was previously enabled
-				if (BindingToReenabled.Num() > 0)
+				if (BindingToReenable.Num() > 0)
 				{
-					for (int32 Index : BindingToReenabled)
+					for (int32 Index : BindingToReenable)
 					{
 						const FMVVMViewClass_CompiledBinding& Binding = CompiledBindings[Index];
 						EnableLibraryBinding(Binding, Index);
@@ -176,7 +208,7 @@ bool UMVVMView::SetViewModel(FName ViewModelName, TScriptInterface<INotifyFieldV
 					for (int32 Index = 0; Index < CompiledBindings.Num(); ++Index)
 					{
 						const FMVVMViewClass_CompiledBinding& Binding = CompiledBindings[Index];
-						if (Binding.IsEnabledByDefault() && Binding.GetSourceObjectPropertyName() == ViewModelName)
+						if (Binding.IsEnabledByDefault() && Binding.GetSourceName() == ViewModelName)
 						{
 							EnableLibraryBinding(Binding, Index);
 						}
@@ -527,26 +559,32 @@ bool UMVVMView::RegisterLibraryBinding(const FMVVMViewClass_CompiledBinding& Bin
 
 
 	UE::FieldNotification::FFieldMulticastDelegate::FDelegate Delegate = UE::FieldNotification::FFieldMulticastDelegate::FDelegate::CreateUObject(this, &UMVVMView::HandledLibraryBindingValueChanged, BindingIndex);
-	bool bResult = Source->AddFieldValueChangedDelegate(FieldId, MoveTemp(Delegate)).IsValid();
-	if (bResult)
+	if (Source->AddFieldValueChangedDelegate(FieldId, MoveTemp(Delegate)).IsValid())
 	{
-		FRegisteredSource* FoundSource = AllSources.FindByPredicate([&Source](const FRegisteredSource& Other) { return Other.Source == Source.GetObject(); });
-		if (FoundSource)
+		if (FMVVMViewSource* RegisteredSource = FindViewSource(Binding.GetSourceName()))
 		{
-			++FoundSource->Count;
+			++RegisteredSource->RegisteredCout;
 		}
 		else
 		{
-			FRegisteredSource Item;
-			Item.Source = Source.GetObject();
-			Item.Count = 1;
-			AllSources.Add(MoveTemp(Item));
+			FMVVMViewSource& NewSource = Sources.AddDefaulted_GetRef();
+			NewSource.Source = Source.GetObject();
+			NewSource.SourceName = Binding.GetSourceName();
+			NewSource.RegisteredCout = 1;
 		}
+
+#if UE_WITH_MVVM_DEBUGGING
+		UE::MVVM::FDebugging::BroadcastLibraryBindingRegistered(this, Binding, UE::MVVM::FDebugging::ERegisterLibraryBindingResult::Success);
+#endif
 	}
+	else
+	{
 #if UE_WITH_MVVM_DEBUGGING
-	UE::MVVM::FDebugging::BroadcastLibraryBindingRegistered(this, Binding, UE::MVVM::FDebugging::ERegisterLibraryBindingResult::Success);
+		UE::MVVM::FDebugging::BroadcastLibraryBindingRegistered(this, Binding, UE::MVVM::FDebugging::ERegisterLibraryBindingResult::Success);
 #endif
-	return bResult;
+	}
+
+	return true;
 }
 
 
@@ -561,16 +599,11 @@ void UMVVMView::UnregisterLibraryBinding(const FMVVMViewClass_CompiledBinding& B
 			Source->RemoveAllFieldValueChangedDelegates(FieldIdResult.GetValue(), this);
 		}
 
-		int32 FoundIndex = AllSources.IndexOfByPredicate([&Source](const FRegisteredSource& Other) { return Other.Source == Source.GetObject(); });
-		if (ensureMsgf(FoundIndex != INDEX_NONE, TEXT("If the binding was enabled, then the source should also be there.")))
+		FMVVMViewSource* RegisteredSource = FindViewSource(Binding.GetSourceName());
+		if (ensureMsgf(RegisteredSource, TEXT("If the binding was enabled, then the source should also be there.")))
 		{
-			FRegisteredSource& FoundSource = AllSources[FoundIndex];
-			ensureMsgf(FoundSource.Count > 0, TEXT("The count should match the number of RegisterLibraryBinding and UnregisterLibraryBinding"));
-			--FoundSource.Count;
-			if (FoundSource.Count == 0)
-			{
-				AllSources.RemoveAtSwap(FoundIndex);
-			}
+			ensureMsgf(RegisteredSource->RegisteredCout > 0, TEXT("The count should match the number of RegisterLibraryBinding and UnregisterLibraryBinding"));
+			--RegisteredSource->RegisteredCout;
 		}
 	}
 }
@@ -582,25 +615,35 @@ TScriptInterface<INotifyFieldValueChanged> UMVVMView::FindSource(const FMVVMView
 	check(UserWidget);
 	check(ClassExtension);
 
-	if (Binding.IsSourceObjectItself())
+	if (Binding.IsSourceUserWidget())
 	{
 		return TScriptInterface<INotifyFieldValueChanged>(UserWidget);
 	}
 	else
 	{
-		const FObjectPropertyBase* SourceObjectProperty = CastField<FObjectPropertyBase>(UserWidget->GetClass()->FindPropertyByName(Binding.GetSourceObjectPropertyName()));
-		if (SourceObjectProperty == nullptr)
+		UObject* Source = nullptr;
+
+		if (const FMVVMViewSource* RegisteredSource = FindViewSource(Binding.GetSourceName()))
 		{
-			UE_LOG(LogMVVM, Error, TEXT("'%s' could not evaluate the source for binding '%s'. The property name is invalid."), *GetFullName(), *Binding.ToString());
-			return TScriptInterface<INotifyFieldValueChanged>();
+			Source = RegisteredSource->Source.Get();
+		}
+		else
+		{
+			const FObjectPropertyBase* SourceObjectProperty = CastField<FObjectPropertyBase>(UserWidget->GetClass()->FindPropertyByName(Binding.GetSourceName()));
+			if (SourceObjectProperty == nullptr)
+			{
+				UE_LOG(LogMVVM, Error, TEXT("'%s' could not evaluate the source for binding '%s'. The property name is invalid."), *GetFullName(), *Binding.ToString());
+				return TScriptInterface<INotifyFieldValueChanged>();
+			}
+
+			Source = SourceObjectProperty->GetObjectPropertyValue_InContainer(UserWidget);
 		}
 
-		UObject* Source = SourceObjectProperty->GetObjectPropertyValue_InContainer(UserWidget);
 		if (Source == nullptr)
 		{
 			if (!bAllowNull)
 			{
-				UE_LOG(LogMVVM, Error, TEXT("'%s' could not evaluate the source for binding '%s'. The path point to an invalid object."), *GetFullName(), *Binding.ToString());
+				UE_LOG(LogMVVM, Error, TEXT("'%s' could not evaluate the source for binding '%s'. The source is an invalid object."), *GetFullName(), *Binding.ToString());
 			}
 			return TScriptInterface<INotifyFieldValueChanged>();
 		}
@@ -616,5 +659,17 @@ TScriptInterface<INotifyFieldValueChanged> UMVVMView::FindSource(const FMVVMView
 }
 
 
+FMVVMViewSource* UMVVMView::FindViewSource(const FName SourceName)
+{
+	return Sources.FindByPredicate([SourceName](const FMVVMViewSource& Other){ return Other.SourceName == SourceName; });
+}
+
+
+const FMVVMViewSource* UMVVMView::FindViewSource(const FName SourceName) const
+{
+	return Sources.FindByPredicate([SourceName](const FMVVMViewSource& Other) { return Other.SourceName == SourceName; });
+}
+
+
 #undef LOCTEXT_NAMESPACE
 
diff --git a/Engine/Plugins/Runtime/ModelViewViewModel/Source/ModelViewViewModel/Private/View/MVVMViewClass.cpp b/Engine/Plugins/Runtime/ModelViewViewModel/Source/ModelViewViewModel/Private/View/MVVMViewClass.cpp
--- a/Engine/Plugins/Runtime/ModelViewViewModel/Source/ModelViewViewModel/Private/View/MVVMViewClass.cpp
+++ b/Engine/Plugins/Runtime/ModelViewViewModel/Source/ModelViewViewModel/Private/View/MVVMViewClass.cpp
@@ -102,99 +102,85 @@ UObject* FMVVMViewClass_SourceCreator::CreateInstance(const UMVVMViewClass* InVi
 
 	UObject* Result = nullptr;
 
-	FObjectPropertyBase* FoundObjectProperty = FindFProperty<FObjectPropertyBase>(InUserWidget->GetClass(), PropertyName);
-	if (ensureAlwaysMsgf(FoundObjectProperty, TEXT("The compiler should have added the property")))
+	if (bCreateInstance)
 	{
-		auto AssignProperty = [FoundObjectProperty, InUserWidget](UObject* NewObject)
+		if (ExpectedSourceType.Get() != nullptr)
 		{
-			if (NewObject && NewObject->GetClass()->IsChildOf(FoundObjectProperty->PropertyClass))
-			{
-				FoundObjectProperty->SetObjectPropertyValue_InContainer(InUserWidget, NewObject);
-			}
-		};
-
-		if (bCreateInstance)
+			Result = NewObject<UObject>(InUserWidget, ExpectedSourceType.Get(), NAME_None, RF_Transient);
+		}
+		else if (!bOptional)
 		{
-			if (ExpectedSourceType.Get() != nullptr)
-			{
-				Result = NewObject<UObject>(InUserWidget, ExpectedSourceType.Get(), NAME_None, RF_Transient);
-			}
-			else if (!bOptional)
-			{
-				UE::MVVM::FMessageLog Log(InUserWidget);
-				Log.Error(FText::Format(LOCTEXT("CreateInstanceCreateInstance", "The source '{0}' could not be created. The class is not loaded."), FText::FromName(PropertyName)));
-			}
+			UE::MVVM::FMessageLog Log(InUserWidget);
+			Log.Error(FText::Format(LOCTEXT("CreateInstanceCreateInstance", "The source '{0}' could not be created. The class is not loaded."), FText::FromName(PropertyName)));
 		}
-		else if (Resolver)
+	}
+	else if (Resolver)
+	{
+		Result = Resolver->CreateInstance(ExpectedSourceType.Get(), InUserWidget, InView);
+		if (!Result && !bOptional)
 		{
-			Result = Resolver->CreateInstance(ExpectedSourceType.Get(), InUserWidget, InView);
-			if (!Result && !bOptional)
-			{
-				UE::MVVM::FMessageLog Log(InUserWidget);
-				Log.Error(FText::Format(LOCTEXT("CreateInstanceFailResolver", "The source '{0}' could not be created. Resolver returned an invalid value."), FText::FromName(PropertyName)));
-			}
-			if (Result && !Result->IsA(ExpectedSourceType.Get()))
-			{
-				Result = nullptr;
+			UE::MVVM::FMessageLog Log(InUserWidget);
+			Log.Error(FText::Format(LOCTEXT("CreateInstanceFailResolver", "The source '{0}' could not be created. Resolver returned an invalid value."), FText::FromName(PropertyName)));
+		}
+		if (Result && !Result->IsA(ExpectedSourceType.Get()))
+		{
+			Result = nullptr;
 
-				UE::MVVM::FMessageLog Log(InUserWidget);
-				Log.Error(FText::Format(LOCTEXT("CreateInstanceFailResolverExpected", "The source '{0}' could not be created. Resolver returned viewodel of an unexpected type."), FText::FromName(PropertyName)));
-			}
+			UE::MVVM::FMessageLog Log(InUserWidget);
+			Log.Error(FText::Format(LOCTEXT("CreateInstanceFailResolverExpected", "The source '{0}' could not be created. Resolver returned viewodel of an unexpected type."), FText::FromName(PropertyName)));
 		}
-		else if (GlobalViewModelInstance.IsValid())
+	}
+	else if (GlobalViewModelInstance.IsValid())
+	{
+		UMVVMViewModelCollectionObject* Collection = nullptr;
+		UMVVMViewModelBase* FoundViewModelInstance = nullptr;
+		if (const UWorld* World = InUserWidget->GetWorld())
 		{
-			UMVVMViewModelCollectionObject* Collection = nullptr;
-			UMVVMViewModelBase* FoundViewModelInstance = nullptr;
-			if (const UWorld* World = InUserWidget->GetWorld())
+			if (const UGameInstance* GameInstance = World->GetGameInstance())
 			{
-				if (const UGameInstance* GameInstance = World->GetGameInstance())
+				Collection = GameInstance->GetSubsystem<UMVVMGameSubsystem>()->GetViewModelCollection();
+				if (Collection)
 				{
-					Collection = GameInstance->GetSubsystem<UMVVMGameSubsystem>()->GetViewModelCollection();
-					if (Collection)
-					{
-						FoundViewModelInstance = Collection->FindViewModelInstance(GlobalViewModelInstance);
-					}
+					FoundViewModelInstance = Collection->FindViewModelInstance(GlobalViewModelInstance);
 				}
 			}
+		}
 
-			if (FoundViewModelInstance != nullptr)
+		if (FoundViewModelInstance != nullptr)
+		{
+			ensureMsgf(FoundViewModelInstance->IsA(GlobalViewModelInstance.ContextClass), TEXT("The Global View Model Instance is not of the expected type."));
+			Result = FoundViewModelInstance;
+		}
+		else if (!bOptional)
+		{
+			UE::MVVM::FMessageLog Log(InUserWidget);
+			if (Collection)
 			{
-				ensureMsgf(FoundViewModelInstance->IsA(GlobalViewModelInstance.ContextClass), TEXT("The Global View Model Instance is not of the expected type."));
-				Result = FoundViewModelInstance;
+				Log.Error(FText::Format(LOCTEXT("CreateInstanceFailedGlobal", "The source '{0}' was not found in the global view model collection."), FText::FromName(GlobalViewModelInstance.ContextName)));
 			}
-			else if (!bOptional)
+			else
 			{
-				UE::MVVM::FMessageLog Log(InUserWidget);
-				if (Collection)
-				{
-					Log.Error(FText::Format(LOCTEXT("CreateInstanceFailedGlobal", "The source '{0}' was not found in the global view model collection."), FText::FromName(GlobalViewModelInstance.ContextName)));
-				}
-				else
-				{
-					Log.Error(FText::Format(LOCTEXT("CreateInstanceFailedGlobalInstance", "The source '{0}' will be invalid because the global view model collection could not be found."), FText::FromName(GlobalViewModelInstance.ContextName)));
-				}
+				Log.Error(FText::Format(LOCTEXT("CreateInstanceFailedGlobalInstance", "The source '{0}' will be invalid because the global view model collection could not be found."), FText::FromName(GlobalViewModelInstance.ContextName)));
 			}
 		}
-		else if (FieldPath.IsValid())
+	}
+	else if (FieldPath.IsValid())
+	{
+		TValueOrError<UE::MVVM::FFieldContext, void> FieldPathResult = InViewClass->GetBindingLibrary().EvaluateFieldPath(InUserWidget, FieldPath);
+		if (FieldPathResult.HasValue())
 		{
-			TValueOrError<UE::MVVM::FFieldContext, void> FieldPathResult = InViewClass->GetBindingLibrary().EvaluateFieldPath(InUserWidget, FieldPath);
-			if (FieldPathResult.HasValue())
+			TValueOrError<UObject*, void> ObjectResult = UE::MVVM::FieldPathHelper::EvaluateObjectProperty(FieldPathResult.GetValue());
+			if (ObjectResult.HasValue() && ObjectResult.GetValue() != nullptr)
 			{
-				TValueOrError<UObject*, void> ObjectResult = UE::MVVM::FieldPathHelper::EvaluateObjectProperty(FieldPathResult.GetValue());
-				if (ObjectResult.HasValue() && ObjectResult.GetValue() != nullptr)
-				{
-					Result = ObjectResult.GetValue();
-				}
-			}
-
-			if (Result == nullptr && !bOptional)
-			{
-				UE::MVVM::FMessageLog Log(InUserWidget);
-				Log.Error(FText::Format(LOCTEXT("CreateInstanceInvalidBiding", "The source '{0}' was evaluated to be invalid at initialization."), FText::FromName(PropertyName)));
+				Result = ObjectResult.GetValue();
 			}
 		}
 
-		AssignProperty(Result);
+		if (Result == nullptr && !bOptional)
+		{
+			UE::MVVM::FMessageLog Log(InUserWidget);
+			Log.Error(FText::Format(LOCTEXT("CreateInstanceInvalidBiding", "The source '{0}' was evaluated to be invalid at initialization."), FText::FromName(PropertyName)));
+		}
 	}
 
 	return Result;
@@ -222,7 +208,13 @@ EMVVMExecutionMode FMVVMViewClass_CompiledBinding::GetExecuteMode() const
 
 FString FMVVMViewClass_CompiledBinding::ToString() const
 {
-	return TEXT("[PH]");
+	return FString::Printf(TEXT("FieldId: '%d', PropertyName: '%s', Mode: '%s', Flags: '%d' \n")
+			TEXT("Binding: '%s'")
+		, TEXT("")
+		, *SourcePropertyName.ToString()
+		, *StaticEnum<EMVVMExecutionMode>()->GetNameByValue((int64)ExecutionMode).ToString()
+		, Flags
+		, TEXT(""));
 }
 
 
diff --git a/Engine/Plugins/Runtime/ModelViewViewModel/Source/ModelViewViewModel/Public/Debugging/MVVMDebugging.h b/Engine/Plugins/Runtime/ModelViewViewModel/Source/ModelViewViewModel/Public/Debugging/MVVMDebugging.h
--- a/Engine/Plugins/Runtime/ModelViewViewModel/Source/ModelViewViewModel/Public/Debugging/MVVMDebugging.h
+++ b/Engine/Plugins/Runtime/ModelViewViewModel/Source/ModelViewViewModel/Public/Debugging/MVVMDebugging.h
@@ -73,6 +73,7 @@ public:
 		Failed_InvalidSource,
 		Failed_InvalidSourceField,
 		Failed_InvalidDestinationField,
+		Failed_Unknown,
 	};
 
 	struct FLibraryBindingRegisteredArgs
diff --git a/Engine/Plugins/Runtime/ModelViewViewModel/Source/ModelViewViewModel/Public/View/MVVMView.h b/Engine/Plugins/Runtime/ModelViewViewModel/Source/ModelViewViewModel/Public/View/MVVMView.h
--- a/Engine/Plugins/Runtime/ModelViewViewModel/Source/ModelViewViewModel/Public/View/MVVMView.h
+++ b/Engine/Plugins/Runtime/ModelViewViewModel/Source/ModelViewViewModel/Public/View/MVVMView.h
@@ -21,6 +21,19 @@ namespace UE::MVVM
 	class FDebugging;
 }
 
+USTRUCT()
+struct FMVVMViewSource
+{
+	GENERATED_BODY()
+
+	UPROPERTY()
+	TObjectPtr<UObject> Source; // TScriptInterface<INotifyFieldValueChanged>
+
+	FName SourceName;
+	int32 RegisteredCout = 0;
+	bool bCreatedSource = false;
+};
+
 /**
  * Instance UMVVMClassExtension_View for the UUserWdiget
  */
@@ -54,6 +67,9 @@ public:
 	//UFUNCTION(BlueprintCallable, Category = "Viewmodel")
 	//bool IsLibraryBindingEnabled(FGuid ViewModelId, FMVVMBindingName BindingName) const;
 
+	UFUNCTION(BlueprintCallable, Category = "Viewmodel")
+	TScriptInterface<INotifyFieldValueChanged> GetViewModel(FName ViewModelName) const;
+
 	UFUNCTION(BlueprintCallable, Category = "Viewmodel")
 	bool SetViewModel(FName ViewModelName, TScriptInterface<INotifyFieldValueChanged> ViewModel);
 
@@ -71,23 +87,19 @@ private:
 	void UnregisterLibraryBinding(const FMVVMViewClass_CompiledBinding& Binding);
 
 	TScriptInterface<INotifyFieldValueChanged> FindSource(const FMVVMViewClass_CompiledBinding& Binding, bool bAllowNull) const;
+	FMVVMViewSource* FindViewSource(const FName SourceName);
+	const FMVVMViewSource* FindViewSource(const FName SourceName) const;
 
 private:
 // todo support dynamic runtime binding.
 	/** Binding that are added dynamically at runtime. */
 	//TArray<FMVVMView_Binding> RegisteredDynamicBindings;
 
-	struct FRegisteredSource
-	{
-		TWeakObjectPtr<UObject> Source;
-		int32 Count = 0;
-	};
-
 	UPROPERTY(Transient)
 	TObjectPtr<const UMVVMViewClass> ClassExtension;
 
-	/** A list of all source and the amount of registered binding. */
-	TArray<FRegisteredSource> AllSources;
+	UPROPERTY(Transient)
+	TArray<FMVVMViewSource> Sources;
 
 	/** The binding that are enabled for the instance. */
 	TBitArray<> EnabledLibraryBindings;
@@ -98,7 +110,7 @@ private:
 	/** Is the Construct method was called. */
 	bool bConstructed = false;
 
-	/** Is the Construct method was called. */
+	/** The view has at least one binding that need to be ticked every frame. */
 	bool bHasEveryTickBinding = false;
 };
 
diff --git a/Engine/Plugins/Runtime/ModelViewViewModel/Source/ModelViewViewModel/Public/View/MVVMViewClass.h b/Engine/Plugins/Runtime/ModelViewViewModel/Source/ModelViewViewModel/Public/View/MVVMViewClass.h
--- a/Engine/Plugins/Runtime/ModelViewViewModel/Source/ModelViewViewModel/Public/View/MVVMViewClass.h
+++ b/Engine/Plugins/Runtime/ModelViewViewModel/Source/ModelViewViewModel/Public/View/MVVMViewClass.h
@@ -48,7 +48,12 @@ public:
 		return ExpectedSourceType.Get();
 	}
 
-	FName GetSourcePropertyName() const
+	bool IsSourceAUserWidgetProperty() const
+	{
+		return bIsUserWidgetProperty;
+	}
+
+	FName GetSourceName() const
 	{
 		return PropertyName;
 	}
@@ -80,6 +85,9 @@ private:
 	UPROPERTY()
 	bool bCreateInstance = false;
 
+	UPROPERTY()
+	bool bIsUserWidgetProperty = true;
+
 	UPROPERTY()
 	bool bOptional = false;
 };
@@ -103,16 +111,17 @@ public:
 	}
 
 	/**
-	 * @return The property name of the object that contains the SourceFieldId.
+	 * @return The unique name of the source object that contains the SourceFieldId.
 	 * It implements INotifyFieldValueChanged if it's not a One Time.
+	 * It can also be the name of a FProperty on the source object.
 	 */
-	FName GetSourceObjectPropertyName() const
+	FName GetSourceName() const
 	{
 		return SourcePropertyName;
 	}
 
-	/** @return true if SourceFieldId is not from a property but on the object itserf (UserWidget). */
-	bool IsSourceObjectItself() const
+	/** @return true if SourceFieldId is not from a property but the object is the UserWidget. */
+	bool IsSourceUserWidget() const
 	{
 		return (Flags & EBindingFlags::SourceObjectIsSelf) != 0;
 	}
diff --git a/Engine/Plugins/Runtime/ModelViewViewModel/Source/ModelViewViewModelBlueprint/Private/MVVMViewBlueprintCompiler.cpp b/Engine/Plugins/Runtime/ModelViewViewModel/Source/ModelViewViewModelBlueprint/Private/MVVMViewBlueprintCompiler.cpp
--- a/Engine/Plugins/Runtime/ModelViewViewModel/Source/ModelViewViewModelBlueprint/Private/MVVMViewBlueprintCompiler.cpp
+++ b/Engine/Plugins/Runtime/ModelViewViewModel/Source/ModelViewViewModelBlueprint/Private/MVVMViewBlueprintCompiler.cpp
@@ -205,6 +205,10 @@ void FMVVMViewBlueprintCompiler::CreateVariables(const FWidgetBlueprintCompilerC
 			{
 				NewProperty->SetMetaData(FBlueprintMetadata::MD_ExposeOnSpawn, TEXT("true"));
 			}
+			if (!SourceContext.bPublicGetter)
+			{
+				NewProperty->SetMetaData(FBlueprintMetadata::MD_Private, TEXT("true"));
+			}
 #endif
 		}
 		return NewProperty;
@@ -301,6 +305,7 @@ void FMVVMViewBlueprintCompiler::CreateSourceLists(const FWidgetBlueprintCompile
 {
 	CompilerUserWidgetPropertyContexts.Reset();
 	CompilerSourceCreatorContexts.Reset();
+	ViewModelDynamicCounter = 0;
 
 	TSet<FGuid> ViewModelGuids;
 	TSet<FName> WidgetSources;
@@ -359,6 +364,7 @@ void FMVVMViewBlueprintCompiler::CreateSourceLists(const FWidgetBlueprintCompile
 		SourceVariable.DisplayName = ViewModelContext.GetDisplayName();
 		SourceVariable.CategoryName = TEXT("Viewmodel");
 		SourceVariable.bExposeOnSpawn = bCreateSetterFunction;
+		//SourceVariable.bPublicGetter = ViewModelContext.bCreateGetterFunction;
 		SourceVariable.BlueprintSetter = CompilerSourceCreatorContexts[FoundSourceCreatorContextIndex].SetterFunctionName;
 		SourceVariable.ViewModelId = ViewModelContext.GetViewModelId();
 		CompilerUserWidgetPropertyContexts.Emplace(MoveTemp(SourceVariable));
@@ -431,6 +437,7 @@ void FMVVMViewBlueprintCompiler::CreateSourceLists(const FWidgetBlueprintCompile
 					SourceVariable.DisplayName = FText::FromString(Widget->GetDisplayLabel());
 					SourceVariable.CategoryName = TEXT("Widget");
 					SourceVariable.ViewModelId = FGuid();
+					SourceVariable.bPublicGetter = false;
 					Self->CompilerUserWidgetPropertyContexts.Emplace(MoveTemp(SourceVariable));
 				}
 			}
@@ -679,7 +686,7 @@ bool FMVVMViewBlueprintCompiler::PreCompileBindingSources(UWidgetBlueprintGenera
 				return false;
 			}
 
-			if (!BindingSourceContext.bIsRootWidget && BindingSourceContext.UserWidgetPropertyContextIndex == INDEX_NONE)
+			if (!BindingSourceContext.bIsRootWidget && BindingSourceContext.UserWidgetPropertyContextIndex == INDEX_NONE && BindingSourceContext.SourceCreatorContextIndex == INDEX_NONE)
 			{
 				AddErrorForBinding(Binding, LOCTEXT("BindingInvalidSource", "Internal error. The binding could not find its source."), ArgumentName);
 				return false;
@@ -1070,7 +1077,7 @@ bool FMVVMViewBlueprintCompiler::PreCompileBindings(UWidgetBlueprintGeneratedCla
 		const UFunction* ConversionFunction = ConversionFunctionReference.ResolveMember<UFunction>(Class);
 		if (!ConversionFunctionWrapper.IsNone() && ConversionFunction == nullptr)
 		{
-			AddErrorForBinding(Binding, FText::Format(LOCTEXT("ConversionFunctionNotFound", "The conversion function '{0}' could not be found."), 
+			AddErrorForBinding(Binding, FText::Format(LOCTEXT("ConversionFunctionNotFound", "The conversion function '{0}' could not be found."),
 				FText::FromName(ConversionFunctionReference.GetMemberName()))
 			);
 			bIsBindingsValid = false;
@@ -1080,7 +1087,7 @@ bool FMVVMViewBlueprintCompiler::PreCompileBindings(UWidgetBlueprintGeneratedCla
 		TValueOrError<FCompilerBinding, FText> AddBindingResult = AddBinding(Class, BindingSourceContext.PropertyPath, SetterPath, ConversionFunction);
 		if (AddBindingResult.HasError())
 		{
-			AddErrorForBinding(Binding, 
+			AddErrorForBinding(Binding,
 				FText::Format(LOCTEXT("CouldNotCreateBinding", "Could not create binding. {0}"), AddBindingResult.GetError())
 			);
 			bIsBindingsValid = false;
@@ -1090,6 +1097,7 @@ bool FMVVMViewBlueprintCompiler::PreCompileBindings(UWidgetBlueprintGeneratedCla
 		FCompilerBinding NewBinding = AddBindingResult.StealValue();
 		NewBinding.BindingIndex = BindingSourceContext.BindingIndex;
 		NewBinding.UserWidgetPropertyContextIndex = BindingSourceContext.UserWidgetPropertyContextIndex;
+		NewBinding.SourceCreatorContextIndex = BindingSourceContext.SourceCreatorContextIndex;
 		NewBinding.FieldIdHandle = AddFieldResult.StealValue();
 		NewBinding.bSourceIsUserWidget = BindingSourceContext.bIsRootWidget;
 		NewBinding.bFieldIdNeeded = !IsOneTimeBinding(Binding.BindingType);
@@ -1119,10 +1127,22 @@ bool FMVVMViewBlueprintCompiler::CompileBindings(const FCompiledBindingLibraryCo
 		FMVVMBlueprintViewBinding& ViewBinding = *BlueprintView->GetBindingAt(CompileBinding.BindingIndex);
 
 		FMVVMViewClass_CompiledBinding NewBinding;
-		const bool bIsSourceSelf = !CompilerUserWidgetPropertyContexts.IsValidIndex(CompileBinding.UserWidgetPropertyContextIndex);
+		const bool bIsSourceSelf = CompileBinding.bSourceIsUserWidget;
 		if (!bIsSourceSelf)
 		{
-			NewBinding.SourcePropertyName = CompilerUserWidgetPropertyContexts[CompileBinding.UserWidgetPropertyContextIndex].PropertyName;
+			if (CompilerUserWidgetPropertyContexts.IsValidIndex(CompileBinding.UserWidgetPropertyContextIndex))
+			{
+				NewBinding.SourcePropertyName = CompilerUserWidgetPropertyContexts[CompileBinding.UserWidgetPropertyContextIndex].PropertyName;
+			}
+			else if (CompilerSourceCreatorContexts.IsValidIndex(CompileBinding.SourceCreatorContextIndex))
+			{
+				NewBinding.SourcePropertyName = CompilerSourceCreatorContexts[CompileBinding.SourceCreatorContextIndex].ViewModelContext.GetViewModelName();
+			}
+			else
+			{
+				WidgetBlueprintCompilerContext.MessageLog.Error(*LOCTEXT("InvalidSourceInternal", "Internal error. The binding doesn't have a valid source.").ToString());
+				return false;
+			}
 		}
 
 		const FMVVMVCompiledFieldId* CompiledFieldId = CompileResult.FieldIds.Find(CompileBinding.FieldIdHandle);
@@ -1146,8 +1166,16 @@ bool FMVVMViewBlueprintCompiler::CompileBindings(const FCompiledBindingLibraryCo
 		bool bIsOptional = false;
 		if (!bIsSourceSelf)
 		{
-			check(CompilerUserWidgetPropertyContexts.IsValidIndex(CompileBinding.UserWidgetPropertyContextIndex));
-			const FGuid ViewModelId = CompilerUserWidgetPropertyContexts[CompileBinding.UserWidgetPropertyContextIndex].ViewModelId;
+			FGuid ViewModelId;
+			if (CompilerUserWidgetPropertyContexts.IsValidIndex(CompileBinding.UserWidgetPropertyContextIndex))
+			{
+				ViewModelId = CompilerUserWidgetPropertyContexts[CompileBinding.UserWidgetPropertyContextIndex].ViewModelId;
+			}
+			else if (CompilerSourceCreatorContexts.IsValidIndex(CompileBinding.SourceCreatorContextIndex))
+			{
+				ViewModelId = CompilerSourceCreatorContexts[CompileBinding.SourceCreatorContextIndex].ViewModelContext.GetViewModelId();
+			}
+
 			if (ViewModelId.IsValid())
 			{
 				const FMVVMBlueprintViewModelContext* ModelContext = BlueprintView->FindViewModel(ViewModelId);
@@ -1158,7 +1186,7 @@ bool FMVVMViewBlueprintCompiler::CompileBindings(const FCompiledBindingLibraryCo
 					continue;
 				}
 
-				bIsOptional = ModelContext->IsValid();
+				bIsOptional = ModelContext->bOptional;
 			}
 		}
 
@@ -1254,31 +1282,50 @@ TValueOrError<FMVVMViewBlueprintCompiler::FBindingSourceContext, FText> FMVVMVie
 	if (bCreateADynamicSourceEntry)
 	{
 		return MakeError(LOCTEXT("DynamicSourceEntryNotSupport", "Long source entry is not yet fully supported. Add the viewmodel manually."));
-		
-		//Result.SourceClass = BindingInfo.NotifyFieldClass;
-		//Result.UserWidgetPropertyContextIndex = INDEX_NONE;
-		//Result.bIsRootWidget = false;
-
-		//FString InterfacePath = FieldPathHelper::ToString(BindingInfo.NotifyFieldInterfacePath);
-		//int32 SourceCreatorContextIndex = CompilerSourceCreatorContexts.IndexOfByPredicate([&InterfacePath](const FCompilerSourceCreatorContext& Other)
-		//	{
-		//		return (Other.Type == ECompilerSourceCreatorType::ViewModel)
-		//			? (Other.ViewModelContext.CreationType == EMVVMBlueprintViewModelContextCreationType::PropertyPath
-		//				&& Other.ViewModelContext.ViewModelPropertyPath == InterfacePath)
-		//			: false;
-		//	});
-
-		//// Add if the path doesn't already exists
-		//if (!CompilerSourceCreatorContexts.IsValidIndex(SourceCreatorContextIndex))
-		//{
-		//	return MakeError(LOCTEXT("ViewModelNotAtRoot", "The field in the path is from a viewmodel that is not the root of the path. This is not yet supported."));
-		//}
-
-		//Result.PropertyPath.RemoveAt(0, BindingInfo.ViewModelIndex);
-		//if (CompilerSourceCreatorContexts[SourceCreatorContextIndex].Type == ECompilerSourceCreatorType::ViewModel)
-		//{
-		//	Result.UserWidgetPropertyContextIndex = SourceCreatorContextIndex;
-		//}
+
+		if (!Result.PropertyPath.IsValidIndex(BindingInfo.ViewModelIndex))
+		{
+			return MakeError(LOCTEXT("DynamicSourceEntryInternalIndex", "Internal error. The source index is not valid."));
+		}
+		//for (0 to index) if they are all interface, then create new BindingCompilerSourceCreatorContexts
+		//create new dynamic vm with property path to access it
+		TStringBuilder<512> PropertyPathBuilder;
+		for (int32 Index = 0; Index < BindingInfo.ViewModelIndex; ++Index)
+		{
+			if (Index > 0)
+			{
+				PropertyPathBuilder << TEXT('.');
+			}
+			PropertyPathBuilder << Result.PropertyPath[Index].GetName();
+		}
+
+		int32 SourceCreatorContextIndex = CompilerSourceCreatorContexts.IndexOfByPredicate([&PropertyPathBuilder](const FCompilerSourceCreatorContext& Other)
+			{
+				return Other.Type == ECompilerSourceCreatorType::ViewModelDynamic
+					&& Other.ViewModelContext.ViewModelPropertyPath == PropertyPathBuilder.ToString();
+			});
+
+		if (SourceCreatorContextIndex == INDEX_NONE)
+		{
+			FCompilerSourceCreatorContext SourceCreatorContext;
+			SourceCreatorContext.Type = ECompilerSourceCreatorType::ViewModelDynamic;
+			SourceCreatorContext.ViewModelContext = FMVVMBlueprintViewModelContext(BindingInfo.NotifyFieldClass, *FString::Printf(TEXT("__Dynamic__%d"), ViewModelDynamicCounter));
+			SourceCreatorContext.ViewModelContext.bCreateSetterFunction = false;
+			SourceCreatorContext.ViewModelContext.bOptional = false;
+			SourceCreatorContext.ViewModelContext.CreationType = EMVVMBlueprintViewModelContextCreationType::PropertyPath;
+			SourceCreatorContext.ViewModelContext.ViewModelPropertyPath = PropertyPathBuilder.ToString();
+			++ViewModelDynamicCounter;
+
+			SourceCreatorContextIndex = CompilerSourceCreatorContexts.Add(MoveTemp(SourceCreatorContext));
+		}
+
+		Result.SourceClass = BindingInfo.NotifyFieldClass;
+		Result.FieldId = BindingInfo.NotifyFieldId;
+		int32 NumberToRemove = Result.PropertyPath.Num() - BindingInfo.ViewModelIndex;
+		Result.PropertyPath.RemoveAt(BindingInfo.ViewModelIndex, NumberToRemove);
+		Result.UserWidgetPropertyContextIndex = INDEX_NONE;
+		Result.SourceCreatorContextIndex = SourceCreatorContextIndex;
+		Result.bIsRootWidget = false;
 	}
 
 	return MakeValue(Result);
diff --git a/Engine/Plugins/Runtime/ModelViewViewModel/Source/ModelViewViewModelBlueprint/Private/MVVMViewBlueprintCompiler.h b/Engine/Plugins/Runtime/ModelViewViewModel/Source/ModelViewViewModelBlueprint/Private/MVVMViewBlueprintCompiler.h
--- a/Engine/Plugins/Runtime/ModelViewViewModel/Source/ModelViewViewModelBlueprint/Private/MVVMViewBlueprintCompiler.h
+++ b/Engine/Plugins/Runtime/ModelViewViewModel/Source/ModelViewViewModelBlueprint/Private/MVVMViewBlueprintCompiler.h
@@ -80,9 +80,8 @@ private:
 
 private:
 	/**
-	 * Describe the Property that is needed to start a PropertyPath on the UserWidget.
+	 * Describe the Property that need to be added (if it doesn't already exist)
 	 * They can be a Widget, a viewmodel, or any object owned by the UserWidget.
-	 * The property will be created if it doesn't exist yet.
 	 */
 	struct FCompilerUserWidgetPropertyContext
 	{
@@ -96,6 +95,7 @@ private:
 		FGuid ViewModelId;
 
 		bool bExposeOnSpawn = false;
+		bool bPublicGetter = false;
 	};
 	TArray<FCompilerUserWidgetPropertyContext> CompilerUserWidgetPropertyContexts;
 
@@ -104,7 +104,8 @@ private:
 	 */
 	enum class ECompilerSourceCreatorType
 	{
-		ViewModel,
+		ViewModel, // added manually
+		ViewModelDynamic, // added by a long binding chain
 	};
 	struct FCompilerSourceCreatorContext
 	{
@@ -115,6 +116,7 @@ private:
 		UEdGraph* SetterGraph = nullptr;
 	};
 	TArray<FCompilerSourceCreatorContext> CompilerSourceCreatorContexts;
+	int32 ViewModelDynamicCounter = 0;
 
 	/** 
 	 * Describe a binding for the compiler.
@@ -128,6 +130,7 @@ private:
 	{
 		int32 BindingIndex = INDEX_NONE;
 		int32 UserWidgetPropertyContextIndex = INDEX_NONE;
+		int32 SourceCreatorContextIndex = INDEX_NONE;
 		bool bSourceIsUserWidget = false;
 		bool bFieldIdNeeded = false;
 		bool bIsConversionFunctionComplex = false;
@@ -166,6 +169,8 @@ private:
 
 		// The source if it's a property
 		int32 UserWidgetPropertyContextIndex = INDEX_NONE;
+		// The source if it's a property
+		int32 SourceCreatorContextIndex = INDEX_NONE;
 		// The source is the UserWidget itself
 		bool bIsRootWidget = false;
 	};
