MVVM: Temporary disable code that make the viewmodel variable private. The code is there to prevent setting the variable without using the Setter function.

diff --git a/Engine/Plugins/Runtime/ModelViewViewModel/Source/ModelViewViewModelBlueprint/Private/MVVMViewBlueprintCompiler.cpp b/Engine/Plugins/Runtime/ModelViewViewModel/Source/ModelViewViewModelBlueprint/Private/MVVMViewBlueprintCompiler.cpp
--- a/Engine/Plugins/Runtime/ModelViewViewModel/Source/ModelViewViewModelBlueprint/Private/MVVMViewBlueprintCompiler.cpp
+++ b/Engine/Plugins/Runtime/ModelViewViewModel/Source/ModelViewViewModelBlueprint/Private/MVVMViewBlueprintCompiler.cpp
@@ -205,10 +205,10 @@ void FMVVMViewBlueprintCompiler::CreateVariables(const FWidgetBlueprintCompilerC
 			{
 				NewProperty->SetMetaData(FBlueprintMetadata::MD_ExposeOnSpawn, TEXT("true"));
 			}
-			if (!SourceContext.bPublicGetter)
-			{
-				NewProperty->SetMetaData(FBlueprintMetadata::MD_Private, TEXT("true"));
-			}
+			//if (!SourceContext.bPublicGetter)
+			//{
+			//	NewProperty->SetMetaData(FBlueprintMetadata::MD_Private, TEXT("true"));
+			//}
 #endif
 		}
 		return NewProperty;
