MVVM: Rename MVVMBlueprintViewBinding.ViewModelPath and WidgetPath to SourePath and DestinationPath. Historicly they were named that way because it made sense to bind grom viewmodel to a widget. Now we can bind from any property to any property. We renamed to remove the confusion.
#test MVVM asset in QAGame

diff --git a/Engine/Plugins/Runtime/ModelViewViewModel/Config/BaseModelViewViewModel.ini b/Engine/Plugins/Runtime/ModelViewViewModel/Config/BaseModelViewViewModel.ini
--- a/Engine/Plugins/Runtime/ModelViewViewModel/Config/BaseModelViewViewModel.ini
+++ b/Engine/Plugins/Runtime/ModelViewViewModel/Config/BaseModelViewViewModel.ini
@@ -1,9 +1,11 @@
 [CoreRedirects]
-+StructRedirects=(OldName="MVVMViewModelPropertyPath",NewName="/script/ModelViewViewModelBlueprint.MVVMBlueprintPropertyPath")
-+StructRedirects=(OldName="MVVMWidgetPropertyPath",NewName="/script/ModelViewViewModelBlueprint.MVVMBlueprintPropertyPath")
++StructRedirects=(OldName="MVVMViewModelPropertyPath",NewName="/Script/ModelViewViewModelBlueprint.MVVMBlueprintPropertyPath")
++StructRedirects=(OldName="MVVMWidgetPropertyPath",NewName="/Script/ModelViewViewModelBlueprint.MVVMBlueprintPropertyPath")
 +FunctionRedirects=(OldName="/Script/ModelViewViewModel.MVVMSubsystem.GetViewFromUserWidget", NewName="/Script/ModelViewViewModel.MVVMSubsystem.K2_GetViewFromUserWidget")
 +FunctionRedirects=(OldName="/Script/ModelViewViewModel.MVVMSubsystem.GetAvailableBindings", NewName="/Script/ModelViewViewModel.MVVMSubsystem.K2_GetAvailableBindings")
 +FunctionRedirects=(OldName="/Script/ModelViewViewModel.MVVMSubsystem.GetAvailableBinding", NewName="/Script/ModelViewViewModel.MVVMSubsystem.K2_GetAvailableBinding")
++PropertyRedirects=(OldName="/Script/ModelViewViewModelBlueprint.MVVMBlueprintViewBinding.ViewModelPath",NewName="SourcePath")
++PropertyRedirects=(OldName="/Script/ModelViewViewModelBlueprint.MVVMBlueprintViewBinding.WidgetPath",NewName="DestinationPath")
 
 [/Script/ModelViewViewModelBlueprint.MVVMDeveloperProjectSettings]
 +FieldSelectorPermissions=(("/Script/UMG.Widget", (DisallowedFieldNames=("TooltipWidget","ForceVolatile","SetUserFocus","Navigation","GetParent","GetOwningPlayer","GetOwningLocalPlayer","GetGameInstance"),AdvancedFieldNames=())))
\ No newline at end of file
diff --git a/Engine/Plugins/Runtime/ModelViewViewModel/Source/ModelViewViewModelBlueprint/Private/MVVMBlueprintView.cpp b/Engine/Plugins/Runtime/ModelViewViewModel/Source/ModelViewViewModelBlueprint/Private/MVVMBlueprintView.cpp
--- a/Engine/Plugins/Runtime/ModelViewViewModel/Source/ModelViewViewModelBlueprint/Private/MVVMBlueprintView.cpp
+++ b/Engine/Plugins/Runtime/ModelViewViewModel/Source/ModelViewViewModelBlueprint/Private/MVVMBlueprintView.cpp
@@ -109,8 +109,8 @@ FMVVMBlueprintViewBinding* UMVVMBlueprintView::FindBinding(const UWidget* Widget
 	FName WidgetName = Widget->GetFName();
 	return Bindings.FindByPredicate([WidgetName, Property](const FMVVMBlueprintViewBinding& Binding)
 		{
-			return Binding.WidgetPath.GetWidgetName() == WidgetName &&
-				Binding.WidgetPath.BasePropertyPathContains(UE::MVVM::FMVVMConstFieldVariant(Property));
+			return Binding.DestinationPath.GetWidgetName() == WidgetName &&
+				Binding.DestinationPath.BasePropertyPathContains(UE::MVVM::FMVVMConstFieldVariant(Property));
 		});
 }
 
@@ -140,8 +140,8 @@ void UMVVMBlueprintView::RemoveBinding(const FMVVMBlueprintViewBinding* Binding)
 FMVVMBlueprintViewBinding& UMVVMBlueprintView::AddBinding(const UWidget* Widget, const FProperty* Property)
 {
 	FMVVMBlueprintViewBinding& NewBinding = Bindings.AddDefaulted_GetRef();
-	NewBinding.WidgetPath.SetWidgetName(Widget->GetFName());
-	NewBinding.WidgetPath.SetBasePropertyPath(UE::MVVM::FMVVMConstFieldVariant(Property));
+	NewBinding.DestinationPath.SetWidgetName(Widget->GetFName());
+	NewBinding.DestinationPath.SetBasePropertyPath(UE::MVVM::FMVVMConstFieldVariant(Property));
 	NewBinding.BindingId = FGuid::NewGuid();
 
 	OnBindingsUpdated.Broadcast();
@@ -262,9 +262,9 @@ void UMVVMBlueprintView::WidgetRenamed(FName OldObjectName, FName NewObjectName)
 	bool bRenamed = false;
 	for (FMVVMBlueprintViewBinding& Binding : Bindings)
 	{
-		if (Binding.WidgetPath.GetWidgetName() == OldObjectName)
+		if (Binding.DestinationPath.GetWidgetName() == OldObjectName)
 		{
-			Binding.WidgetPath.SetWidgetName(NewObjectName);
+			Binding.DestinationPath.SetWidgetName(NewObjectName);
 			bRenamed = true;
 		}
 	}
diff --git a/Engine/Plugins/Runtime/ModelViewViewModel/Source/ModelViewViewModelBlueprint/Private/MVVMBlueprintViewBinding.cpp b/Engine/Plugins/Runtime/ModelViewViewModel/Source/ModelViewViewModelBlueprint/Private/MVVMBlueprintViewBinding.cpp
--- a/Engine/Plugins/Runtime/ModelViewViewModel/Source/ModelViewViewModelBlueprint/Private/MVVMBlueprintViewBinding.cpp
+++ b/Engine/Plugins/Runtime/ModelViewViewModel/Source/ModelViewViewModelBlueprint/Private/MVVMBlueprintViewBinding.cpp
@@ -98,7 +98,7 @@ namespace UE::MVVM::Private
 				NameBuilder << SourceToDestFunction->GetName();
 				NameBuilder << TEXT("(");
 			
-				AppendViewModelPathString(BlueprintView, Binding.ViewModelPath, NameBuilder);
+				AppendViewModelPathString(BlueprintView, Binding.SourcePath, NameBuilder);
 
 				NameBuilder << TEXT(")");
 			}
@@ -108,9 +108,9 @@ namespace UE::MVVM::Private
 				NameBuilder << TEXT("()");
 			}
 		}
-		else if (!Binding.ViewModelPath.IsEmpty())
+		else if (!Binding.SourcePath.IsEmpty())
 		{
-			AppendViewModelPathString(BlueprintView, Binding.ViewModelPath, NameBuilder);
+			AppendViewModelPathString(BlueprintView, Binding.SourcePath, NameBuilder);
 		}
 
 		FString Name = NameBuilder.ToString();
@@ -195,7 +195,7 @@ namespace UE::MVVM::Private
 				NameBuilder << DestToSourceFunction->GetName();
 				NameBuilder << TEXT("(");
 			
-				AppendWidgetPathString(WidgetBlueprint, Binding.WidgetPath, NameBuilder);
+				AppendWidgetPathString(WidgetBlueprint, Binding.DestinationPath, NameBuilder);
 
 				NameBuilder << TEXT(")");
 			}
@@ -205,9 +205,9 @@ namespace UE::MVVM::Private
 				NameBuilder << TEXT("()");
 			}
 		}
-		else if (!Binding.WidgetPath.IsEmpty())
+		else if (!Binding.DestinationPath.IsEmpty())
 		{
-			AppendWidgetPathString(WidgetBlueprint, Binding.WidgetPath, NameBuilder);
+			AppendWidgetPathString(WidgetBlueprint, Binding.DestinationPath, NameBuilder);
 		}
 
 		FString Name = NameBuilder.ToString();
diff --git a/Engine/Plugins/Runtime/ModelViewViewModel/Source/ModelViewViewModelBlueprint/Private/MVVMViewBlueprintCompiler.cpp b/Engine/Plugins/Runtime/ModelViewViewModel/Source/ModelViewViewModelBlueprint/Private/MVVMViewBlueprintCompiler.cpp
--- a/Engine/Plugins/Runtime/ModelViewViewModel/Source/ModelViewViewModelBlueprint/Private/MVVMViewBlueprintCompiler.cpp
+++ b/Engine/Plugins/Runtime/ModelViewViewModel/Source/ModelViewViewModelBlueprint/Private/MVVMViewBlueprintCompiler.cpp
@@ -518,18 +518,18 @@ void FMVVMViewBlueprintCompiler::CreateSourceLists(const FWidgetBlueprintCompile
 				// generate destination source
 				if (bIsForwardBinding)
 				{
-					bAreSourceContextsValid &= GenerateCompilerSourceContext(Binding.WidgetPath, false);
+					bAreSourceContextsValid &= GenerateCompilerSourceContext(Binding.DestinationPath, false);
 				}
 				else
 				{
-					bAreSourceContextsValid &= GenerateCompilerSourceContext(Binding.ViewModelPath, true);
+					bAreSourceContextsValid &= GenerateCompilerSourceContext(Binding.SourcePath, true);
 				}
 			}
 			else
 			{
 				// if we aren't using a conversion function, just validate the widget and viewmodel paths
-				bAreSourceContextsValid &= GenerateCompilerSourceContext(Binding.WidgetPath, false);
-				bAreSourceContextsValid &= GenerateCompilerSourceContext(Binding.ViewModelPath, true);
+				bAreSourceContextsValid &= GenerateCompilerSourceContext(Binding.DestinationPath, false);
+				bAreSourceContextsValid &= GenerateCompilerSourceContext(Binding.SourcePath, true);
 			}
 		}
 	}
@@ -660,7 +660,7 @@ bool FMVVMViewBlueprintCompiler::PreCompileBindingSources(UWidgetBlueprintGenera
 			{
 				AddErrorForBinding(Binding, 
 					FText::Format(LOCTEXT("PropertyPathInvalidWithReason", "The property path '{0}' is invalid. {1}"),
-						PropertyPathToText(BlueprintView, Binding.ViewModelPath),
+						PropertyPathToText(BlueprintView, Binding.SourcePath),
 						CreatedBindingSourceContext.GetError()
 					),
 					ArgumentName
@@ -673,7 +673,7 @@ bool FMVVMViewBlueprintCompiler::PreCompileBindingSources(UWidgetBlueprintGenera
 			{
 				AddErrorForBinding(Binding, 
 					FText::Format(LOCTEXT("PropertyPathIsInvalid", "The property path '{0}' is invalid."), 
-						PropertyPathToText(BlueprintView, Binding.ViewModelPath)
+						PropertyPathToText(BlueprintView, Binding.SourcePath)
 					),
 					ArgumentName
 				);
@@ -712,9 +712,9 @@ bool FMVVMViewBlueprintCompiler::PreCompileBindingSources(UWidgetBlueprintGenera
 		{
 			CreateSourcesForConversionFunction(true);
 
-			if (!Binding.ViewModelPath.IsEmpty())
+			if (!Binding.SourcePath.IsEmpty())
 			{
-				if (!CreateSourceContextForPropertyPath(Binding.ViewModelPath, true))
+				if (!CreateSourceContextForPropertyPath(Binding.SourcePath, true))
 				{
 					bIsBindingsValid = false;
 					continue;
@@ -726,9 +726,9 @@ bool FMVVMViewBlueprintCompiler::PreCompileBindingSources(UWidgetBlueprintGenera
 		{
 			CreateSourcesForConversionFunction(false);
 
-			if (!Binding.WidgetPath.IsEmpty() && Binding.WidgetPath.HasPaths())
+			if (!Binding.DestinationPath.IsEmpty() && Binding.DestinationPath.HasPaths())
 			{
-				if (!CreateSourceContextForPropertyPath(Binding.WidgetPath, false))
+				if (!CreateSourceContextForPropertyPath(Binding.DestinationPath, false))
 				{
 					bIsBindingsValid = false;
 					continue;
@@ -1055,7 +1055,7 @@ bool FMVVMViewBlueprintCompiler::PreCompileBindings(UWidgetBlueprintGeneratedCla
 
 		TArray<UE::MVVM::FMVVMConstFieldVariant> SetterPath;
 		{
-			const FMVVMBlueprintPropertyPath& DestinationPath = BindingSourceContext.bIsForwardBinding ? Binding.WidgetPath : Binding.ViewModelPath;
+			const FMVVMBlueprintPropertyPath& DestinationPath = BindingSourceContext.bIsForwardBinding ? Binding.DestinationPath : Binding.SourcePath;
 			SetterPath = CreateBindingDestinationPath(BlueprintView, Class, DestinationPath);
 			if (!IsPropertyPathValid(SetterPath))
 			{
diff --git a/Engine/Plugins/Runtime/ModelViewViewModel/Source/ModelViewViewModelBlueprint/Public/MVVMBlueprintViewBinding.h b/Engine/Plugins/Runtime/ModelViewViewModel/Source/ModelViewViewModelBlueprint/Public/MVVMBlueprintViewBinding.h
--- a/Engine/Plugins/Runtime/ModelViewViewModel/Source/ModelViewViewModelBlueprint/Public/MVVMBlueprintViewBinding.h
+++ b/Engine/Plugins/Runtime/ModelViewViewModel/Source/ModelViewViewModelBlueprint/Public/MVVMBlueprintViewBinding.h
@@ -43,10 +43,10 @@ struct MODELVIEWVIEWMODELBLUEPRINT_API FMVVMBlueprintViewBinding
 	GENERATED_BODY()
 
 	UPROPERTY(EditAnywhere, Category = "Viewmodel")
-	FMVVMBlueprintPropertyPath ViewModelPath;
+	FMVVMBlueprintPropertyPath SourcePath;
 
 	UPROPERTY(EditAnywhere, Category = "Viewmodel")
-	FMVVMBlueprintPropertyPath WidgetPath;
+	FMVVMBlueprintPropertyPath DestinationPath;
 
 	/** */
 	UPROPERTY(EditAnywhere, Category = "Viewmodel")
diff --git a/Engine/Plugins/Runtime/ModelViewViewModel/Source/ModelViewViewModelEditor/Private/Customizations/MVVMPropertyBindingExtension.cpp b/Engine/Plugins/Runtime/ModelViewViewModel/Source/ModelViewViewModelEditor/Private/Customizations/MVVMPropertyBindingExtension.cpp
--- a/Engine/Plugins/Runtime/ModelViewViewModel/Source/ModelViewViewModelEditor/Private/Customizations/MVVMPropertyBindingExtension.cpp
+++ b/Engine/Plugins/Runtime/ModelViewViewModel/Source/ModelViewViewModelEditor/Private/Customizations/MVVMPropertyBindingExtension.cpp
@@ -79,11 +79,11 @@ static void ExtendBindingsMenu(FMenuBuilder& MenuBuilder, const UWidgetBlueprint
 	{
 		FMVVMBlueprintViewBinding& NewBinding = MVVMBlueprintView->AddDefaultBinding();
 
-		NewBinding.ViewModelPath.SetViewModelId(ViewModelId);
-		NewBinding.ViewModelPath.SetBasePropertyPath(UE::MVVM::FMVVMConstFieldVariant(ViewModelProperty));
+		NewBinding.SourcePath.SetViewModelId(ViewModelId);
+		NewBinding.SourcePath.SetBasePropertyPath(UE::MVVM::FMVVMConstFieldVariant(ViewModelProperty));
 
-		NewBinding.WidgetPath.SetWidgetName(Widget->GetFName());
-		NewBinding.WidgetPath.SetBasePropertyPath(UE::MVVM::FMVVMConstFieldVariant(WidgetProperty));
+		NewBinding.DestinationPath.SetWidgetName(Widget->GetFName());
+		NewBinding.DestinationPath.SetBasePropertyPath(UE::MVVM::FMVVMConstFieldVariant(WidgetProperty));
 
 		NewBinding.BindingType = EMVVMBindingMode::OneWayToDestination;
 	};
@@ -159,7 +159,7 @@ TOptional<FName> FMVVMPropertyBindingExtension::GetCurrentValue(const UWidgetBlu
 		return TOptional<FName>();
 	}
 
-	TArray<FName> Names = Binding->ViewModelPath.GetPaths();
+	TArray<FName> Names = Binding->SourcePath.GetPaths();
 	if (Names.Num() > 0)
 	{
 		return Names.Last();
@@ -182,7 +182,7 @@ const FSlateBrush* FMVVMPropertyBindingExtension::GetCurrentIcon(const UWidgetBl
 		return nullptr;
 	}
 
-	TArray<UE::MVVM::FMVVMConstFieldVariant> Fields = Binding->ViewModelPath.GetFields();
+	TArray<UE::MVVM::FMVVMConstFieldVariant> Fields = Binding->SourcePath.GetFields();
 	if (Fields.IsEmpty())
 	{
 		return nullptr;
@@ -211,7 +211,7 @@ void FMVVMPropertyBindingExtension::ClearCurrentValue(const UWidgetBlueprint* Wi
 		{
 			if (FMVVMBlueprintViewBinding* Binding = MVVMBlueprintView->FindBinding(Widget, Property))
 			{
-				Binding->ViewModelPath.ResetBasePropertyPath();
+				Binding->SourcePath.ResetBasePropertyPath();
 			}
 		}
 	}
diff --git a/Engine/Plugins/Runtime/ModelViewViewModel/Source/ModelViewViewModelEditor/Private/Customizations/MVVMPropertyPathCustomization.cpp b/Engine/Plugins/Runtime/ModelViewViewModel/Source/ModelViewViewModelEditor/Private/Customizations/MVVMPropertyPathCustomization.cpp
--- a/Engine/Plugins/Runtime/ModelViewViewModel/Source/ModelViewViewModelEditor/Private/Customizations/MVVMPropertyPathCustomization.cpp
+++ b/Engine/Plugins/Runtime/ModelViewViewModel/Source/ModelViewViewModelEditor/Private/Customizations/MVVMPropertyPathCustomization.cpp
@@ -30,12 +30,12 @@ void FPropertyPathCustomization::CustomizeHeader(TSharedRef<IPropertyHandle> InP
 
 	const FName PropertyName = PropertyHandle->GetProperty()->GetFName();
 	FText PropertyDisplayName;
-	if (PropertyName == GET_MEMBER_NAME_CHECKED(FMVVMBlueprintViewBinding, WidgetPath))
+	if (PropertyName == GET_MEMBER_NAME_CHECKED(FMVVMBlueprintViewBinding, DestinationPath))
 	{
 		bIsWidget = true;
 		PropertyDisplayName = LOCTEXT("PropertyDisplay_Widget", "Destination");
 	}
-	else if (PropertyName == GET_MEMBER_NAME_CHECKED(FMVVMBlueprintViewBinding, ViewModelPath))
+	else if (PropertyName == GET_MEMBER_NAME_CHECKED(FMVVMBlueprintViewBinding, SourcePath))
 	{
 		bIsWidget = false;
 
@@ -51,11 +51,11 @@ void FPropertyPathCustomization::CustomizeHeader(TSharedRef<IPropertyHandle> InP
 
 	if (bIsWidget)
 	{
-		OtherHandle = ParentHandle->GetChildHandle(GET_MEMBER_NAME_CHECKED(FMVVMBlueprintViewBinding, ViewModelPath));
+		OtherHandle = ParentHandle->GetChildHandle(GET_MEMBER_NAME_CHECKED(FMVVMBlueprintViewBinding, SourcePath));
 	}
 	else
 	{
-		OtherHandle = ParentHandle->GetChildHandle(GET_MEMBER_NAME_CHECKED(FMVVMBlueprintViewBinding, WidgetPath));
+		OtherHandle = ParentHandle->GetChildHandle(GET_MEMBER_NAME_CHECKED(FMVVMBlueprintViewBinding, DestinationPath));
 	}
 
 	if (OtherHandle.IsValid() && OtherHandle->IsValidHandle())
diff --git a/Engine/Plugins/Runtime/ModelViewViewModel/Source/ModelViewViewModelEditor/Private/MVVMEditorSubsystem.cpp b/Engine/Plugins/Runtime/ModelViewViewModel/Source/ModelViewViewModelEditor/Private/MVVMEditorSubsystem.cpp
--- a/Engine/Plugins/Runtime/ModelViewViewModel/Source/ModelViewViewModelEditor/Private/MVVMEditorSubsystem.cpp
+++ b/Engine/Plugins/Runtime/ModelViewViewModel/Source/ModelViewViewModelEditor/Private/MVVMEditorSubsystem.cpp
@@ -399,7 +399,7 @@ void UMVVMEditorSubsystem::SetSourceToDestinationConversionFunction(UWidgetBluep
 {
 	if (UMVVMBlueprintView* View = GetView(WidgetBlueprint))
 	{
-		if (ConversionFunction == nullptr || IsValidConversionFunction(ConversionFunction, Binding.ViewModelPath, Binding.WidgetPath))
+		if (ConversionFunction == nullptr || IsValidConversionFunction(ConversionFunction, Binding.SourcePath, Binding.DestinationPath))
 		{
 			FScopedTransaction Transaction(LOCTEXT("SetConversionFunction", "Set Conversion Function"));
 
@@ -433,7 +433,7 @@ void UMVVMEditorSubsystem::SetDestinationToSourceConversionFunction(UWidgetBluep
 {
 	if (UMVVMBlueprintView* View = GetView(WidgetBlueprint))
 	{
-		if (ConversionFunction == nullptr || IsValidConversionFunction(ConversionFunction, Binding.WidgetPath, Binding.ViewModelPath))
+		if (ConversionFunction == nullptr || IsValidConversionFunction(ConversionFunction, Binding.DestinationPath, Binding.SourcePath))
 		{
 			FScopedTransaction Transaction(LOCTEXT("SetConversionFunction", "Set Conversion Function"));
 
@@ -469,11 +469,11 @@ void UMVVMEditorSubsystem::SetWidgetPropertyForBinding(UWidgetBlueprint* WidgetB
 	{
 		FScopedTransaction Transaction(LOCTEXT("SetBindingProperty", "Set Binding Property"));
 
-		UE::MVVM::Private::OnBindingPreEditChange(View, GET_MEMBER_NAME_CHECKED(FMVVMBlueprintViewBinding, WidgetPath));
+		UE::MVVM::Private::OnBindingPreEditChange(View, GET_MEMBER_NAME_CHECKED(FMVVMBlueprintViewBinding, DestinationPath));
 
-		Binding.WidgetPath = Field;
+		Binding.DestinationPath = Field;
 
-		UE::MVVM::Private::OnBindingPostEditChange(View, GET_MEMBER_NAME_CHECKED(FMVVMBlueprintViewBinding, WidgetPath));
+		UE::MVVM::Private::OnBindingPostEditChange(View, GET_MEMBER_NAME_CHECKED(FMVVMBlueprintViewBinding, DestinationPath));
 	}
 }
 
@@ -483,11 +483,11 @@ void UMVVMEditorSubsystem::SetViewModelPropertyForBinding(UWidgetBlueprint* Widg
 	{
 		FScopedTransaction Transaction(LOCTEXT("SetBindingProperty", "Set Binding Property"));
 
-		UE::MVVM::Private::OnBindingPreEditChange(View, GET_MEMBER_NAME_CHECKED(FMVVMBlueprintViewBinding, ViewModelPath));
+		UE::MVVM::Private::OnBindingPreEditChange(View, GET_MEMBER_NAME_CHECKED(FMVVMBlueprintViewBinding, SourcePath));
 
-		Binding.ViewModelPath = Field;
+		Binding.SourcePath = Field;
 
-		UE::MVVM::Private::OnBindingPostEditChange(View, GET_MEMBER_NAME_CHECKED(FMVVMBlueprintViewBinding, ViewModelPath));
+		UE::MVVM::Private::OnBindingPostEditChange(View, GET_MEMBER_NAME_CHECKED(FMVVMBlueprintViewBinding, SourcePath));
 	}
 }
 
@@ -765,11 +765,11 @@ FMVVMBlueprintPropertyPath UMVVMEditorSubsystem::GetPathForConversionFunctionArg
 	{
 		if (bSourceToDestination)
 		{
-			return Binding.ViewModelPath;
+			return Binding.SourcePath;
 		}
 		else
 		{
-			return Binding.WidgetPath;
+			return Binding.DestinationPath;
 		}
 	}
 
@@ -790,11 +790,11 @@ void UMVVMEditorSubsystem::SetPathForConversionFunctionArgument(UWidgetBlueprint
 		// simple conversion function
 		if (bSourceToDestination)
 		{
-			Binding.ViewModelPath = Path;
+			Binding.SourcePath = Path;
 		}
 		else
 		{
-			Binding.WidgetPath = Path;
+			Binding.DestinationPath = Path;
 		}
 
 		return;
diff --git a/Engine/Plugins/Runtime/ModelViewViewModel/Source/ModelViewViewModelEditor/Private/Widgets/SMVVMConversionPath.cpp b/Engine/Plugins/Runtime/ModelViewViewModel/Source/ModelViewViewModelEditor/Private/Widgets/SMVVMConversionPath.cpp
--- a/Engine/Plugins/Runtime/ModelViewViewModel/Source/ModelViewViewModelEditor/Private/Widgets/SMVVMConversionPath.cpp
+++ b/Engine/Plugins/Runtime/ModelViewViewModel/Source/ModelViewViewModelEditor/Private/Widgets/SMVVMConversionPath.cpp
@@ -177,8 +177,8 @@ TSharedRef<SWidget> SMVVMConversionPath::GetFunctionMenuContent()
 
 	for (const FMVVMBlueprintViewBinding* Binding : ViewBindings)
 	{
-		FMVVMBlueprintPropertyPath SourcePath = bSourceToDestination ? Binding->ViewModelPath : Binding->WidgetPath;
-		FMVVMBlueprintPropertyPath DestPath = bSourceToDestination ? Binding->WidgetPath : Binding->ViewModelPath;
+		FMVVMBlueprintPropertyPath SourcePath = bSourceToDestination ? Binding->SourcePath : Binding->DestinationPath;
+		FMVVMBlueprintPropertyPath DestPath = bSourceToDestination ? Binding->DestinationPath : Binding->SourcePath;
 
 		UMVVMEditorSubsystem* EditorSubsystem = GEditor->GetEditorSubsystem<UMVVMEditorSubsystem>();
 		TArray<UFunction*> FunctionsForThis = EditorSubsystem->GetAvailableConversionFunctions(WidgetBlueprint, SourcePath, DestPath);
diff --git a/Engine/Plugins/Runtime/ModelViewViewModel/Source/ModelViewViewModelEditor/Private/Widgets/SMVVMViewBindingListView.cpp b/Engine/Plugins/Runtime/ModelViewViewModel/Source/ModelViewViewModelEditor/Private/Widgets/SMVVMViewBindingListView.cpp
--- a/Engine/Plugins/Runtime/ModelViewViewModel/Source/ModelViewViewModelEditor/Private/Widgets/SMVVMViewBindingListView.cpp
+++ b/Engine/Plugins/Runtime/ModelViewViewModel/Source/ModelViewViewModelEditor/Private/Widgets/SMVVMViewBindingListView.cpp
@@ -238,7 +238,7 @@ private:
 			{
 				if (FMVVMBlueprintViewBinding* Binding = ChildEntry->GetBinding(View))
 				{
-					FMVVMBlueprintPropertyPath CurrentPath = Binding->WidgetPath;
+					FMVVMBlueprintPropertyPath CurrentPath = Binding->DestinationPath;
 					CurrentPath.SetWidgetName(Source.Name);
 
 					EditorSubsystem->SetWidgetPropertyForBinding(WidgetBlueprint, *Binding, CurrentPath);
@@ -281,7 +281,7 @@ public:
 		CVarDefaultExecutionMode = IConsoleManager::Get().FindConsoleVariable(TEXT("MVVM.DefaultExecutionMode"));
 		
 		FMVVMBlueprintViewBinding* ViewBinding = GetThisViewBinding();
-		FBindingSource WidgetSource = FBindingSource::CreateForWidget(InWidgetBlueprint, ViewBinding->WidgetPath.GetWidgetName());
+		FBindingSource WidgetSource = FBindingSource::CreateForWidget(InWidgetBlueprint, ViewBinding->DestinationPath.GetWidgetName());
 
 		STableRow<TSharedPtr<FBindingEntry>>::Construct(
 			STableRow<TSharedPtr<FBindingEntry>>::FArguments()
@@ -571,11 +571,11 @@ private:
 			TArray<FMVVMConstFieldVariant> Fields;
 			if (bViewModel)
 			{
-				Fields = ViewBinding->WidgetPath.GetFields();
+				Fields = ViewBinding->DestinationPath.GetFields();
 			}
 			else
 			{
-				Fields = ViewBinding->ViewModelPath.GetFields();
+				Fields = ViewBinding->SourcePath.GetFields();
 			}
 
 			if (Fields.Num() > 0)
@@ -638,7 +638,7 @@ private:
 	{
 		if (FMVVMBlueprintViewBinding* ViewBinding = GetThisViewBinding())
 		{
-			return ViewBinding->ViewModelPath;
+			return ViewBinding->SourcePath;
 		}
 		return FMVVMBlueprintPropertyPath();
 	}
@@ -647,7 +647,7 @@ private:
 	{
 		if (FMVVMBlueprintViewBinding* ViewBinding = GetThisViewBinding())
 		{
-			return ViewBinding->WidgetPath;
+			return ViewBinding->DestinationPath;
 		}
 		return FMVVMBlueprintPropertyPath();
 	}
@@ -667,7 +667,7 @@ private:
 	{
 		if (FMVVMBlueprintViewBinding* ViewBinding = GetThisViewBinding())
 		{
-			if (ViewBinding->ViewModelPath != SelectedField)
+			if (ViewBinding->SourcePath != SelectedField)
 			{
 				UMVVMEditorSubsystem* Subsystem = GEditor->GetEditorSubsystem<UMVVMEditorSubsystem>();
 				Subsystem->SetViewModelPropertyForBinding(WidgetBlueprintWeak.Get(), *ViewBinding, SelectedField);
@@ -684,7 +684,7 @@ private:
 	{
 		if (FMVVMBlueprintViewBinding* ViewBinding = GetThisViewBinding())
 		{
-			if (ViewBinding->WidgetPath != SelectedField)
+			if (ViewBinding->DestinationPath != SelectedField)
 			{
 				UMVVMEditorSubsystem* Subsystem = GEditor->GetEditorSubsystem<UMVVMEditorSubsystem>();
 				Subsystem->SetWidgetPropertyForBinding(WidgetBlueprintWeak.Get(), *ViewBinding, SelectedField);
@@ -1175,7 +1175,7 @@ void SBindingsList::Refresh()
 		{
 			const FMVVMBlueprintViewBinding& Binding = Bindings[BindingIndex];
 			
-			FName WidgetName = Binding.WidgetPath.GetWidgetName();
+			FName WidgetName = Binding.DestinationPath.GetWidgetName();
 
 			TSharedPtr<FBindingEntry> ExistingWidget;
 			for (TSharedPtr<FBindingEntry> Widget : RootWidgets)
