MVVM: Expose properties to the DetailView. The MVVMView can now be inspect with the Object Inspector. It let you see which viewmodel is attached/binded to the view.
#test QAGame with the MVVM asset

diff --git a/Engine/Plugins/Runtime/ModelViewViewModel/Source/ModelViewViewModel/Public/View/MVVMView.h b/Engine/Plugins/Runtime/ModelViewViewModel/Source/ModelViewViewModel/Public/View/MVVMView.h
--- a/Engine/Plugins/Runtime/ModelViewViewModel/Source/ModelViewViewModel/Public/View/MVVMView.h
+++ b/Engine/Plugins/Runtime/ModelViewViewModel/Source/ModelViewViewModel/Public/View/MVVMView.h
@@ -26,11 +26,16 @@ struct FMVVMViewSource
 {
 	GENERATED_BODY()
 
-	UPROPERTY()
+	UPROPERTY(VisibleAnywhere, Category = "Viewmodel")
 	TObjectPtr<UObject> Source; // TScriptInterface<INotifyFieldValueChanged>
 
+	UPROPERTY(VisibleAnywhere, Category = "Viewmodel")
 	FName SourceName;
+
+	UPROPERTY(VisibleAnywhere, Category = "Viewmodel")
 	int32 RegisteredCout = 0;
+
+	UPROPERTY(VisibleAnywhere, Category = "Viewmodel")
 	bool bCreatedSource = false;
 };
 
@@ -98,19 +103,22 @@ private:
 	UPROPERTY(Transient)
 	TObjectPtr<const UMVVMViewClass> ClassExtension;
 
-	UPROPERTY(Transient)
+	UPROPERTY(VisibleAnywhere, Transient, Category = "Viewmodel")
 	TArray<FMVVMViewSource> Sources;
 
 	/** The binding that are enabled for the instance. */
 	TBitArray<> EnabledLibraryBindings;
 
 	/** Should log when a binding is executed. */
+	UPROPERTY(EditAnywhere, Transient, Category = "Viewmodel")
 	bool bLogBinding = false;
 
 	/** Is the Construct method was called. */
+	UPROPERTY(VisibleAnywhere, Transient, Category = "Viewmodel")
 	bool bConstructed = false;
 
 	/** The view has at least one binding that need to be ticked every frame. */
+	UPROPERTY(VisibleAnywhere, Transient, Category = "Viewmodel")
 	bool bHasEveryTickBinding = false;
 };
 
diff --git a/Engine/Plugins/Runtime/ModelViewViewModel/Source/ModelViewViewModel/Public/View/MVVMViewClass.h b/Engine/Plugins/Runtime/ModelViewViewModel/Source/ModelViewViewModel/Public/View/MVVMViewClass.h
--- a/Engine/Plugins/Runtime/ModelViewViewModel/Source/ModelViewViewModel/Public/View/MVVMViewClass.h
+++ b/Engine/Plugins/Runtime/ModelViewViewModel/Source/ModelViewViewModel/Public/View/MVVMViewClass.h
@@ -172,13 +172,13 @@ private:
 	UPROPERTY()
 	FMVVMVCompiledFieldId FieldId;
 
-	UPROPERTY()
+	UPROPERTY(VisibleAnywhere, Category = "Viewmodel")
 	FName SourcePropertyName;
 
 	UPROPERTY()
 	FMVVMVCompiledBinding Binding;
 
-	UPROPERTY()
+	UPROPERTY(VisibleAnywhere, Category = "Viewmodel")
 	EMVVMExecutionMode ExecutionMode = EMVVMExecutionMode::Immediate;
 
 	enum EBindingFlags
@@ -200,7 +200,7 @@ private:
 		SourceObjectIsSelf = 1 << 7,
 	};
 
-	UPROPERTY()
+	UPROPERTY(VisibleAnywhere, Category = "Viewmodel")
 	uint8 Flags = EBindingFlags::None;
 };
 
