MVVM: Don't override the viewmodel when set before the widget had a chance to initialized. The viewmodels are initialized when the widget is constructed. The widget pointer is available before that, so it's valid to call SetViewModel on it. When the widget is contructed, the view will fetch the viewmodel and override the viewmodel that was previously set.

diff --git a/Engine/Plugins/Runtime/ModelViewViewModel/Source/ModelViewViewModel/Private/View/MVVMView.cpp b/Engine/Plugins/Runtime/ModelViewViewModel/Source/ModelViewViewModel/Private/View/MVVMView.cpp
--- a/Engine/Plugins/Runtime/ModelViewViewModel/Source/ModelViewViewModel/Private/View/MVVMView.cpp
+++ b/Engine/Plugins/Runtime/ModelViewViewModel/Source/ModelViewViewModel/Private/View/MVVMView.cpp
@@ -26,6 +26,15 @@ void UMVVMView::ConstructView(const UMVVMViewClass* InClassExtension)
 {
 	ensure(ClassExtension == nullptr);
 	ClassExtension = InClassExtension;
+
+	Sources.Reserve(ClassExtension->GetViewModelCreators().Num());
+	for (const FMVVMViewClass_SourceCreator& Item : ClassExtension->GetViewModelCreators())
+	{
+		FMVVMViewSource& NewCreatedSource = Sources.AddDefaulted_GetRef();
+		NewCreatedSource.Source = nullptr;
+		NewCreatedSource.SourceName = Item.GetSourceName();
+		NewCreatedSource.bCreatedSource = true;
+	}
 }
 
 
@@ -33,35 +42,37 @@ void UMVVMView::Construct()
 {
 	check(ClassExtension);
 	check(bConstructed == false);
-	check(Sources.Num() == 0);
 
 	// Init ViewModel instances
-	for (const FMVVMViewClass_SourceCreator& Item : ClassExtension->GetViewModelCreators())
+	UUserWidget* UserWidget = GetUserWidget();
+	const TArrayView<const FMVVMViewClass_SourceCreator> AllViewModelCreators = ClassExtension->GetViewModelCreators();
+	for (int32 Index = 0; Index < AllViewModelCreators.Num(); ++Index)
 	{
-		UUserWidget* UserWidget = GetUserWidget();
-		UObject* NewSource = Item.CreateInstance(ClassExtension, this, UserWidget);
-		if (NewSource)
+		const FMVVMViewClass_SourceCreator& Item = AllViewModelCreators[Index];
+		FMVVMViewSource& ViewSource = Sources[Index];
+		check(ViewSource.SourceName == Item.GetSourceName());
+		check(ViewSource.bCreatedSource == true);
+
+		if (!ViewSource.bSetManually)
 		{
-			if (Item.IsSourceAUserWidgetProperty())
+			UObject* NewSource = Item.CreateInstance(ClassExtension, this, UserWidget);
+			if (NewSource)
 			{
-				FObjectPropertyBase* FoundObjectProperty = FindFProperty<FObjectPropertyBase>(UserWidget->GetClass(), Item.GetSourceName());
-				if (ensureAlwaysMsgf(FoundObjectProperty, TEXT("The compiler should have added the property")))
+				if (Item.IsSourceAUserWidgetProperty())
 				{
-					if (ensure(NewSource->GetClass()->IsChildOf(FoundObjectProperty->PropertyClass)))
+					FObjectPropertyBase* FoundObjectProperty = FindFProperty<FObjectPropertyBase>(UserWidget->GetClass(), Item.GetSourceName());
+					if (ensureAlwaysMsgf(FoundObjectProperty, TEXT("The compiler should have added the property")))
 					{
-						FoundObjectProperty->SetObjectPropertyValue_InContainer(UserWidget, NewSource);
+						if (ensure(NewSource->GetClass()->IsChildOf(FoundObjectProperty->PropertyClass)))
+						{
+							FoundObjectProperty->SetObjectPropertyValue_InContainer(UserWidget, NewSource);
+							ViewSource.bAssignedToUserWidgetProperty = true;
+						}
 					}
 				}
 			}
+			ViewSource.Source = NewSource;
 		}
-
-		check(Item.GetSourceName().IsValid());
-		check(!FindViewSource(Item.GetSourceName()));
-
-		FMVVMViewSource& NewCreatedSource = Sources.AddDefaulted_GetRef();
-		NewCreatedSource.Source = NewSource;
-		NewCreatedSource.SourceName = Item.GetSourceName();
-		NewCreatedSource.bCreatedSource = true;
 	}
 
 #if UE_WITH_MVVM_DEBUGGING
@@ -98,7 +109,10 @@ void UMVVMView::Destruct()
 	UE::MVVM::FDebugging::BroadcastViewBeginDestruction(this);
 #endif
 
-	for (const FMVVMViewSource& Source : Sources)
+	UUserWidget* UserWidget = GetUserWidget();
+	check(UserWidget);
+
+	for (FMVVMViewSource& Source : Sources)
 	{
 		if (Source.RegisteredCout > 0 && Source.Source)
 		{
@@ -106,9 +120,23 @@ void UMVVMView::Destruct()
 			checkf(SourceAsInterface.GetInterface(), TEXT("It was added as a INotifyFieldValueChanged. It should still be."));
 			SourceAsInterface->RemoveAllFieldValueChangedDelegates(this);
 		}
+
+		// For GC release any object used by the view
+		if (!Source.bSetManually)
+		{
+			Source.Source = nullptr;
+			if (Source.bAssignedToUserWidgetProperty)
+			{
+				FObjectPropertyBase* FoundObjectProperty = FindFProperty<FObjectPropertyBase>(UserWidget->GetClass(), Source.SourceName);
+				if (ensureAlwaysMsgf(FoundObjectProperty, TEXT("The compiler should have added the property")))
+				{
+					FoundObjectProperty->SetObjectPropertyValue_InContainer(UserWidget, nullptr);
+				}
+				Source.bAssignedToUserWidgetProperty = false;
+			}
+		}
 	}
-	// Todo, do we need to reset the FProperty also?
-	Sources.Reset(); // For GC release any object used by the view
+
 	EnabledLibraryBindings.Reset();
 
 	if (bHasEveryTickBinding)
@@ -180,12 +208,14 @@ bool UMVVMView::SetViewModel(FName ViewModelName, TScriptInterface<INotifyFieldV
 			}
 
 			ViewSource.Source = NewValue.GetObject();
+			ViewSource.bSetManually = true;
 			if (ClassExtension->GetViewModelCreators()[AllCreatedSourcesIndex].IsSourceAUserWidgetProperty())
 			{
 				FObjectPropertyBase* FoundObjectProperty = FindFProperty<FObjectPropertyBase>(GetUserWidget()->GetClass(), ViewModelName);
 				if (ensure(FoundObjectProperty))
 				{
 					FoundObjectProperty->SetObjectPropertyValue_InContainer(GetUserWidget(), NewValue.GetObject());
+					ViewSource.bAssignedToUserWidgetProperty = true;
 				}
 			}
 
@@ -397,6 +427,7 @@ void UMVVMView::ExecuteLibraryBinding(const FMVVMViewClass_CompiledBinding& Bind
 		UE::MVVM::FDebugging::BroadcastLibraryBindingExecuted(this, Binding);
 	}
 #endif
+
 	if (ExecutionResult.HasError())
 	{
 		UE::MVVM::FMessageLog Log(GetUserWidget());
@@ -567,6 +598,7 @@ bool UMVVMView::RegisterLibraryBinding(const FMVVMViewClass_CompiledBinding& Bin
 		}
 		else
 		{
+			// if it's a widget, it may not be in the Sources
 			FMVVMViewSource& NewSource = Sources.AddDefaulted_GetRef();
 			NewSource.Source = Source.GetObject();
 			NewSource.SourceName = Binding.GetSourceName();
diff --git a/Engine/Plugins/Runtime/ModelViewViewModel/Source/ModelViewViewModel/Public/View/MVVMView.h b/Engine/Plugins/Runtime/ModelViewViewModel/Source/ModelViewViewModel/Public/View/MVVMView.h
--- a/Engine/Plugins/Runtime/ModelViewViewModel/Source/ModelViewViewModel/Public/View/MVVMView.h
+++ b/Engine/Plugins/Runtime/ModelViewViewModel/Source/ModelViewViewModel/Public/View/MVVMView.h
@@ -32,11 +32,21 @@ struct FMVVMViewSource
 	UPROPERTY(VisibleAnywhere, Category = "Viewmodel")
 	FName SourceName;
 
+	// Number of bindings connected to the source.
 	UPROPERTY(VisibleAnywhere, Category = "Viewmodel")
 	int32 RegisteredCout = 0;
 
+	// The source is defined in the ClassExtension.
 	UPROPERTY(VisibleAnywhere, Category = "Viewmodel")
 	bool bCreatedSource = false;
+
+	// The source was set manually via SetViewModel.
+	UPROPERTY(VisibleAnywhere, Category = "Viewmodel")
+	bool bSetManually = false;
+
+	// The source was set to a UserWidget property.
+	UPROPERTY(VisibleAnywhere, Category = "Viewmodel")
+	bool bAssignedToUserWidgetProperty = false;
 };
 
 /**
