MVVM: Add a global option to allow viewmodel setter function. This prevent the validation for object /Script/CoreUObject.Function.
Add an option to show the detail view in the binding panel. This prevent the bad memory access in the customizer used by a the binding panel.

diff --git a/Engine/Plugins/Runtime/ModelViewViewModel/Source/ModelViewViewModelBlueprint/Private/MVVMViewBlueprintCompiler.cpp b/Engine/Plugins/Runtime/ModelViewViewModel/Source/ModelViewViewModelBlueprint/Private/MVVMViewBlueprintCompiler.cpp
--- a/Engine/Plugins/Runtime/ModelViewViewModel/Source/ModelViewViewModelBlueprint/Private/MVVMViewBlueprintCompiler.cpp
+++ b/Engine/Plugins/Runtime/ModelViewViewModel/Source/ModelViewViewModelBlueprint/Private/MVVMViewBlueprintCompiler.cpp
@@ -142,6 +142,11 @@ void FMVVMViewBlueprintCompiler::CreateFunctions(UMVVMBlueprintView* BlueprintVi
 		return;
 	}
 
+	if (!GetDefault<UMVVMDeveloperProjectSettings>()->bAllowGeneratedViewModelSetter)
+	{
+		return;
+	}
+
 	for (const FCompilerSourceCreatorContext& SourceCreator : CompilerSourceCreatorContexts)
 	{
 		if (SourceCreator.SetterGraph)
@@ -334,7 +339,8 @@ void FMVVMViewBlueprintCompiler::CreateSourceLists(const FWidgetBlueprintCompile
 			continue;
 		}
 
-		const bool bCreateSetterFunction = ViewModelContext.bCreateSetterFunction || ViewModelContext.CreationType == EMVVMBlueprintViewModelContextCreationType::Manual;
+		const bool bCreateSetterFunction = GetDefault<UMVVMDeveloperProjectSettings>()->bAllowGeneratedViewModelSetter
+			&& (ViewModelContext.bCreateSetterFunction || ViewModelContext.CreationType == EMVVMBlueprintViewModelContextCreationType::Manual);
 
 		int32 FoundSourceCreatorContextIndex = INDEX_NONE;
 		if (Context.GetCompileType() == EKismetCompileType::SkeletonOnly)
@@ -559,6 +565,11 @@ void FMVVMViewBlueprintCompiler::CreateFunctionsDeclaration(const FWidgetBluepri
 		BlueprintView->TemporaryGraph.Reset();
 	}
 
+	if (!GetDefault<UMVVMDeveloperProjectSettings>()->bAllowGeneratedViewModelSetter)
+	{
+		return;
+	}
+
 	for (FCompilerSourceCreatorContext& SourceCreator : CompilerSourceCreatorContexts)
 	{
 		if (!SourceCreator.SetterFunctionName.IsEmpty() && SourceCreator.Type == ECompilerSourceCreatorType::ViewModel)
diff --git a/Engine/Plugins/Runtime/ModelViewViewModel/Source/ModelViewViewModelBlueprint/Public/MVVMDeveloperProjectSettings.h b/Engine/Plugins/Runtime/ModelViewViewModel/Source/ModelViewViewModelBlueprint/Public/MVVMDeveloperProjectSettings.h
--- a/Engine/Plugins/Runtime/ModelViewViewModel/Source/ModelViewViewModelBlueprint/Public/MVVMDeveloperProjectSettings.h
+++ b/Engine/Plugins/Runtime/ModelViewViewModel/Source/ModelViewViewModelBlueprint/Public/MVVMDeveloperProjectSettings.h
@@ -90,6 +90,14 @@ public:
 	UPROPERTY(EditAnywhere, config, Category = "Viewmodel")
 	bool bAllowBindingFromDetailView = true;
 
+	/** When generating a source from the Viewmodel editor, allow the compiler to generate a setter function. */
+	UPROPERTY(EditAnywhere, config, Category = "Viewmodel")
+	bool bAllowGeneratedViewModelSetter = true;
+	
+	/** For the binding list widget, allow the user to edit the binding in the detail view. */
+	UPROPERTY(EditAnywhere, config, Category = "Viewmodel")
+	bool bShowDetailViewOptionInBindingPanel = true;
+
 	/** Permission list for filtering which execution mode is allowed. */
 	UPROPERTY(EditAnywhere, config, Category = "Viewmodel")
 	EMVVMDeveloperConversionFunctionFilterType ConversionFunctionFilter = EMVVMDeveloperConversionFunctionFilterType::BlueprintActionRegistry;
diff --git a/Engine/Plugins/Runtime/ModelViewViewModel/Source/ModelViewViewModelEditor/Private/Customizations/MVVMBlueprintViewModelContextCustomization.cpp b/Engine/Plugins/Runtime/ModelViewViewModel/Source/ModelViewViewModelEditor/Private/Customizations/MVVMBlueprintViewModelContextCustomization.cpp
--- a/Engine/Plugins/Runtime/ModelViewViewModel/Source/ModelViewViewModelEditor/Private/Customizations/MVVMBlueprintViewModelContextCustomization.cpp
+++ b/Engine/Plugins/Runtime/ModelViewViewModel/Source/ModelViewViewModelEditor/Private/Customizations/MVVMBlueprintViewModelContextCustomization.cpp
@@ -8,6 +8,7 @@
 #include "IDetailChildrenBuilder.h"
 #include "IPropertyAccessEditor.h"
 #include "IPropertyTypeCustomization.h"
+#include "MVVMDeveloperProjectSettings.h"
 #include "MVVMEditorSubsystem.h"
 #include "PropertyHandle.h"
 
@@ -130,6 +131,7 @@ void FBlueprintViewModelContextDetailCustomization::CustomizeChildren(TSharedRef
 	FName Name_ViewModelName = GET_MEMBER_NAME_CHECKED(FMVVMBlueprintViewModelContext, ViewModelName);
 	FName Name_ViewModelPropertyPath = GET_MEMBER_NAME_CHECKED(FMVVMBlueprintViewModelContext, ViewModelPropertyPath);
 	FName Name_CreationType = GET_MEMBER_NAME_CHECKED(FMVVMBlueprintViewModelContext, CreationType);
+	FName Name_CreateSetterFunction = GET_MEMBER_NAME_CHECKED(FMVVMBlueprintViewModelContext, bCreateSetterFunction);
 
 	FGuid ViewModelContextId;
 	FName ViewModelPropertyName;
@@ -276,6 +278,13 @@ void FBlueprintViewModelContextDetailCustomization::CustomizeChildren(TSharedRef
 						]
 					];
 			}
+			else if (PropertyName == Name_CreateSetterFunction)
+			{
+				if (GetDefault<UMVVMDeveloperProjectSettings>()->bAllowGeneratedViewModelSetter)
+				{
+					ChildBuilder.AddProperty(ChildHandle.ToSharedRef());
+				}
+			}
 			else
 			{
 				ChildBuilder.AddProperty(ChildHandle.ToSharedRef());
diff --git a/Engine/Plugins/Runtime/ModelViewViewModel/Source/ModelViewViewModelEditor/Private/Widgets/SMVVMViewBindingPanel.cpp b/Engine/Plugins/Runtime/ModelViewViewModel/Source/ModelViewViewModelEditor/Private/Widgets/SMVVMViewBindingPanel.cpp
--- a/Engine/Plugins/Runtime/ModelViewViewModel/Source/ModelViewViewModelEditor/Private/Widgets/SMVVMViewBindingPanel.cpp
+++ b/Engine/Plugins/Runtime/ModelViewViewModel/Source/ModelViewViewModelEditor/Private/Widgets/SMVVMViewBindingPanel.cpp
@@ -7,6 +7,7 @@
 #include "IDetailsView.h"
 #include "WidgetBlueprintToolMenuContext.h"
 #include "MVVMBlueprintView.h"
+#include "MVVMDeveloperProjectSettings.h"
 #include "MVVMWidgetBlueprintExtension_View.h"
 #include "Customizations/MVVMConversionPathCustomization.h"
 #include "Customizations/MVVMPropertyPathCustomization.h"
@@ -297,18 +298,22 @@ TSharedRef<SWidget> SBindingsPanel::GenerateEditViewWidget()
 
 		ToolbarBuilderGlobal.BeginSection("Options");
 
-		ToolbarBuilderGlobal.AddToolBarButton(
-			FUIAction(
-				FExecuteAction::CreateSP(this, &SBindingsPanel::ToggleDetailsVisibility),
-				FCanExecuteAction(),
-				FGetActionCheckState::CreateSP(this, &SBindingsPanel::GetDetailsVisibleCheckState)
-			),
-			"ToggleDetails",
-			LOCTEXT("Details", "Details"),
-			LOCTEXT("DetailsToolTip", "Open Details View"),
-			FSlateIcon(FAppStyle::GetAppStyleSetName(), "WorldBrowser.DetailsButtonBrush"),
-			EUserInterfaceActionType::ToggleButton
-		);
+		if (GetDefault<UMVVMDeveloperProjectSettings>()->bShowDetailViewOptionInBindingPanel)
+		{
+			ToolbarBuilderGlobal.AddToolBarButton(
+				FUIAction(
+					FExecuteAction::CreateSP(this, &SBindingsPanel::ToggleDetailsVisibility),
+					FCanExecuteAction(),
+					FGetActionCheckState::CreateSP(this, &SBindingsPanel::GetDetailsVisibleCheckState)
+				),
+				"ToggleDetails",
+				LOCTEXT("Details", "Details"),
+				LOCTEXT("DetailsToolTip", "Open Details View"),
+				FSlateIcon(FAppStyle::GetAppStyleSetName(), "WorldBrowser.DetailsButtonBrush"),
+				EUserInterfaceActionType::ToggleButton
+			);
+		}
+
 		ToolbarBuilderGlobal.EndSection();
 	}
 
