MVVM:
- Add FText specialization to SetPropertyValue

diff --git a/Engine/Plugins/Runtime/ModelViewViewModel/Source/ModelViewViewModel/Public/MVVMViewModelBase.h b/Engine/Plugins/Runtime/ModelViewViewModel/Source/ModelViewViewModel/Public/MVVMViewModelBase.h
--- a/Engine/Plugins/Runtime/ModelViewViewModel/Source/ModelViewViewModel/Public/MVVMViewModelBase.h
+++ b/Engine/Plugins/Runtime/ModelViewViewModel/Source/ModelViewViewModel/Public/MVVMViewModelBase.h
@@ -82,6 +82,18 @@ protected:
 		return true;
 	}
 
+	bool SetPropertyValue(FText& Value, const FText& NewValue, UE::FieldNotification::FFieldId FieldId)
+	{
+		if (Value.EqualTo(NewValue))
+		{
+			return false;
+		}
+
+		Value = NewValue;
+		BroadcastFieldValueChanged(FieldId);
+		return true;
+	}
+
 private:
 	DECLARE_FUNCTION(execK2_SetPropertyValue);
 
