MVVM: Add option to log all bindings after compilations. Complete the ToString library from CompiledBindings. Expose the LogMVVM to other MVVM modules. In the view, properly report error in the message log.

diff --git a/Engine/Plugins/Runtime/ModelViewViewModel/Source/ModelViewViewModel/Private/Bindings/MVVMCompiledBindingLibrary.cpp b/Engine/Plugins/Runtime/ModelViewViewModel/Source/ModelViewViewModel/Private/Bindings/MVVMCompiledBindingLibrary.cpp
--- a/Engine/Plugins/Runtime/ModelViewViewModel/Source/ModelViewViewModel/Private/Bindings/MVVMCompiledBindingLibrary.cpp
+++ b/Engine/Plugins/Runtime/ModelViewViewModel/Source/ModelViewViewModel/Private/Bindings/MVVMCompiledBindingLibrary.cpp
@@ -9,9 +9,13 @@
 #include "Templates/ValueOrError.h"
 #include "Types/MVVMFieldContext.h"
 #include "Types/MVVMFunctionContext.h"
+#include "UObject/UnrealType.h"
 
 #include UE_INLINE_GENERATED_CPP_BY_NAME(MVVMCompiledBindingLibrary)
 
+
+#define LOCTEXT_NAMESPACE "MVVMCompiledBindingLibrary"
+
 /**
  *
  */
@@ -50,7 +54,19 @@ FMVVMCompiledBindingLibrary::FMVVMCompiledBindingLibrary()
 
 }
 
-FMVVMCompiledBindingLibrary::~FMVVMCompiledBindingLibrary() = default;
+
+FText FMVVMCompiledBindingLibrary::LexToText(EExecutionFailingReason Reason)
+{
+	switch (Reason)
+	{
+	case EExecutionFailingReason::IncompatibleLibrary: return LOCTEXT("FailingReasonIncompatibleLibrary", "Incompatible Library");
+	case EExecutionFailingReason::InvalidSource: return LOCTEXT("FailingReasonInvalidSource", "Invalid Source");
+	case EExecutionFailingReason::InvalidDestination: return LOCTEXT("FailingReasonInvalidDestination", "Invalid Destination");
+	case EExecutionFailingReason::InvalidConversionFunction: return LOCTEXT("FailingReasonInvalidConversionFunction", "Invalid Conversion Function");
+	case EExecutionFailingReason::InvalidCast: return LOCTEXT("FailingReasonInvalidCast", "Invalid Cast");
+	}
+	return FText();
+}
 
 
 void FMVVMCompiledBindingLibrary::Load()
@@ -105,6 +121,16 @@ void FMVVMCompiledBindingLibrary::Load()
 }
 
 
+#if WITH_EDITOR
+void FMVVMCompiledBindingLibrary::Unload()
+{
+	LoadedProperties.Empty();
+	LoadedFunctions.Empty();
+	LoadedFieldIds.Empty();
+}
+#endif
+
+
 TValueOrError<void, FMVVMCompiledBindingLibrary::EExecutionFailingReason> FMVVMCompiledBindingLibrary::Execute(UObject* InExecutionSource, const FMVVMVCompiledBinding& InBinding, EConversionFunctionType InFunctionType) const
 {
 	check(InExecutionSource);
@@ -344,7 +370,7 @@ TValueOrError<UE::MVVM::FFieldContext, void> FMVVMCompiledBindingLibrary::Evalua
 				const FObjectPropertyBase* ReturnObjectProperty = CastFieldChecked<const FObjectPropertyBase>(UE::MVVM::BindingHelper::GetReturnProperty(Function));
 				check(ReturnObjectProperty);
 				check(Function->NumParms == 1);
-				
+
 				if (!CurrentContainer.IsUObject())
 				{
 					return MakeError();
@@ -410,7 +436,7 @@ TValueOrError<UE::MVVM::FMVVMFieldVariant, void> FMVVMCompiledBindingLibrary::Ge
 }
 
 
-TValueOrError<FString, FString> FMVVMCompiledBindingLibrary::FieldPathToString(const FMVVMVCompiledFieldPath& FieldPath) const
+TValueOrError<FString, FString> FMVVMCompiledBindingLibrary::FieldPathToString(FMVVMVCompiledFieldPath FieldPath, bool bUseDisplayName) const
 {
 #if WITH_EDITORONLY_DATA
 	const bool bIsValidBinding = FieldPath.CompiledBindingLibraryId == CompiledBindingLibraryId;
@@ -421,38 +447,88 @@ TValueOrError<FString, FString> FMVVMCompiledBindingLibrary::FieldPathToString(c
 	}
 #endif
 
+	bool bHasError = false;
 	TStringBuilder<512> StringBuilder;
 	const int32 IterationMax = FieldPath.StartIndex + FieldPath.Num;
-	checkf(IterationMax < FieldPaths.Num(), TEXT("The number of field is bigger than the number of field that was compiled."));
+	checkf(IterationMax <= FieldPaths.Num(), TEXT("The number of field is bigger than the number of field that was compiled."));
 	for (int32 Index = FieldPath.StartIndex; Index < IterationMax; ++Index)
 	{
 		if (StringBuilder.Len() > 0)
 		{
-			StringBuilder.AppendChar(TEXT('.'));
+			StringBuilder << TEXT('.');
 		}
 
 		check(FieldPaths.IsValidIndex(Index));
 		FMVVMCompiledLoadedPropertyOrFunctionIndex PathIndex = FieldPaths[Index];
 		if (PathIndex.bIsProperty)
 		{
-			check(LoadedProperties.IsValidIndex(PathIndex.Index));
-			const FProperty* Property = LoadedProperties[PathIndex.Index];
-			if (!Property)
+			if (LoadedProperties.IsValidIndex(PathIndex.Index))
 			{
-				return MakeError<FString>(StringBuilder.ToString());
+				const FProperty* Property = LoadedProperties[PathIndex.Index];
+				if (!Property)
+				{
+					StringBuilder << TEXT("<Invalid>");
+					bHasError = true;
+				}
+
+#if WITH_EDITOR
+				if (bUseDisplayName)
+				{
+					StringBuilder << Property->GetDisplayNameText().ToString();
+				}
+				else
+#endif
+				{
+					StringBuilder << Property->GetFName();
+				}
+			}
+			else
+			{
+				ensureMsgf(false, TEXT("The library is not loaded"));
+				StringBuilder << TEXT("<Invalid>");
+				bHasError = true;
 			}
 		}
 		else
 		{
-			check(LoadedFunctions.IsValidIndex(PathIndex.Index));
-			UFunction* Function = LoadedFunctions[PathIndex.Index];
-			if (!Function)
+			if (LoadedFunctions.IsValidIndex(PathIndex.Index))
+			{
+				const UFunction* Function = LoadedFunctions[PathIndex.Index];
+				if (!Function)
+				{
+					StringBuilder << TEXT("<Invalid>");
+					bHasError = true;
+				}
+
+#if WITH_EDITOR
+				if (bUseDisplayName)
+				{
+					StringBuilder << Function->GetDisplayNameText().ToString();
+				}
+				else
+#endif
+				{
+					StringBuilder << Function->GetFName();
+				}
+			}
+			else
 			{
-				return MakeError<FString>(StringBuilder.ToString());
+				ensureMsgf(false, TEXT("The library is not loaded"));
+				StringBuilder << TEXT("<Invalid>");
+				bHasError = true;
 			}
 		}
 	}
-	return MakeValue<FString>(StringBuilder.ToString());
+
+	if (bHasError)
+	{
+		return MakeError<FString>(StringBuilder.ToString());
+	}
+	else
+	{
+
+		return MakeValue<FString>(StringBuilder.ToString());
+	}
 }
 
 
@@ -471,3 +547,4 @@ TValueOrError<UE::FieldNotification::FFieldId, void> FMVVMCompiledBindingLibrary
 	return MakeValue(LoadedFieldIds[InFieldId.FieldIdIndex]);
 }
 
+#undef LOCTEXT_NAMESPACE
diff --git a/Engine/Plugins/Runtime/ModelViewViewModel/Source/ModelViewViewModel/Private/MVVMMessageLog.h b/Engine/Plugins/Runtime/ModelViewViewModel/Source/ModelViewViewModel/Private/MVVMMessageLog.h
deleted file mode 100644
diff --git a/Engine/Plugins/Runtime/ModelViewViewModel/Source/ModelViewViewModel/Private/View/MVVMView.cpp b/Engine/Plugins/Runtime/ModelViewViewModel/Source/ModelViewViewModel/Private/View/MVVMView.cpp
--- a/Engine/Plugins/Runtime/ModelViewViewModel/Source/ModelViewViewModel/Private/View/MVVMView.cpp
+++ b/Engine/Plugins/Runtime/ModelViewViewModel/Source/ModelViewViewModel/Private/View/MVVMView.cpp
@@ -397,12 +397,17 @@ void UMVVMView::ExecuteLibraryBinding(const FMVVMViewClass_CompiledBinding& Bind
 	if (ExecutionResult.HasError())
 	{
 		UE::MVVM::FMessageLog Log(GetUserWidget());
-		Log.Error(FText::Format(LOCTEXT("ExecuteBindingGenericExecute", "The binding '{0}' was not executed. {1}."), FText::FromString(Binding.ToString()), LOCTEXT("todo", "Reason is todo.")));
+		Log.Error(FText::Format(LOCTEXT("ExecuteBindingGenericExecute", "The binding '{0}' was not executed. {1}.")
+			, FText::FromString(Binding.ToString(ClassExtension->GetBindingLibrary(), FMVVMViewClass_CompiledBinding::FToStringArgs::Short()))
+			, FMVVMCompiledBindingLibrary::LexToText(ExecutionResult.GetError())
+			));
 	}
 	else if (bLogBinding)
 	{
 		UE::MVVM::FMessageLog Log(GetUserWidget());
-		Log.Info(FText::Format(LOCTEXT("ExecuteLibraryBinding", "Execute binding '{0}'."), FText::FromString(Binding.ToString())));
+		Log.Info(FText::Format(LOCTEXT("ExecuteLibraryBinding", "Execute binding '{0}'.")
+			, FText::FromString(Binding.ToString(ClassExtension->GetBindingLibrary(), FMVVMViewClass_CompiledBinding::FToStringArgs::All()))
+			));
 	}
 }
 
@@ -431,12 +436,17 @@ void UMVVMView::ExecuteLibraryBinding(const FMVVMViewClass_CompiledBinding& Bind
 	if (ExecutionResult.HasError())
 	{
 		UE::MVVM::FMessageLog Log(GetUserWidget());
-		Log.Error(FText::Format(LOCTEXT("ExecuteBindingGenericExecute", "The binding '{0}' was not executed. {1}."), FText::FromString(Binding.ToString()), LOCTEXT("todo", "Reason is todo.")));
+		Log.Error(FText::Format(LOCTEXT("ExecuteBindingGenericExecute", "The binding '{0}' was not executed. {1}.")
+			, FText::FromString(Binding.ToString(ClassExtension->GetBindingLibrary(), FMVVMViewClass_CompiledBinding::FToStringArgs::Short()))
+			, FMVVMCompiledBindingLibrary::LexToText(ExecutionResult.GetError())
+			));
 	}
 	else if (bLogBinding)
 	{
 		UE::MVVM::FMessageLog Log(GetUserWidget());
-		Log.Info(FText::Format(LOCTEXT("ExecuteLibraryBinding", "Execute binding '{0}'."), FText::FromString(Binding.ToString())));
+		Log.Info(FText::Format(LOCTEXT("ExecuteLibraryBinding", "Execute binding '{0}'.")
+			, FText::FromString(Binding.ToString(ClassExtension->GetBindingLibrary(), FMVVMViewClass_CompiledBinding::FToStringArgs::All()))
+			));
 	}
 }
 
@@ -498,7 +508,9 @@ void UMVVMView::EnableLibraryBinding(const FMVVMViewClass_CompiledBinding& Bindi
 	if (bRegistered && bLogBinding)
 	{
 		UE::MVVM::FMessageLog Log(GetUserWidget());
-		Log.Info(FText::Format(LOCTEXT("EnableLibraryBinding", "Enable binding '{0}'."), FText::FromString(Binding.ToString())));
+		Log.Info(FText::Format(LOCTEXT("EnableLibraryBinding", "Enable binding '{0}'.")
+			, FText::FromString(Binding.ToString(ClassExtension->GetBindingLibrary(), FMVVMViewClass_CompiledBinding::FToStringArgs::All()))
+		));
 	}
 
 	EnabledLibraryBindings[BindingIndex] = bRegistered;
@@ -527,7 +539,9 @@ void UMVVMView::DisableLibraryBinding(const FMVVMViewClass_CompiledBinding& Bind
 	if (bLogBinding)
 	{
 		UE::MVVM::FMessageLog Log(GetUserWidget());
-		Log.Info(FText::Format(LOCTEXT("DisableLibraryBinding", "Disable binding '{0}'."), FText::FromString(Binding.ToString())));
+		Log.Info(FText::Format(LOCTEXT("DisableLibraryBinding", "Disable binding '{0}'.")
+			, FText::FromString(Binding.ToString(ClassExtension->GetBindingLibrary(), FMVVMViewClass_CompiledBinding::FToStringArgs::All()))
+			));
 	}
 }
 
@@ -536,23 +550,35 @@ bool UMVVMView::RegisterLibraryBinding(const FMVVMViewClass_CompiledBinding& Bin
 {
 	check(ClassExtension);
 
+	auto LogMessage = [this, &Binding](const FText& Message)
+	{
+		UE::MVVM::FMessageLog Log(GetUserWidget());
+		Log.Error(FText::Format(LOCTEXT("RegisterBindingFailed_Format", "Widget '{0}' can't register binding '{1}'. {2}.")
+			, FText::FromString(GetFullName())
+			, FText::FromString(Binding.ToString(ClassExtension->GetBindingLibrary(), FMVVMViewClass_CompiledBinding::FToStringArgs::Short()))
+			, Message
+			));
+	};
+
 	TValueOrError<UE::FieldNotification::FFieldId, void> FieldIdResult = ClassExtension->GetBindingLibrary().GetFieldId(Binding.GetSourceFieldId());
 	if (FieldIdResult.HasError())
 	{
+		LogMessage(LOCTEXT("RegisterBindingFailed_InvalidFieldId", "The FieldId is invalid."));
+
 #if UE_WITH_MVVM_DEBUGGING
 		UE::MVVM::FDebugging::BroadcastLibraryBindingRegistered(this, Binding, UE::MVVM::FDebugging::ERegisterLibraryBindingResult::Failed_InvalidFieldId);
 #endif
-		UE_LOG(LogMVVM, Error, TEXT("'%s' can't register binding '%s'. The FieldId is invalid."), *GetFullName(), *Binding.ToString());
 		return false;
 	}
 
 	UE::FieldNotification::FFieldId FieldId = FieldIdResult.StealValue();
 	if (!FieldId.IsValid())
 	{
+		LogMessage(LOCTEXT("RegisterBindingFailed_FieldIdNotFound", "The FieldId was not found on the source."));
+
 #if UE_WITH_MVVM_DEBUGGING
 		UE::MVVM::FDebugging::BroadcastLibraryBindingRegistered(this, Binding, UE::MVVM::FDebugging::ERegisterLibraryBindingResult::Failed_FieldIdNotFound);
 #endif
-		UE_LOG(LogMVVM, Error, TEXT("'%s' can't register binding '%s'. The FieldId was not found on the source."), *GetFullName(), *Binding.ToString());
 		return false;
 	}
 
@@ -560,31 +586,34 @@ bool UMVVMView::RegisterLibraryBinding(const FMVVMViewClass_CompiledBinding& Bin
 	TScriptInterface<INotifyFieldValueChanged> Source = FindSource(Binding, true);
 	if (Source.GetInterface() == nullptr)
 	{
-#if UE_WITH_MVVM_DEBUGGING
-		UE::MVVM::FDebugging::BroadcastLibraryBindingRegistered(this, Binding, UE::MVVM::FDebugging::ERegisterLibraryBindingResult::Failed_InvalidSource);
-#endif
 		if (!Binding.IsRegistrationOptional())
 		{
-			UE_LOG(LogMVVM, Error, TEXT("'%s' can't register binding '%s'. The source is invalid."), *GetFullName(), *Binding.ToString());
+			LogMessage(LOCTEXT("RegisterBindingFailed_InvalidSourceObject", "The source object is invalid."));
 		}
+
+#if UE_WITH_MVVM_DEBUGGING
+		UE::MVVM::FDebugging::BroadcastLibraryBindingRegistered(this, Binding, UE::MVVM::FDebugging::ERegisterLibraryBindingResult::Failed_InvalidSource);
+#endif
 		return false;
 	}
 
 	// Only bind if the source and the destination are valid.
 	if (ClassExtension->GetBindingLibrary().EvaluateFieldPath(GetUserWidget(), Binding.GetBinding().GetSourceFieldPath()).HasError())
 	{
+		LogMessage(LOCTEXT("RegisterBindingFailed_InvalidSource", "The source could not be evaluated."));
+
 #if UE_WITH_MVVM_DEBUGGING
 		UE::MVVM::FDebugging::BroadcastLibraryBindingRegistered(this, Binding, UE::MVVM::FDebugging::ERegisterLibraryBindingResult::Failed_InvalidSourceField);
 #endif
-		UE_LOG(LogMVVM, Warning, TEXT("'%s' can't register binding '%s'. The destination was not evaluated."), *GetFullName(), *Binding.ToString());
 		return false;
 	}
 	if (ClassExtension->GetBindingLibrary().EvaluateFieldPath(GetUserWidget(), Binding.GetBinding().GetDestinationFieldPath()).HasError())
 	{
+		LogMessage(LOCTEXT("RegisterBindingFailed_InvalidDestination", "The destination could not be evaluated."));
+
 #if UE_WITH_MVVM_DEBUGGING
 		UE::MVVM::FDebugging::BroadcastLibraryBindingRegistered(this, Binding, UE::MVVM::FDebugging::ERegisterLibraryBindingResult::Failed_InvalidDestinationField);
 #endif
-		UE_LOG(LogMVVM, Warning, TEXT("'%s' can't register binding '%s'. The destination was not evaluated."), *GetFullName(), *Binding.ToString());
 		return false;
 	}
 
@@ -653,6 +682,16 @@ TScriptInterface<INotifyFieldValueChanged> UMVVMView::FindSource(const FMVVMView
 	}
 	else
 	{
+		auto LogMessage = [this, &Binding](const FText& Message)
+		{
+			UE::MVVM::FMessageLog Log(GetUserWidget());
+			Log.Error(FText::Format(LOCTEXT("FindSourceFailed_Format", "Widget '{0}' can't evaluate the source for binding '{1}'. {2}.")
+				, FText::FromString(GetFullName())
+				, FText::FromString(Binding.ToString(ClassExtension->GetBindingLibrary(), FMVVMViewClass_CompiledBinding::FToStringArgs::Short()))
+				, Message
+			));
+		};
+
 		UObject* Source = nullptr;
 
 		if (const FMVVMViewSource* RegisteredSource = FindViewSource(Binding.GetSourceName()))
@@ -664,7 +703,7 @@ TScriptInterface<INotifyFieldValueChanged> UMVVMView::FindSource(const FMVVMView
 			const FObjectPropertyBase* SourceObjectProperty = CastField<FObjectPropertyBase>(UserWidget->GetClass()->FindPropertyByName(Binding.GetSourceName()));
 			if (SourceObjectProperty == nullptr)
 			{
-				UE_LOG(LogMVVM, Error, TEXT("'%s' could not evaluate the source for binding '%s'. The property name is invalid."), *GetFullName(), *Binding.ToString());
+				LogMessage(LOCTEXT("FindSourceFailed_InvalidPropertyName", "The property name is invalid."));
 				return TScriptInterface<INotifyFieldValueChanged>();
 			}
 
@@ -675,14 +714,14 @@ TScriptInterface<INotifyFieldValueChanged> UMVVMView::FindSource(const FMVVMView
 		{
 			if (!bAllowNull)
 			{
-				UE_LOG(LogMVVM, Error, TEXT("'%s' could not evaluate the source for binding '%s'. The source is an invalid object."), *GetFullName(), *Binding.ToString());
+				LogMessage(LOCTEXT("FindSourceFailed_InvalidSourceObject", "The source is an invalid object."));
 			}
 			return TScriptInterface<INotifyFieldValueChanged>();
 		}
 
 		if (!Source->Implements<UNotifyFieldValueChanged>())
 		{
-			UE_LOG(LogMVVM, Error, TEXT("'%s' could not evaluate the source for binding '%s'. The object '%s' doesn't implements INotifyFieldValueChanged."), *GetFullName(), *Binding.ToString(), *Source->GetFName().ToString());
+			LogMessage(LOCTEXT("FindSourceFailed_SourceObjectNotImplements", "The source does not implements INotifyFieldValueChanged."));
 			return TScriptInterface<INotifyFieldValueChanged>();
 		}
 
diff --git a/Engine/Plugins/Runtime/ModelViewViewModel/Source/ModelViewViewModel/Private/View/MVVMViewClass.cpp b/Engine/Plugins/Runtime/ModelViewViewModel/Source/ModelViewViewModel/Private/View/MVVMViewClass.cpp
--- a/Engine/Plugins/Runtime/ModelViewViewModel/Source/ModelViewViewModel/Private/View/MVVMViewClass.cpp
+++ b/Engine/Plugins/Runtime/ModelViewViewModel/Source/ModelViewViewModel/Private/View/MVVMViewClass.cpp
@@ -49,7 +49,6 @@ FMVVMViewClass_SourceCreator FMVVMViewClass_SourceCreator::MakeInstance(FName In
 	return Result;
 }
 
-
 FMVVMViewClass_SourceCreator FMVVMViewClass_SourceCreator::MakeFieldPath(FName InName, UClass* InNotifyFieldValueChangedClass, FMVVMVCompiledFieldPath InFieldPath, bool bOptional)
 {
 	FMVVMViewClass_SourceCreator Result;
@@ -66,7 +65,6 @@ FMVVMViewClass_SourceCreator FMVVMViewClass_SourceCreator::MakeFieldPath(FName I
 	return Result;
 }
 
-
 FMVVMViewClass_SourceCreator FMVVMViewClass_SourceCreator::MakeGlobalContext(FName InName, FMVVMViewModelContext InContext, bool bOptional)
 {
 	FMVVMViewClass_SourceCreator Result;
@@ -206,15 +204,167 @@ EMVVMExecutionMode FMVVMViewClass_CompiledBinding::GetExecuteMode() const
 	return (Flags & EBindingFlags::OverrideExecuteMode) == 0 ? DefaultMode : ExecutionMode;
 }
 
+FMVVMViewClass_CompiledBinding::FToStringArgs FMVVMViewClass_CompiledBinding::FToStringArgs::Short()
+{
+	FToStringArgs Result;
+	Result.bUseDisplayName = false;
+	Result.bAddFieldPath = false;
+	Result.bAddFlags = false;
+	return Result;
+}
+
+FMVVMViewClass_CompiledBinding::FToStringArgs FMVVMViewClass_CompiledBinding::FToStringArgs::All()
+{
+	return FToStringArgs();
+}
+
 FString FMVVMViewClass_CompiledBinding::ToString() const
 {
-	return FString::Printf(TEXT("FieldId: '%d', PropertyName: '%s', Mode: '%s', Flags: '%d' \n")
-			TEXT("Binding: '%s'")
-		, TEXT("")
-		, *SourcePropertyName.ToString()
-		, *StaticEnum<EMVVMExecutionMode>()->GetNameByValue((int64)ExecutionMode).ToString()
-		, Flags
-		, TEXT(""));
+	return EditorId.ToString();
+}
+
+FString FMVVMViewClass_CompiledBinding::ToString(const FMVVMCompiledBindingLibrary& BindingLibrary, FToStringArgs Args) const
+{
+	TStringBuilder<1024> StringBuilder;
+
+#if WITH_EDITOR
+	if (Args.bAddBindingId)
+	{
+		StringBuilder << TEXT("BindingId: ");
+		EditorId.AppendString(StringBuilder);
+		StringBuilder << TEXT("\n    ");
+	}
+#endif
+
+	if (Args.bAddFieldPath)
+	{
+		StringBuilder << TEXT("Binding: ");
+		TValueOrError<FString, FString> DestinationString = BindingLibrary.FieldPathToString(GetBinding().GetDestinationFieldPath(), Args.bUseDisplayName);
+		TValueOrError<FString, FString> SourceString = BindingLibrary.FieldPathToString(GetBinding().GetSourceFieldPath(), Args.bUseDisplayName);
+		TValueOrError<FString, FString> ConversionString = MakeValue(FString());
+		if (GetBinding().GetConversionFunctionFieldPath().IsValid())
+		{
+			ConversionString = BindingLibrary.FieldPathToString(GetBinding().GetConversionFunctionFieldPath(), Args.bUseDisplayName);
+		}
+
+		if (DestinationString.HasError() || SourceString.HasError() || ConversionString.HasError())
+		{
+			StringBuilder << TEXT("Error: ");
+		}
+
+		if (Args.bUseDisplayName)
+		{
+			StringBuilder << TEXT('"');
+		}
+		StringBuilder << (DestinationString.HasValue() ? DestinationString.GetValue() : DestinationString.GetError());
+		if (Args.bUseDisplayName)
+		{
+			StringBuilder << TEXT('"');
+		}
+		StringBuilder << TEXT(" = ");
+		if (GetBinding().GetConversionFunctionFieldPath().IsValid())
+		{
+			if (Args.bUseDisplayName)
+			{
+				StringBuilder << TEXT('"');
+			}
+			StringBuilder << (ConversionString.HasValue() ? ConversionString.GetValue() : ConversionString.GetError());
+			if (Args.bUseDisplayName)
+			{
+				StringBuilder << TEXT('"');
+			}
+			StringBuilder << TEXT(" ( ");
+		}
+
+		if (Args.bUseDisplayName)
+		{
+			StringBuilder << TEXT('"');
+		}
+		StringBuilder << (SourceString.HasValue() ? SourceString.GetValue() : SourceString.GetError());
+		if (Args.bUseDisplayName)
+		{
+			StringBuilder << TEXT('"');
+		}
+
+		if (GetBinding().GetConversionFunctionFieldPath().IsValid())
+		{
+			StringBuilder << TEXT(" )");
+		}
+
+		StringBuilder << TEXT("\n    ");
+	}
+
+	StringBuilder << TEXT("PropertyName: '");
+	StringBuilder << SourcePropertyName;
+
+	if (Args.bAddFieldId)
+	{
+		if (GetSourceFieldId().IsValid())
+		{
+			TValueOrError<UE::FieldNotification::FFieldId, void> SourceFieldId = BindingLibrary.GetFieldId(GetSourceFieldId());
+			StringBuilder << TEXT("', FieldId: '");
+			StringBuilder << (SourceFieldId.HasValue() ? SourceFieldId.GetValue().GetName() : FName());
+		}
+		else
+		{
+			ensureMsgf(IsOneTime(), TEXT("The field Id should be valid."));
+		}
+	}
+
+	if (Args.bAddFlags)
+	{
+		StringBuilder << TEXT("`, Mode: '");
+		StringBuilder << StaticEnum<EMVVMExecutionMode>()->GetNameByValue((int64)ExecutionMode);
+		StringBuilder << TEXT("`, Flags: '");
+
+		auto AddPipe = [&StringBuilder]()
+		{
+			StringBuilder << TEXT('|');
+		};
+
+		if ((Flags & EBindingFlags::ForwardBinding) != 0)
+		{
+			StringBuilder << TEXT("Forward");
+		}
+		else
+		{
+			StringBuilder << TEXT("Backward");
+		}
+
+		if ((Flags & EBindingFlags::OneTime) != 0)
+		{
+			AddPipe();
+			StringBuilder << TEXT("OneTime");
+		}
+		if ((Flags & EBindingFlags::EnabledByDefault) != 0)
+		{
+			AddPipe();
+			StringBuilder << TEXT("EnabledByDefault");
+		}
+		if ((Flags & EBindingFlags::ViewModelOptional) != 0)
+		{
+			AddPipe();
+			StringBuilder << TEXT("Optional");
+		}
+		if ((Flags & EBindingFlags::ConversionFunctionIsComplex) != 0)
+		{
+			AddPipe();
+			StringBuilder << TEXT("Complex");
+		}
+		if ((Flags & EBindingFlags::OverrideExecuteMode) != 0)
+		{
+			AddPipe();
+			StringBuilder << TEXT("OverrideExecutionMode");
+		}
+		if ((Flags & EBindingFlags::SourceObjectIsSelf) != 0)
+		{
+			AddPipe();
+			StringBuilder << TEXT("Self");
+		}
+		StringBuilder << TEXT("'");
+	}
+
+	return StringBuilder.ToString();
 }
 
 
diff --git a/Engine/Plugins/Runtime/ModelViewViewModel/Source/ModelViewViewModel/Public/Bindings/MVVMCompiledBindingLibrary.h b/Engine/Plugins/Runtime/ModelViewViewModel/Source/ModelViewViewModel/Public/Bindings/MVVMCompiledBindingLibrary.h
--- a/Engine/Plugins/Runtime/ModelViewViewModel/Source/ModelViewViewModel/Public/Bindings/MVVMCompiledBindingLibrary.h
+++ b/Engine/Plugins/Runtime/ModelViewViewModel/Source/ModelViewViewModel/Public/Bindings/MVVMCompiledBindingLibrary.h
@@ -243,7 +243,6 @@ struct MODELVIEWVIEWMODEL_API FMVVMCompiledBindingLibrary
 
 public:
 	FMVVMCompiledBindingLibrary();
-	~FMVVMCompiledBindingLibrary();
 
 public:
 	/** */
@@ -256,6 +255,8 @@ public:
 		InvalidCast,
 	};
 
+	static FText LexToText(EExecutionFailingReason Reason);
+
 	/** */
 	enum class EConversionFunctionType
 	{
@@ -268,6 +269,11 @@ public:
 	/** Fetch the FProperty and UFunction. */
 	void Load();
 
+#if WITH_EDITOR
+	/** Release the acquired FProperty and UFunction. */
+	void Unload();
+#endif
+
 	/**
 	 * Execute a binding, in one direction.
 	 * The ExecutionSource is the View (UserWidget) instance.
@@ -288,7 +294,7 @@ public:
 	TValueOrError<UE::FieldNotification::FFieldId, void> GetFieldId(const FMVVMVCompiledFieldId& FieldId) const;
 
 	/** Return a readable version of the FieldPath. */
-	TValueOrError<FString, FString> FieldPathToString(const FMVVMVCompiledFieldPath& FieldPath) const;
+	TValueOrError<FString, FString> FieldPathToString(FMVVMVCompiledFieldPath FieldPath, bool bUseDisplayName) const;
 
 private:
 	TValueOrError<void, EExecutionFailingReason> ExecuteImpl(UObject* ExecutionSource, const FMVVMVCompiledBinding& Binding, UObject* Source, EConversionFunctionType ConversionType) const;
diff --git a/Engine/Plugins/Runtime/ModelViewViewModel/Source/ModelViewViewModel/Public/MVVMMessageLog.h b/Engine/Plugins/Runtime/ModelViewViewModel/Source/ModelViewViewModel/Public/MVVMMessageLog.h
new file mode 100644
--- /dev/null
+++ b/Engine/Plugins/Runtime/ModelViewViewModel/Source/ModelViewViewModel/Public/MVVMMessageLog.h
@@ -0,0 +1,36 @@
+// Copyright Epic Games, Inc. All Rights Reserved.
+
+#pragma once
+
+#include "Logging/MessageLog.h"
+
+#include "UObject/WeakObjectPtr.h"
+
+class UUserWidget;
+
+MODELVIEWVIEWMODEL_API DECLARE_LOG_CATEGORY_EXTERN(LogMVVM, Log, All);
+
+namespace UE::MVVM
+{
+	class MODELVIEWVIEWMODEL_API FMessageLog : private ::FMessageLog
+	{
+	public:
+		static const FName LogName;
+
+		FMessageLog(const UUserWidget* InUserWidget)
+			: ::FMessageLog(LogName)
+		{}
+
+		using ::FMessageLog::AddMessage;
+		using ::FMessageLog::AddMessages;
+		using ::FMessageLog::Message;
+		using ::FMessageLog::CriticalError;
+		using ::FMessageLog::Error;
+		using ::FMessageLog::PerformanceWarning;
+		using ::FMessageLog::Warning;
+		using ::FMessageLog::Info;
+
+	private:
+		TWeakObjectPtr<UUserWidget> UserWidget;
+	};
+}
diff --git a/Engine/Plugins/Runtime/ModelViewViewModel/Source/ModelViewViewModel/Public/View/MVVMViewClass.h b/Engine/Plugins/Runtime/ModelViewViewModel/Source/ModelViewViewModel/Public/View/MVVMViewClass.h
--- a/Engine/Plugins/Runtime/ModelViewViewModel/Source/ModelViewViewModel/Public/View/MVVMViewClass.h
+++ b/Engine/Plugins/Runtime/ModelViewViewModel/Source/ModelViewViewModel/Public/View/MVVMViewClass.h
@@ -165,9 +165,31 @@ public:
 	/** How the binding should be executed. */
 	EMVVMExecutionMode GetExecuteMode() const;
 
-	/** @return a human readable version of the binding that can be use for debugging purposes. */
+	struct FToStringArgs
+	{
+		bool bUseDisplayName = true;
+		bool bAddFieldPath = true;
+		bool bAddFieldId = true;
+		bool bAddFlags = true;
+		bool bAddBindingId = true;
+
+		MODELVIEWVIEWMODEL_API static FToStringArgs Short();
+		MODELVIEWVIEWMODEL_API static FToStringArgs All();
+	};
+
+	UE_DEPRECATED(5.3, "ToString with no argument is deprecated.")
 	FString ToString() const;
 
+	/** @return a human readable version of the binding that can be use for debugging purposes. */
+	FString ToString(const FMVVMCompiledBindingLibrary& CompiledBindingLibrary, FToStringArgs Args) const;
+
+#if WITH_EDITOR
+	FGuid GetEditorId() const
+	{
+		return EditorId;
+	}
+#endif
+
 private:
 	UPROPERTY()
 	FMVVMVCompiledFieldId FieldId;
@@ -186,8 +208,6 @@ private:
 		None = 0,
 		/** True when the binding goes from Source to Destination. */
 		ForwardBinding = 1 << 0,
-		/** The binding is one part of a 2 ways binding. */
-		TwoWayBinding = 1 << 1,
 		OneTime = 1 << 2,
 		EnabledByDefault = 1 << 3,
 		/** The source (viewmodel) can be nullptr and the binding could failed and should not log a warning. */
@@ -202,6 +222,11 @@ private:
 
 	UPROPERTY(VisibleAnywhere, Category = "Viewmodel")
 	uint8 Flags = EBindingFlags::None;
+
+#if WITH_EDITORONLY_DATA
+	UPROPERTY(meta = (IgnoreForMemberInitializationTest))
+	FGuid EditorId;
+#endif
 };
 
 
diff --git a/Engine/Plugins/Runtime/ModelViewViewModel/Source/ModelViewViewModelBlueprint/Private/Bindings/MVVMCompiledBindingLibraryCompiler.cpp b/Engine/Plugins/Runtime/ModelViewViewModel/Source/ModelViewViewModelBlueprint/Private/Bindings/MVVMCompiledBindingLibraryCompiler.cpp
--- a/Engine/Plugins/Runtime/ModelViewViewModel/Source/ModelViewViewModelBlueprint/Private/Bindings/MVVMCompiledBindingLibraryCompiler.cpp
+++ b/Engine/Plugins/Runtime/ModelViewViewModel/Source/ModelViewViewModelBlueprint/Private/Bindings/MVVMCompiledBindingLibraryCompiler.cpp
@@ -530,14 +530,14 @@ TValueOrError<FCompiledBindingLibraryCompiler::FBindingHandle, FText> FCompiledB
 	NewBinding.SourcePathHandles = InSourceHandles;
 	NewBinding.DestinationPathHandle = InDestinationHandle;
 	NewBinding.ConversionFunctionPathHandle = InConversionFunctionHandle;
-	const int32 FoundIndex = Impl->Bindings.IndexOfByPredicate([&NewBinding](const Private::FRawBinding& Binding)
+	const int32 FoundSameBindingIndex = Impl->Bindings.IndexOfByPredicate([&NewBinding](const Private::FRawBinding& Binding)
 		{
 			return NewBinding.IsSameBinding(Binding);
 		});
 
-	if (FoundIndex != INDEX_NONE)
+	if (FoundSameBindingIndex != INDEX_NONE)
 	{
-		return MakeValue(Impl->Bindings[FoundIndex].BindingHandle);
+		return MakeError(LOCTEXT("BindingAlreadyAdded", "The binding already exist."));
 	}
 	else
 	{
diff --git a/Engine/Plugins/Runtime/ModelViewViewModel/Source/ModelViewViewModelBlueprint/Private/MVVMViewBlueprintCompiler.cpp b/Engine/Plugins/Runtime/ModelViewViewModel/Source/ModelViewViewModelBlueprint/Private/MVVMViewBlueprintCompiler.cpp
--- a/Engine/Plugins/Runtime/ModelViewViewModel/Source/ModelViewViewModelBlueprint/Private/MVVMViewBlueprintCompiler.cpp
+++ b/Engine/Plugins/Runtime/ModelViewViewModel/Source/ModelViewViewModelBlueprint/Private/MVVMViewBlueprintCompiler.cpp
@@ -8,6 +8,7 @@
 #include "HAL/IConsoleManager.h"
 #include "MVVMBlueprintView.h"
 #include "MVVMDeveloperProjectSettings.h"
+#include "MVVMMessageLog.h"
 #include "PropertyPermissionList.h"
 #include "MVVMFunctionGraphHelper.h"
 #include "Templates/ValueOrError.h"
@@ -20,6 +21,12 @@
 
 namespace UE::MVVM::Private
 {
+	FAutoConsoleVariable CVarOnCompliedLogBinding(
+		TEXT("MVVM.OnCompliedLogBinding"),
+		false,
+		TEXT("After the view is compiled log the compiled bindings.")
+	);
+
 FString PropertyPathToString(const UMVVMBlueprintView* BlueprintView, const FMVVMBlueprintPropertyPath& PropertyPath)
 {
 	if (PropertyPath.IsEmpty())
@@ -639,6 +646,20 @@ bool FMVVMViewBlueprintCompiler::Compile(UWidgetBlueprintGeneratedClass* Class,
 	if (bResult)
 	{
 		ViewExtension->BindingLibrary = MoveTemp(CompileResult.GetValue().Library);
+
+		if (CVarOnCompliedLogBinding->GetBool())
+		{
+			ViewExtension->BindingLibrary.Load();
+			UE_LOG(LogMVVM, Log, TEXT("Compiled Bindings for %s"), *ViewExtension->GetOutermost()->GetName());
+
+			FMVVMViewClass_CompiledBinding::FToStringArgs ToStringArgs = FMVVMViewClass_CompiledBinding::FToStringArgs::All();
+			ToStringArgs.bUseDisplayName = false;
+			for (const FMVVMViewClass_CompiledBinding& Binding : ViewExtension->GetCompiledBindings())
+			{
+				UE_LOG(LogMVVM, Log, TEXT("%s"), *Binding.ToString(ViewExtension->GetBindingLibrary(), ToStringArgs));
+			}
+			ViewExtension->BindingLibrary.Unload();
+		}
 	}
 
 	return bResult;
@@ -1197,18 +1218,19 @@ bool FMVVMViewBlueprintCompiler::CompileBindings(const FCompiledBindingLibraryCo
 					continue;
 				}
 
-				bIsOptional = ModelContext->bOptional;
+				bIsOptional = ModelContext->bOptional || ModelContext->CreationType == EMVVMBlueprintViewModelContextCreationType::Manual;
 			}
 		}
 
 		NewBinding.FieldId = CompiledFieldId  ? *CompiledFieldId : FMVVMVCompiledFieldId();
 		NewBinding.Binding = *CompiledBinding;
 		NewBinding.ExecutionMode = ViewBinding.bOverrideExecutionMode ? ViewBinding.OverrideExecutionMode : (EMVVMExecutionMode)CVarDefaultExecutionMode->GetInt();
+		NewBinding.EditorId = ViewBinding.BindingId;
 
 		NewBinding.Flags = 0;
 		NewBinding.Flags |= (ViewBinding.bEnabled) ? FMVVMViewClass_CompiledBinding::EBindingFlags::EnabledByDefault : 0;
 		NewBinding.Flags |= (IsForwardBinding(ViewBinding.BindingType)) ? FMVVMViewClass_CompiledBinding::EBindingFlags::ForwardBinding : 0;
-		NewBinding.Flags |= (ViewBinding.BindingType == EMVVMBindingMode::TwoWay) ? FMVVMViewClass_CompiledBinding::EBindingFlags::TwoWayBinding : 0;
+		//NewBinding.Flags |= (ViewBinding.BindingType == EMVVMBindingMode::TwoWay) ? FMVVMViewClass_CompiledBinding::EBindingFlags::TwoWayBinding : 0;
 		NewBinding.Flags |= (IsOneTimeBinding(ViewBinding.BindingType)) ? FMVVMViewClass_CompiledBinding::EBindingFlags::OneTime : 0;
 		NewBinding.Flags |= (bIsOptional) ? FMVVMViewClass_CompiledBinding::EBindingFlags::ViewModelOptional : 0;
 		NewBinding.Flags |= (CompileBinding.bIsConversionFunctionComplex) ? FMVVMViewClass_CompiledBinding::EBindingFlags::ConversionFunctionIsComplex : 0;
