Fix bad compilation unit.
#rnx

diff --git a/Engine/Plugins/Runtime/ModelViewViewModel/Source/ModelViewViewModel/Private/View/MVVMViewClass.cpp b/Engine/Plugins/Runtime/ModelViewViewModel/Source/ModelViewViewModel/Private/View/MVVMViewClass.cpp
--- a/Engine/Plugins/Runtime/ModelViewViewModel/Source/ModelViewViewModel/Private/View/MVVMViewClass.cpp
+++ b/Engine/Plugins/Runtime/ModelViewViewModel/Source/ModelViewViewModel/Private/View/MVVMViewClass.cpp
@@ -220,7 +220,9 @@ FMVVMViewClass_CompiledBinding::FToStringArgs FMVVMViewClass_CompiledBinding::FT
 
 FString FMVVMViewClass_CompiledBinding::ToString() const
 {
+#if WITH_EDITOR
 	return EditorId.ToString();
+#endif
 }
 
 FString FMVVMViewClass_CompiledBinding::ToString(const FMVVMCompiledBindingLibrary& BindingLibrary, FToStringArgs Args) const
