MVVM: Fix null access to pointer.
#rnx

diff --git a/Engine/Plugins/Runtime/ModelViewViewModel/Source/ModelViewViewModel/Private/Bindings/MVVMCompiledBindingLibrary.cpp b/Engine/Plugins/Runtime/ModelViewViewModel/Source/ModelViewViewModel/Private/Bindings/MVVMCompiledBindingLibrary.cpp
--- a/Engine/Plugins/Runtime/ModelViewViewModel/Source/ModelViewViewModel/Private/Bindings/MVVMCompiledBindingLibrary.cpp
+++ b/Engine/Plugins/Runtime/ModelViewViewModel/Source/ModelViewViewModel/Private/Bindings/MVVMCompiledBindingLibrary.cpp
@@ -470,16 +470,18 @@ TValueOrError<FString, FString> FMVVMCompiledBindingLibrary::FieldPathToString(F
 					StringBuilder << TEXT("<Invalid>");
 					bHasError = true;
 				}
-
-#if WITH_EDITOR
-				if (bUseDisplayName)
-				{
-					StringBuilder << Property->GetDisplayNameText().ToString();
-				}
 				else
-#endif
 				{
-					StringBuilder << Property->GetFName();
+#if WITH_EDITOR
+					if (bUseDisplayName)
+					{
+						StringBuilder << Property->GetDisplayNameText().ToString();
+					}
+					else
+#endif
+					{
+						StringBuilder << Property->GetFName();
+					}
 				}
 			}
 			else
@@ -499,16 +501,18 @@ TValueOrError<FString, FString> FMVVMCompiledBindingLibrary::FieldPathToString(F
 					StringBuilder << TEXT("<Invalid>");
 					bHasError = true;
 				}
-
-#if WITH_EDITOR
-				if (bUseDisplayName)
-				{
-					StringBuilder << Function->GetDisplayNameText().ToString();
-				}
 				else
-#endif
 				{
-					StringBuilder << Function->GetFName();
+#if WITH_EDITOR
+					if (bUseDisplayName)
+					{
+						StringBuilder << Function->GetDisplayNameText().ToString();
+					}
+					else
+#endif
+					{
+						StringBuilder << Function->GetFName();
+					}
 				}
 			}
 			else
