- Move the list of messages (errors, warnings, info) from the bindings themselves to the owning viewmodel.
- Add a border on each MVVM binding row in the list view of bindings that changes color to red/yellow if there is any error/warning for that binding
- detect if a property is bound to both a MVVM binding and a legacy binding at the same time and add a warning for it.

diff --git a/Engine/Plugins/Runtime/ModelViewViewModel/Source/ModelViewViewModelBlueprint/Private/MVVMBlueprintView.cpp b/Engine/Plugins/Runtime/ModelViewViewModel/Source/ModelViewViewModelBlueprint/Private/MVVMBlueprintView.cpp
--- a/Engine/Plugins/Runtime/ModelViewViewModel/Source/ModelViewViewModelBlueprint/Private/MVVMBlueprintView.cpp
+++ b/Engine/Plugins/Runtime/ModelViewViewModel/Source/ModelViewViewModelBlueprint/Private/MVVMBlueprintView.cpp
@@ -185,6 +185,51 @@ const FMVVMBlueprintViewBinding* UMVVMBlueprintView::GetBinding(FGuid Id) const
 	return Bindings.FindByPredicate([Id](const FMVVMBlueprintViewBinding& Binding) { return Id == Binding.BindingId; });
 }
 
+TArray<FText> UMVVMBlueprintView::GetBindingMessages(FGuid Id, UE::MVVM::EBindingMessageType InMessageType) const
+{
+	TArray<FText> Results;
+
+	if (BindingMessages.Contains(Id))
+	{
+		const TArray<UE::MVVM::FBindingMessage>& AllBindingMessages = BindingMessages[Id];
+		for (const UE::MVVM::FBindingMessage& Message : AllBindingMessages)
+		{
+			if (Message.MessageType == InMessageType)
+			{
+				Results.Add(Message.MessageText);
+			}
+		}
+	}
+	return Results;
+}
+
+bool UMVVMBlueprintView::HasBindingMessage(FGuid Id, UE::MVVM::EBindingMessageType InMessageType) const
+{
+	if (const TArray<UE::MVVM::FBindingMessage>* FoundBindingMessages = BindingMessages.Find(Id))
+	{
+		for (const UE::MVVM::FBindingMessage& Message : *FoundBindingMessages)
+		{
+			if (Message.MessageType == InMessageType)
+			{
+				return true;
+			}
+		}
+	}
+
+	return false;
+}
+
+void UMVVMBlueprintView::AddMessageToBinding(FGuid Id, UE::MVVM::FBindingMessage MessageToAdd)
+{
+	TArray<UE::MVVM::FBindingMessage>& FoundBindingMessages = BindingMessages.FindOrAdd(Id);
+	FoundBindingMessages.Add(MessageToAdd);
+}
+
+void UMVVMBlueprintView::ResetBindingMessages()
+{
+	BindingMessages.Reset();
+}
+
 void UMVVMBlueprintView::PostLoad()
 {
 	Super::PostLoad();
diff --git a/Engine/Plugins/Runtime/ModelViewViewModel/Source/ModelViewViewModelBlueprint/Private/MVVMViewBlueprintCompiler.cpp b/Engine/Plugins/Runtime/ModelViewViewModel/Source/ModelViewViewModelBlueprint/Private/MVVMViewBlueprintCompiler.cpp
--- a/Engine/Plugins/Runtime/ModelViewViewModel/Source/ModelViewViewModelBlueprint/Private/MVVMViewBlueprintCompiler.cpp
+++ b/Engine/Plugins/Runtime/ModelViewViewModel/Source/ModelViewViewModelBlueprint/Private/MVVMViewBlueprintCompiler.cpp
@@ -1,6 +1,7 @@
 // Copyright Epic Games, Inc. All Rights Reserved.
 
 #include "MVVMViewBlueprintCompiler.h"
+#include "Blueprint/WidgetTree.h"
 #include "Bindings/MVVMBindingHelper.h"
 #include "Bindings/MVVMConversionFunctionHelper.h"
 #include "Bindings/MVVMFieldPathHelper.h"
@@ -74,22 +75,36 @@ FText GetViewModelIdText(const FMVVMBlueprintPropertyPath& PropertyPath)
 	return FText::FromString(PropertyPath.GetViewModelId().ToString(EGuidFormats::DigitsWithHyphensInBraces));
 }
 
-void FMVVMViewBlueprintCompiler::AddErrorForBinding(FMVVMBlueprintViewBinding& Binding, const FText& Message, FName ArgumentName) const
+void FMVVMViewBlueprintCompiler::AddMessageForBinding(FMVVMBlueprintViewBinding& Binding, UMVVMBlueprintView* BlueprintView, const FText& MessageText, EBindingMessageType MessageType, FName ArgumentName) const
 {
 	const FText BindingName = FText::FromString(Binding.GetDisplayNameString(WidgetBlueprintCompilerContext.WidgetBlueprint()));
 
 	FText FormattedError;
 	if (!ArgumentName.IsNone())
 	{
-		FormattedError = FText::Format(LOCTEXT("BindingFormatWithArgument", "Binding '{0}': Argument '{1}' - {2}"), BindingName, FText::FromName(ArgumentName), Message);
+		FormattedError = FText::Format(LOCTEXT("BindingFormatWithArgument", "Binding '{0}': Argument '{1}' - {2}"), BindingName, FText::FromName(ArgumentName), MessageText);
 	}
 	else
 	{
-		FormattedError = FText::Format(LOCTEXT("BindingFormat", "Binding '{0}': {1}"), BindingName, Message);
+		FormattedError = FText::Format(LOCTEXT("BindingFormat", "Binding '{0}': {1}"), BindingName, MessageText);
 	}
 
-	WidgetBlueprintCompilerContext.MessageLog.Error(*FormattedError.ToString());
-	Binding.Errors.Add(FormattedError);
+	switch (MessageType)
+	{
+	case EBindingMessageType::Info:
+		WidgetBlueprintCompilerContext.MessageLog.Note(*FormattedError.ToString());
+		break;
+	case EBindingMessageType::Warning:
+		WidgetBlueprintCompilerContext.MessageLog.Warning(*FormattedError.ToString());
+		break;
+	case EBindingMessageType::Error:
+		WidgetBlueprintCompilerContext.MessageLog.Error(*FormattedError.ToString());
+		break;
+	default:
+		break;
+	}
+	FBindingMessage NewMessage = { FormattedError, MessageType };
+	BlueprintView->AddMessageToBinding(Binding.BindingId, NewMessage);
 }
 
 void FMVVMViewBlueprintCompiler::AddErrorForViewModel(const FMVVMBlueprintViewModelContext& ViewModel, const FText& Message) const
@@ -413,10 +428,12 @@ void FMVVMViewBlueprintCompiler::CreateSourceLists(const FWidgetBlueprintCompile
 			{
 				if (bViewModelPath)
 				{
-					Self->AddErrorForBinding(Binding,
+					Self->AddMessageForBinding(Binding,
+							BlueprintView,
 							FText::Format(LOCTEXT("ExpectedViewModelPath", "Expected a viewmodel path, but received a path from widget: {0}"), 
 								FText::FromName(PropertyPath.GetWidgetName())
-							)
+							),
+							EBindingMessageType::Error
 						);
 				}
 
@@ -434,10 +451,12 @@ void FMVVMViewBlueprintCompiler::CreateSourceLists(const FWidgetBlueprintCompile
 					UWidget** WidgetPtr = Self->WidgetNameToWidgetPointerMap.Find(PropertyPath.GetWidgetName());
 					if (WidgetPtr == nullptr || *WidgetPtr == nullptr)
 					{
-						Self->AddErrorForBinding(Binding,
+						Self->AddMessageForBinding(Binding,
+							BlueprintView,
 							FText::Format(LOCTEXT("InvalidWidgetFormat", "Could not find the targeted widget: {0}"), 
 								FText::FromName(PropertyPath.GetWidgetName())
 							),
+							EBindingMessageType::Error,
 							ArgumentName
 						);
 						return false;
@@ -459,8 +478,10 @@ void FMVVMViewBlueprintCompiler::CreateSourceLists(const FWidgetBlueprintCompile
 				const FMVVMBlueprintViewModelContext* SourceViewModelContext = BlueprintView->FindViewModel(PropertyPath.GetViewModelId());
 				if (SourceViewModelContext == nullptr)
 				{
-					Self->AddErrorForBinding(Binding,
+					Self->AddMessageForBinding(Binding,
+						BlueprintView,
 						FText::Format(LOCTEXT("BindingViewModelNotFound", "Could not find viewmodel with GUID {0}."), GetViewModelIdText(PropertyPath)),
+						EBindingMessageType::Error,
 						ArgumentName
 					);
 					return false;
@@ -468,21 +489,25 @@ void FMVVMViewBlueprintCompiler::CreateSourceLists(const FWidgetBlueprintCompile
 
 				if (!bViewModelPath)
 				{
-					Self->AddErrorForBinding(Binding,
+					Self->AddMessageForBinding(Binding,
+							BlueprintView,
 							FText::Format(LOCTEXT("ExpectedWidgetPath", "Expected a widget path, but received a path from viewmodel: {0}"), 
 								SourceViewModelContext->GetDisplayName()
 							),
+							EBindingMessageType::Error,
 							ArgumentName
 						);
 				}
 
 				if (!ViewModelGuids.Contains(SourceViewModelContext->GetViewModelId()))
 				{
-					Self->AddErrorForBinding(Binding,
+					Self->AddMessageForBinding(Binding,
+						BlueprintView,
 						FText::Format(LOCTEXT("BindingViewModelInvalid", "Viewmodel {0} {1} was invalid."), 
 							SourceViewModelContext->GetDisplayName(),
 							GetViewModelIdText(PropertyPath)
 						),
+						EBindingMessageType::Error,
 						ArgumentName
 					);
 					return false;
@@ -490,9 +515,11 @@ void FMVVMViewBlueprintCompiler::CreateSourceLists(const FWidgetBlueprintCompile
 			}
 			else
 			{
-				Self->AddErrorForBinding(Binding,
+				Self->AddMessageForBinding(Binding,
+					BlueprintView,
 					bViewModelPath ? LOCTEXT("ViewModelPathNotSet", "A viewmodel path is required, but not set.") :
 							LOCTEXT("WidgetPathNotSet", "A widget path is required, but not set."),
+					EBindingMessageType::Error,
 					ArgumentName
 				);
 				return false;
@@ -510,7 +537,7 @@ void FMVVMViewBlueprintCompiler::CreateSourceLists(const FWidgetBlueprintCompile
 			{
 				if (Binding.BindingType == EMVVMBindingMode::TwoWay)
 				{
-					Self->AddErrorForBinding(Binding, LOCTEXT("TwoWayBindingsWithConversion", "Two-way bindings are not allowed to use conversion functions."));
+					Self->AddMessageForBinding(Binding, BlueprintView, LOCTEXT("TwoWayBindingsWithConversion", "Two-way bindings are not allowed to use conversion functions."), EBindingMessageType::Error);
 					bAreSourceContextsValid = false;
 					continue;
 				}
@@ -690,11 +717,12 @@ bool FMVVMViewBlueprintCompiler::PreCompileBindingSources(UWidgetBlueprintGenera
 			const TValueOrError<FBindingSourceContext, FText> CreatedBindingSourceContext = CreateBindingSourceContext(BlueprintView, Class, Path);
 			if (CreatedBindingSourceContext.HasError())
 			{
-				AddErrorForBinding(Binding, 
+				AddMessageForBinding(Binding, BlueprintView,
 					FText::Format(LOCTEXT("PropertyPathInvalidWithReason", "The property path '{0}' is invalid. {1}"),
 						PropertyPathToText(BlueprintView, Binding.SourcePath),
 						CreatedBindingSourceContext.GetError()
 					),
+					EBindingMessageType::Error,
 					ArgumentName
 				);
 				return false;
@@ -703,10 +731,12 @@ bool FMVVMViewBlueprintCompiler::PreCompileBindingSources(UWidgetBlueprintGenera
 			FBindingSourceContext BindingSourceContext = CreatedBindingSourceContext.GetValue();
 			if (!IsPropertyPathValid(BindingSourceContext.PropertyPath))
 			{
-				AddErrorForBinding(Binding, 
+				AddMessageForBinding(Binding,
+					BlueprintView,
 					FText::Format(LOCTEXT("PropertyPathIsInvalid", "The property path '{0}' is invalid."), 
 						PropertyPathToText(BlueprintView, Binding.SourcePath)
 					),
+					EBindingMessageType::Error,
 					ArgumentName
 				);
 				return false;
@@ -714,13 +744,13 @@ bool FMVVMViewBlueprintCompiler::PreCompileBindingSources(UWidgetBlueprintGenera
 
 			if (BindingSourceContext.SourceClass == nullptr)
 			{
-				AddErrorForBinding(Binding, LOCTEXT("BindingInvalidSourceClass", "Internal error. The binding could not find its source class."), ArgumentName);
+				AddMessageForBinding(Binding, BlueprintView, LOCTEXT("BindingInvalidSourceClass", "Internal error. The binding could not find its source class."), EBindingMessageType::Error, ArgumentName);
 				return false;
 			}
 
 			if (!BindingSourceContext.bIsRootWidget && BindingSourceContext.UserWidgetPropertyContextIndex == INDEX_NONE && BindingSourceContext.SourceCreatorContextIndex == INDEX_NONE)
 			{
-				AddErrorForBinding(Binding, LOCTEXT("BindingInvalidSource", "Internal error. The binding could not find its source."), ArgumentName);
+				AddMessageForBinding(Binding, BlueprintView, LOCTEXT("BindingInvalidSource", "Internal error. The binding could not find its source."), EBindingMessageType::Error, ArgumentName);
 				return false;
 			}
 
@@ -740,12 +770,68 @@ bool FMVVMViewBlueprintCompiler::PreCompileBindingSources(UWidgetBlueprintGenera
 			}
 		};
 
+		auto AddWarningForPropertyWithMVVMAndLegacyBinding = [this, &Binding, &BlueprintView, Class](const FMVVMBlueprintPropertyPath& Path)
+		{
+			if (!Path.HasPaths())
+			{
+				return;
+			}
+
+			// There can't be a legacy binding in the local scope, so we can skip this if the MVVM binding refers to a property in local scope.
+			if (Path.HasFieldInLocalScope())
+			{
+				return;
+			}
+
+			TArrayView<FMVVMBlueprintFieldPath const> MVVMBindingPath = Path.GetFieldPaths();
+			TArray< FDelegateRuntimeBinding > LegacyBindings = Class->Bindings;
+			FName MVVMFieldName = Path.GetPaths().Last();
+			FName MVVMObjectName = Path.GetWidgetName();
+
+			if (Path.GetFieldPaths().Last().GetBindingKind() == EBindingKind::Function)
+			{
+				return;
+			}
+
+			// If the first field is a UserWidget, we know this property resides in a nested UserWidget.
+			if (MVVMBindingPath[0].GetParentClass() && MVVMBindingPath[0].GetParentClass()->IsChildOf(UUserWidget::StaticClass()) && MVVMBindingPath.Num() > 1)
+			{
+				if (UWidgetBlueprintGeneratedClass* NestedBPGClass = Cast<UWidgetBlueprintGeneratedClass>(MVVMBindingPath[MVVMBindingPath.Num() - 2].GetParentClass()))
+				{
+					LegacyBindings = NestedBPGClass->Bindings;
+
+					// We can't use Path.GetWidgetName() when we are dealing with nested UserWidgets, because it refers to the topmost UserWidget.
+					MVVMObjectName = MVVMBindingPath[MVVMBindingPath.Num() - 2].GetFieldName();
+				}
+				else
+				{
+					return;
+				}
+			}
+
+			for (const FDelegateRuntimeBinding& LegacyBinding : LegacyBindings)
+			{
+				if (LegacyBinding.ObjectName == MVVMObjectName) 
+				{
+					if (LegacyBinding.PropertyName == MVVMFieldName)
+					{
+						AddMessageForBinding(Binding, BlueprintView, LOCTEXT("BindingConflictWithLegacy", "The binding is set on a property with legacy binding."), EBindingMessageType::Warning);
+						break;
+					}
+				}
+			}
+		};
+
 		if (IsForwardBinding(Binding.BindingType))
 		{
 			CreateSourcesForConversionFunction(true);
 
 			if (!Binding.SourcePath.IsEmpty())
 			{
+				if (!Binding.DestinationPath.IsEmpty())
+				{
+					AddWarningForPropertyWithMVVMAndLegacyBinding(Binding.DestinationPath);
+				}
 				if (!CreateSourceContextForPropertyPath(Binding.SourcePath, true))
 				{
 					bIsBindingsValid = false;
@@ -760,6 +846,10 @@ bool FMVVMViewBlueprintCompiler::PreCompileBindingSources(UWidgetBlueprintGenera
 
 			if (!Binding.DestinationPath.IsEmpty() && Binding.DestinationPath.HasPaths())
 			{
+				if (!Binding.SourcePath.IsEmpty())
+				{
+					AddWarningForPropertyWithMVVMAndLegacyBinding(Binding.SourcePath);
+				}
 				if (!CreateSourceContextForPropertyPath(Binding.DestinationPath, false))
 				{
 					bIsBindingsValid = false;
@@ -1072,15 +1162,15 @@ bool FMVVMViewBlueprintCompiler::PreCompileBindings(UWidgetBlueprintGeneratedCla
 		{
 			if (!GetDefault<UMVVMDeveloperProjectSettings>()->IsExecutionModeAllowed(Binding.OverrideExecutionMode))
 			{
-				AddErrorForBinding(Binding, LOCTEXT("NotAllowedExecutionMode", "The binding has a restricted execution mode."));
+				AddMessageForBinding(Binding, BlueprintView, LOCTEXT("NotAllowedExecutionMode", "The binding has a restricted execution mode."), EBindingMessageType::Error);
 			}
 		}
 
 		TValueOrError<FCompiledBindingLibraryCompiler::FFieldIdHandle, FText> AddFieldResult = AddFieldId(BindingSourceContext.SourceClass, true, Binding.BindingType, BindingSourceContext.FieldId.GetFieldName());
 		if (AddFieldResult.HasError())
 		{
-			AddErrorForBinding(Binding, FText::Format(LOCTEXT("CouldNotCreateSource", "Could not create source. {0}"), 
-				AddFieldResult.GetError()));
+			AddMessageForBinding(Binding, BlueprintView, FText::Format(LOCTEXT("CouldNotCreateSource", "Could not create source. {0}"),
+				AddFieldResult.GetError()), EBindingMessageType::Error);
 			bIsBindingsValid = false;
 			continue;
 		}
@@ -1091,8 +1181,9 @@ bool FMVVMViewBlueprintCompiler::PreCompileBindings(UWidgetBlueprintGeneratedCla
 			SetterPath = CreateBindingDestinationPath(BlueprintView, Class, DestinationPath);
 			if (!IsPropertyPathValid(SetterPath))
 			{
-				AddErrorForBinding(Binding, FText::Format(LOCTEXT("PropertyPathIsInvalid", "The property path '{0}' is invalid."), 
-					PropertyPathToText(BlueprintView, DestinationPath))
+				AddMessageForBinding(Binding, BlueprintView, FText::Format(LOCTEXT("PropertyPathIsInvalid", "The property path '{0}' is invalid."),
+					PropertyPathToText(BlueprintView, DestinationPath)),
+					EBindingMessageType::Error
 				);
 				bIsBindingsValid = false;
 				continue;
@@ -1109,8 +1200,9 @@ bool FMVVMViewBlueprintCompiler::PreCompileBindings(UWidgetBlueprintGeneratedCla
 		const UFunction* ConversionFunction = ConversionFunctionReference.ResolveMember<UFunction>(Class);
 		if (!ConversionFunctionWrapper.IsNone() && ConversionFunction == nullptr)
 		{
-			AddErrorForBinding(Binding, FText::Format(LOCTEXT("ConversionFunctionNotFound", "The conversion function '{0}' could not be found."),
-				FText::FromName(ConversionFunctionReference.GetMemberName()))
+			AddMessageForBinding(Binding, BlueprintView, FText::Format(LOCTEXT("ConversionFunctionNotFound", "The conversion function '{0}' could not be found."), 
+				FText::FromName(ConversionFunctionReference.GetMemberName())),
+				EBindingMessageType::Error
 			);
 			bIsBindingsValid = false;
 			continue;
@@ -1119,8 +1211,9 @@ bool FMVVMViewBlueprintCompiler::PreCompileBindings(UWidgetBlueprintGeneratedCla
 		TValueOrError<FCompilerBinding, FText> AddBindingResult = AddBinding(Class, BindingSourceContext.PropertyPath, SetterPath, ConversionFunction);
 		if (AddBindingResult.HasError())
 		{
-			AddErrorForBinding(Binding,
-				FText::Format(LOCTEXT("CouldNotCreateBinding", "Could not create binding. {0}"), AddBindingResult.GetError())
+			AddMessageForBinding(Binding, BlueprintView,
+				FText::Format(LOCTEXT("CouldNotCreateBinding", "Could not create binding. {0}"), AddBindingResult.GetError()),
+				EBindingMessageType::Error
 			);
 			bIsBindingsValid = false;
 			continue;
@@ -1180,8 +1273,9 @@ bool FMVVMViewBlueprintCompiler::CompileBindings(const FCompiledBindingLibraryCo
 		const FMVVMVCompiledFieldId* CompiledFieldId = CompileResult.FieldIds.Find(CompileBinding.FieldIdHandle);
 		if (CompiledFieldId == nullptr && CompileBinding.bFieldIdNeeded)
 		{
-			AddErrorForBinding(ViewBinding,
-				FText::Format(LOCTEXT("FieldIdNotGenerated", "Could not generate field ID for property '{0}'."), FText::FromName(NewBinding.SourcePropertyName))
+			AddMessageForBinding(ViewBinding, BlueprintView,
+				FText::Format(LOCTEXT("FieldIdNotGenerated", "Could not generate field ID for property '{0}'."), FText::FromName(NewBinding.SourcePropertyName)),
+				EBindingMessageType::Error
 			);
 			bIsBindingsValid = false;
 			continue;
@@ -1190,7 +1284,7 @@ bool FMVVMViewBlueprintCompiler::CompileBindings(const FCompiledBindingLibraryCo
 		const FMVVMVCompiledBinding* CompiledBinding = CompileResult.Bindings.Find(CompileBinding.BindingHandle);
 		if (CompiledBinding == nullptr)
 		{
-			AddErrorForBinding(ViewBinding, LOCTEXT("CompiledBindingNotGenerated", "Could not generate compiled binding."));
+			AddMessageForBinding(ViewBinding, BlueprintView, LOCTEXT("CompiledBindingNotGenerated", "Could not generate compiled binding."), EBindingMessageType::Error);
 			bIsBindingsValid = false;
 			continue;
 		}
@@ -1213,7 +1307,7 @@ bool FMVVMViewBlueprintCompiler::CompileBindings(const FCompiledBindingLibraryCo
 				const FMVVMBlueprintViewModelContext* ModelContext = BlueprintView->FindViewModel(ViewModelId);
 				if (ModelContext == nullptr)
 				{
-					AddErrorForBinding(ViewBinding, LOCTEXT("CompiledBindingWithInvalidIVewModelId", "Internal error: the view model became invalid."));
+					AddMessageForBinding(ViewBinding, BlueprintView, LOCTEXT("CompiledBindingWithInvalidIVewModelId", "Internal error: the view model became invalid."), EBindingMessageType::Error);
 					bIsBindingsValid = false;
 					continue;
 				}
diff --git a/Engine/Plugins/Runtime/ModelViewViewModel/Source/ModelViewViewModelBlueprint/Private/MVVMViewBlueprintCompiler.h b/Engine/Plugins/Runtime/ModelViewViewModel/Source/ModelViewViewModelBlueprint/Private/MVVMViewBlueprintCompiler.h
--- a/Engine/Plugins/Runtime/ModelViewViewModel/Source/ModelViewViewModelBlueprint/Private/MVVMViewBlueprintCompiler.h
+++ b/Engine/Plugins/Runtime/ModelViewViewModel/Source/ModelViewViewModelBlueprint/Private/MVVMViewBlueprintCompiler.h
@@ -16,6 +16,10 @@ class UEdGraph;
 class UMVVMBlueprintView;
 class UMVVMViewClass;
 class UWidgetBlueprintGeneratedClass;
+namespace  UE::MVVM
+{
+	enum class EBindingMessageType : uint8;
+}
 
 namespace UE::MVVM::Private
 {
@@ -69,7 +73,7 @@ private:
 
 	const FCompilerSourceCreatorContext* FindViewModelSource(FGuid Id) const;
 
-	void AddErrorForBinding(FMVVMBlueprintViewBinding& Binding, const FText& Message, FName ArgumentName = FName()) const;
+	void AddMessageForBinding(FMVVMBlueprintViewBinding& Binding, UMVVMBlueprintView* BlueprintView, const FText& MessageText, EBindingMessageType MessageType, FName ArgumentName = FName()) const;
 	void AddErrorForViewModel(const FMVVMBlueprintViewModelContext& ViewModel, const FText& Message) const;
 
 	TValueOrError<FBindingSourceContext, FText> CreateBindingSourceContext(const UMVVMBlueprintView* BlueprintView, const UWidgetBlueprintGeneratedClass* Class, const FMVVMBlueprintPropertyPath& PropertyPath);
diff --git a/Engine/Plugins/Runtime/ModelViewViewModel/Source/ModelViewViewModelBlueprint/Private/MVVMWidgetBlueprintExtension_View.cpp b/Engine/Plugins/Runtime/ModelViewViewModel/Source/ModelViewViewModelBlueprint/Private/MVVMWidgetBlueprintExtension_View.cpp
--- a/Engine/Plugins/Runtime/ModelViewViewModel/Source/ModelViewViewModelBlueprint/Private/MVVMWidgetBlueprintExtension_View.cpp
+++ b/Engine/Plugins/Runtime/ModelViewViewModel/Source/ModelViewViewModelBlueprint/Private/MVVMWidgetBlueprintExtension_View.cpp
@@ -50,11 +50,7 @@ void UMVVMWidgetBlueprintExtension_View::HandleBeginCompilation(FWidgetBlueprint
 	CurrentCompilerContext.Reset();
 	if (BlueprintView)
 	{
-		for (FMVVMBlueprintViewBinding& Binding : BlueprintView->GetBindings())
-		{
-			Binding.Errors.Reset();
-		}
-
+		BlueprintView->ResetBindingMessages();
 		CurrentCompilerContext = MakePimpl<UE::MVVM::Private::FMVVMViewBlueprintCompiler>(InCreationContext);
 	}
 }
diff --git a/Engine/Plugins/Runtime/ModelViewViewModel/Source/ModelViewViewModelBlueprint/Public/MVVMBlueprintView.h b/Engine/Plugins/Runtime/ModelViewViewModel/Source/ModelViewViewModelBlueprint/Public/MVVMBlueprintView.h
--- a/Engine/Plugins/Runtime/ModelViewViewModel/Source/ModelViewViewModelBlueprint/Public/MVVMBlueprintView.h
+++ b/Engine/Plugins/Runtime/ModelViewViewModel/Source/ModelViewViewModelBlueprint/Public/MVVMBlueprintView.h
@@ -12,6 +12,22 @@ class UMVVMWidgetBlueprintExtension_View;
 class UWidget;
 class UWidgetBlueprint;
 
+namespace  UE::MVVM
+{
+	enum class EBindingMessageType : uint8
+	{
+		Info,
+		Warning,
+		Error
+	};
+
+	struct FBindingMessage
+	{
+		FText MessageText;
+		EBindingMessageType MessageType;
+	};
+}
+
 /**
  * 
  */
@@ -65,6 +81,11 @@ public:
 		return Bindings;
 	}
 
+	TArray<FText> GetBindingMessages(FGuid Id, UE::MVVM::EBindingMessageType InMessageType) const;
+	bool HasBindingMessage(FGuid Id, UE::MVVM::EBindingMessageType InMessageType) const;
+	void AddMessageToBinding(FGuid Id, UE::MVVM::FBindingMessage MessageToAdd);
+	void ResetBindingMessages();
+
 #if WITH_EDITOR
 //	virtual void PostEditUndo() override;
 
@@ -93,6 +114,10 @@ private:
 
 	UPROPERTY(EditAnywhere, Category = "Viewmodel")
 	TArray<FMVVMBlueprintViewModelContext> AvailableViewModels;
+
+	TMap<FGuid, TArray<UE::MVVM::FBindingMessage>> BindingMessages;
+
+	bool bIsContextSensitive;
 };
 
 #if UE_ENABLE_INCLUDE_ORDER_DEPRECATED_IN_5_2
diff --git a/Engine/Plugins/Runtime/ModelViewViewModel/Source/ModelViewViewModelBlueprint/Public/MVVMBlueprintViewBinding.h b/Engine/Plugins/Runtime/ModelViewViewModel/Source/ModelViewViewModelBlueprint/Public/MVVMBlueprintViewBinding.h
--- a/Engine/Plugins/Runtime/ModelViewViewModel/Source/ModelViewViewModelBlueprint/Public/MVVMBlueprintViewBinding.h
+++ b/Engine/Plugins/Runtime/ModelViewViewModel/Source/ModelViewViewModelBlueprint/Public/MVVMBlueprintViewBinding.h
@@ -62,10 +62,6 @@ struct MODELVIEWVIEWMODELBLUEPRINT_API FMVVMBlueprintViewBinding
 	UPROPERTY(EditAnywhere, Category = "Viewmodel")
 	FMVVMBlueprintViewConversionPath Conversion;
 
-	/** */
-	UPROPERTY(VisibleAnywhere, Category = "Viewmodel", Transient)
-	TArray<FText> Errors;
-
 	/** The unique ID of this binding. */
 	UPROPERTY(VisibleAnywhere, Category = "Viewmodel")
 	FGuid BindingId;
diff --git a/Engine/Plugins/Runtime/ModelViewViewModel/Source/ModelViewViewModelBlueprint/Public/MVVMPropertyPath.h b/Engine/Plugins/Runtime/ModelViewViewModel/Source/ModelViewViewModelBlueprint/Public/MVVMPropertyPath.h
--- a/Engine/Plugins/Runtime/ModelViewViewModel/Source/ModelViewViewModelBlueprint/Public/MVVMPropertyPath.h
+++ b/Engine/Plugins/Runtime/ModelViewViewModel/Source/ModelViewViewModelBlueprint/Public/MVVMPropertyPath.h
@@ -99,6 +99,20 @@ public:
 		return UE::MVVM::FMVVMConstFieldVariant();
 	}
 
+	bool IsFieldLocalScope() const
+	{
+		return BindingReference.IsLocalScope();
+	}
+
+	UClass* GetParentClass() const
+	{
+		return BindingReference.GetMemberParentClass();
+	}
+
+	EBindingKind GetBindingKind() const
+	{
+		return BindingKind;
+	}
 	/** */
 	void SetBindingReference(UE::MVVM::FMVVMConstFieldVariant InField)
 	{
@@ -232,6 +246,11 @@ public:
 		return Result;
 	}
 
+	TArrayView<FMVVMBlueprintFieldPath const> GetFieldPaths() const
+	{
+		return Paths;
+	}
+
 	/** Get the binding name, resolves reference deprecation / redirectors / etc before returning */
 	TArray<UE::MVVM::FMVVMConstFieldVariant> GetFields() const
 	{
@@ -246,6 +265,19 @@ public:
 		return Result;
 	}
 
+	bool HasFieldInLocalScope() const
+	{
+		for (const FMVVMBlueprintFieldPath& Path : Paths)
+		{
+			if (Path.IsFieldLocalScope())
+			{
+				return true;
+			}
+		}
+
+		return false;
+	}
+
 	bool BasePropertyPathContains(UE::MVVM::FMVVMConstFieldVariant Field) const
 	{
 		return Paths.ContainsByPredicate([Field](const FMVVMBlueprintFieldPath& FieldPath) { return FieldPath.GetField() == Field;  });
diff --git a/Engine/Plugins/Runtime/ModelViewViewModel/Source/ModelViewViewModelEditor/Private/Widgets/SMVVMViewBindingListView.cpp b/Engine/Plugins/Runtime/ModelViewViewModel/Source/ModelViewViewModelEditor/Private/Widgets/SMVVMViewBindingListView.cpp
--- a/Engine/Plugins/Runtime/ModelViewViewModel/Source/ModelViewViewModelEditor/Private/Widgets/SMVVMViewBindingListView.cpp
+++ b/Engine/Plugins/Runtime/ModelViewViewModel/Source/ModelViewViewModelEditor/Private/Widgets/SMVVMViewBindingListView.cpp
@@ -288,134 +288,140 @@ public:
 			.ShowWires(true)
 			.Style(FMVVMEditorStyle::Get(), "BindingView.BindingRow")
 			[
-				SNew(SBox)
-				.HeightOverride(30)
-				[
-					SNew(SHorizontalBox)
-
-					+ SHorizontalBox::Slot()
-					.Padding(2, 0)
-					.VAlign(VAlign_Center)
-					.HAlign(HAlign_Left) 
-					.AutoWidth()
+				SNew(SBorder)
+				.BorderImage(FAppStyle::Get().GetBrush("PlainBorder"))
+			    .Padding(0)
+			    .BorderBackgroundColor(this, &SBindingRow::GetErrorBorderColor)
+			     [
+					SNew(SBox)
+					.HeightOverride(30)
 					[
-						SNew(SCheckBox)
-						.IsChecked(this, &SBindingRow::IsBindingCompiled)
-						.OnCheckStateChanged(this, &SBindingRow::OnIsBindingCompileChanged)
-					]
+						SNew(SHorizontalBox)
 
-					+ SHorizontalBox::Slot()
-					.Padding(2, 0)
-					.VAlign(VAlign_Center)
-					.HAlign(HAlign_Left) 
-					.AutoWidth()
-					[
-						SNew(SSimpleButton)
-						.Icon(FAppStyle::Get().GetBrush("Icons.Error"))
-						.Visibility(this, &SBindingRow::GetErrorVisibility)
-						.ToolTipText(this, &SBindingRow::GetErrorToolTip)
-						.OnClicked(this, &SBindingRow::OnErrorButtonClicked)
-					]
+						+ SHorizontalBox::Slot()
+						.Padding(2, 0)
+						.VAlign(VAlign_Center)
+						.HAlign(HAlign_Left) 
+						.AutoWidth()
+						[
+							SNew(SCheckBox)
+							.IsChecked(this, &SBindingRow::IsBindingCompiled)
+							.OnCheckStateChanged(this, &SBindingRow::OnIsBindingCompileChanged)
+						]
 
-					+ SHorizontalBox::Slot()
-					.Padding(4, 0)
-					.VAlign(VAlign_Center)
-					.HAlign(HAlign_Left)
-					.AutoWidth()
-					[
-						SNew(SBox)
-						.MinDesiredWidth(150)
+						+ SHorizontalBox::Slot()
+						.Padding(2, 0)
+						.VAlign(VAlign_Center)
+						.HAlign(HAlign_Left) 
+						.AutoWidth()
 						[
-							SAssignNew(WidgetFieldSelector, SFieldSelector, InWidgetBlueprint, false)
-							.BindingMode(this, &SBindingRow::GetCurrentBindingMode)
-							.SelectedField(this, &SBindingRow::GetSelectedWidgetProperty)
-							.OnFieldSelectionChanged(this, &SBindingRow::OnWidgetPropertySelected)
-							.SelectedConversionFunction(this, &SBindingRow::GetSelectedConversionFunction, false)
-							.OnConversionFunctionSelectionChanged(this, &SBindingRow::OnConversionFunctionChanged, false)
-							.ShowConversionFunctions(this, &SBindingRow::ShouldShowConversionFunctions, false)
-							.ShowSource(false)
-							.Source(WidgetSource)
-							.AssignableTo(this, &SBindingRow::GetAssignableToProperty, false)
+							SNew(SSimpleButton)
+							.Icon(FAppStyle::Get().GetBrush("Icons.Error"))
+							.Visibility(this, &SBindingRow::GetErrorButtonVisibility)
+							.ToolTipText(this, &SBindingRow::GetErrorButtonToolTip)
+							.OnClicked(this, &SBindingRow::OnErrorButtonClicked)
 						]
-					]
 
-					+ SHorizontalBox::Slot()
-					.Padding(4, 0)
-					.VAlign(VAlign_Fill)
-					.HAlign(HAlign_Left) 
-					.AutoWidth()
-					[
-						SNew(SComboBox<FName>)
-						.OptionsSource(Private::GetBindingModeNames())
-						.InitiallySelectedItem(StaticEnum<EMVVMBindingMode>()->GetNameByValue((int64) ViewBinding->BindingType))
-						.OnSelectionChanged(this, &SBindingRow::OnBindingModeSelectionChanged)
-						.OnGenerateWidget(this, &SBindingRow::GenerateBindingModeWidget)
-						.ToolTipText(this, &SBindingRow::GetCurrentBindingModeLabel)
-						.Content()
+						+ SHorizontalBox::Slot()
+						.Padding(4, 0)
+						.VAlign(VAlign_Center)
+						.HAlign(HAlign_Left)
+						.AutoWidth()
 						[
 							SNew(SBox)
-							.HAlign(HAlign_Center)
-							.VAlign(VAlign_Center)
-							.WidthOverride(16)
-							.HeightOverride(16)
+							.MinDesiredWidth(150)
 							[
-								SNew(SImage)
-								.Image(this, &SBindingRow::GetCurrentBindingModeBrush)
+								SAssignNew(WidgetFieldSelector, SFieldSelector, InWidgetBlueprint, false)
+								.BindingMode(this, &SBindingRow::GetCurrentBindingMode)
+								.SelectedField(this, &SBindingRow::GetSelectedWidgetProperty)
+								.OnFieldSelectionChanged(this, &SBindingRow::OnWidgetPropertySelected)
+								.SelectedConversionFunction(this, &SBindingRow::GetSelectedConversionFunction, false)
+								.OnConversionFunctionSelectionChanged(this, &SBindingRow::OnConversionFunctionChanged, false)
+								.ShowConversionFunctions(this, &SBindingRow::ShouldShowConversionFunctions, false)
+								.ShowSource(false)
+								.Source(WidgetSource)
+								.AssignableTo(this, &SBindingRow::GetAssignableToProperty, false)
+							]
+						]
+
+						+ SHorizontalBox::Slot()
+						.Padding(4, 0)
+						.VAlign(VAlign_Fill)
+						.HAlign(HAlign_Left) 
+						.AutoWidth()
+						[
+							SNew(SComboBox<FName>)
+							.OptionsSource(Private::GetBindingModeNames())
+							.InitiallySelectedItem(StaticEnum<EMVVMBindingMode>()->GetNameByValue((int64) ViewBinding->BindingType))
+							.OnSelectionChanged(this, &SBindingRow::OnBindingModeSelectionChanged)
+							.OnGenerateWidget(this, &SBindingRow::GenerateBindingModeWidget)
+							.ToolTipText(this, &SBindingRow::GetCurrentBindingModeLabel)
+							.Content()
+							[
+								SNew(SBox)
+								.HAlign(HAlign_Center)
+								.VAlign(VAlign_Center)
+								.WidthOverride(16)
+								.HeightOverride(16)
+								[
+									SNew(SImage)
+									.Image(this, &SBindingRow::GetCurrentBindingModeBrush)
+								]
 							]
 						]
-					]
 				
-					+ SHorizontalBox::Slot()
-					.Padding(4, 0, 2, 0)
-					.VAlign(VAlign_Center)
-					.HAlign(HAlign_Left)
-					.AutoWidth()
-					[
-						SNew(SBox)
-						.MinDesiredWidth(150)
+						+ SHorizontalBox::Slot()
+						.Padding(4, 0, 2, 0)
+						.VAlign(VAlign_Center)
+						.HAlign(HAlign_Left)
+						.AutoWidth()
 						[
-							SAssignNew(ViewModelFieldSelector, SFieldSelector, InWidgetBlueprint, true)
-							.BindingMode(this, &SBindingRow::GetCurrentBindingMode)
-							.SelectedField(this, &SBindingRow::GetSelectedViewModelProperty)
-							.OnFieldSelectionChanged(this, &SBindingRow::OnViewModelPropertySelected)
-							.SelectedConversionFunction(this, &SBindingRow::GetSelectedConversionFunction, true)
-							.OnConversionFunctionSelectionChanged(this, &SBindingRow::OnConversionFunctionChanged, true)
-							.ShowConversionFunctions(this, &SBindingRow::ShouldShowConversionFunctions, true)
-							.AssignableTo(this, &SBindingRow::GetAssignableToProperty, true)
+							SNew(SBox)
+							.MinDesiredWidth(150)
+							[
+								SAssignNew(ViewModelFieldSelector, SFieldSelector, InWidgetBlueprint, true)
+								.BindingMode(this, &SBindingRow::GetCurrentBindingMode)
+								.SelectedField(this, &SBindingRow::GetSelectedViewModelProperty)
+								.OnFieldSelectionChanged(this, &SBindingRow::OnViewModelPropertySelected)
+								.SelectedConversionFunction(this, &SBindingRow::GetSelectedConversionFunction, true)
+								.OnConversionFunctionSelectionChanged(this, &SBindingRow::OnConversionFunctionChanged, true)
+								.ShowConversionFunctions(this, &SBindingRow::ShouldShowConversionFunctions, true)
+								.AssignableTo(this, &SBindingRow::GetAssignableToProperty, true)
+							]
 						]
-					]
 
-					+ SHorizontalBox::Slot()
-					[
-						SNew(SSpacer)
-					]
+						+ SHorizontalBox::Slot()
+						[
+							SNew(SSpacer)
+						]
 
-					+ SHorizontalBox::Slot()
-					.Padding(2, 1)
-					.VAlign(VAlign_Fill)
-					.HAlign(HAlign_Right)
-					.AutoWidth()
-					[
-						SNew(SCheckBox)
-						.IsChecked(this, &SBindingRow::IsExecutionModeOverrideChecked)
-						.OnCheckStateChanged(this, &SBindingRow::OnExecutionModeOverrideChanged)
-					]
+						+ SHorizontalBox::Slot()
+						.Padding(2, 1)
+						.VAlign(VAlign_Fill)
+						.HAlign(HAlign_Right)
+						.AutoWidth()
+						[
+							SNew(SCheckBox)
+							.IsChecked(this, &SBindingRow::IsExecutionModeOverrideChecked)
+							.OnCheckStateChanged(this, &SBindingRow::OnExecutionModeOverrideChanged)
+						]
 
-					+ SHorizontalBox::Slot()
-					.Padding(2, 1)
-					.VAlign(VAlign_Fill)
-					.HAlign(HAlign_Right)
-					.AutoWidth()
-					[
-						SNew(SComboButton)
-						.ContentPadding(FMargin(4.f, 0.f))
-						.OnGetMenuContent(this, &SBindingRow::OnGetExecutionModeMenuContent)
-						.IsEnabled(this, &SBindingRow::IsExecutionModeOverridden)
-						.ButtonContent()
+						+ SHorizontalBox::Slot()
+						.Padding(2, 1)
+						.VAlign(VAlign_Fill)
+						.HAlign(HAlign_Right)
+						.AutoWidth()
 						[
-							SNew(STextBlock)
-							.Text(this, &SBindingRow::GetExecutioModeValue)
-							.ToolTipText(this, &SBindingRow::GetExecutioModeValueToolTip)
+							SNew(SComboButton)
+							.ContentPadding(FMargin(4.f, 0.f))
+							.OnGetMenuContent(this, &SBindingRow::OnGetExecutionModeMenuContent)
+							.IsEnabled(this, &SBindingRow::IsExecutionModeOverridden)
+							.ButtonContent()
+							[
+								SNew(STextBlock)
+								.Text(this, &SBindingRow::GetExecutioModeValue)
+								.ToolTipText(this, &SBindingRow::GetExecutioModeValueToolTip)
+							]
 						]
 					]
 				]
@@ -473,6 +479,24 @@ public:
 
 private:
 
+	FSlateColor GetErrorBorderColor() const
+	{
+		UMVVMEditorSubsystem* EditorSubsystem = GEditor->GetEditorSubsystem<UMVVMEditorSubsystem>();
+		UMVVMBlueprintView* BlueprintViewPtr = EditorSubsystem->GetView(WidgetBlueprintWeak.Get());
+		if (BlueprintViewPtr->HasBindingMessage(GetThisViewBinding()->BindingId, EBindingMessageType::Error))
+		{
+			return FStyleColors::Error;
+		}
+		else if (BlueprintViewPtr->HasBindingMessage(GetThisViewBinding()->BindingId, EBindingMessageType::Warning))
+		{
+			return FStyleColors::Warning;
+		}
+		else
+		{
+			return FStyleColors::Transparent;
+		}
+	}
+
 	ECheckBoxState IsBindingEnabled() const
 	{
 		if (const FMVVMBlueprintViewBinding* ViewModelBinding = GetThisViewBinding())
@@ -491,16 +515,26 @@ private:
 		return ECheckBoxState::Undetermined;
 	}
 
-	EVisibility GetErrorVisibility() const
+	EVisibility GetErrorButtonVisibility() const
 	{
-		return GetThisViewBinding()->Errors.IsEmpty() ? EVisibility::Collapsed : EVisibility::Visible;
+		UMVVMEditorSubsystem* EditorSubsystem = GEditor->GetEditorSubsystem<UMVVMEditorSubsystem>();
+		UMVVMBlueprintView* BlueprintViewPtr = EditorSubsystem->GetView(WidgetBlueprintWeak.Get());
+		bool HasBindingError = BlueprintViewPtr->HasBindingMessage(GetThisViewBinding()->BindingId, EBindingMessageType::Error) || BlueprintViewPtr->HasBindingMessage(GetThisViewBinding()->BindingId, EBindingMessageType::Warning);
+		return HasBindingError ? EVisibility::Visible : EVisibility::Collapsed;
 	}
 
-	FText GetErrorToolTip() const
+	FText GetErrorButtonToolTip() const
 	{
+		// Get error messages of this binding
+		UMVVMEditorSubsystem* EditorSubsystem = GEditor->GetEditorSubsystem<UMVVMEditorSubsystem>();
+		UMVVMBlueprintView* BlueprintViewPtr = EditorSubsystem->GetView(WidgetBlueprintWeak.Get());
+		TArray<FText> BindingErrorList = BlueprintViewPtr->GetBindingMessages(GetThisViewBinding()->BindingId, EBindingMessageType::Error);
+		TArray<FText> BindingWarningList = BlueprintViewPtr->GetBindingMessages(GetThisViewBinding()->BindingId, EBindingMessageType::Warning);
+		BindingErrorList.Append(BindingWarningList);
+
 		static const FText NewLineText = FText::FromString(TEXT("\n"));
-		FText HintText = LOCTEXT("ErrorButtonText", "Errors: (Click to show in a separate window)");
-		FText ErrorsText = FText::Join(NewLineText, GetThisViewBinding()->Errors);
+		FText HintText = LOCTEXT("ErrorButtonText", "Errors and Warnings: (Click to show in a separate window)");
+		FText ErrorsText = FText::Join(NewLineText, BindingErrorList);
 		return FText::Join(NewLineText, HintText, ErrorsText);
 	}
 
@@ -515,14 +549,19 @@ private:
 		{
 			if (const FMVVMBlueprintViewBinding* ViewBinding = Entry->GetBinding(View))
 			{
-				for (const FText& ErrorText : ViewBinding->Errors)
+				for (const FText& ErrorText : View->GetBindingMessages(ViewBinding->BindingId, EBindingMessageType::Error))
 				{
 					ErrorItems.Add(MakeShared<FText>(ErrorText));
 				}
 
+				for (const FText& WarningText : View->GetBindingMessages(ViewBinding->BindingId, EBindingMessageType::Warning))
+				{
+					ErrorItems.Add(MakeShared<FText>(WarningText));
+				}
+
 				const FText BindingDisplayName = FText::FromString(ViewBinding->GetDisplayNameString(WidgetBlueprint));
 				ErrorDialog = SNew(SCustomDialog)
-					.Title(FText::Format(LOCTEXT("Compilation Errors", "Compilation Errors for {0}"), BindingDisplayName))
+					.Title(FText::Format(LOCTEXT("Compilation Errors and Warnings", "Compilation Errors and Warnings for {0}"), BindingDisplayName))
 					.Buttons({
 						SCustomDialog::FButton(LOCTEXT("OK", "OK"))
 					})
