MVVM: Uses FText::Identical to compare FText. We want the widget to have the correct key. It may cause the widget to update/redraw. If the value that corresponds to that key changes at runtime, the widget needs the correct representation.

diff --git a/Engine/Plugins/Runtime/ModelViewViewModel/Source/ModelViewViewModel/Public/MVVMViewModelBase.h b/Engine/Plugins/Runtime/ModelViewViewModel/Source/ModelViewViewModel/Public/MVVMViewModelBase.h
--- a/Engine/Plugins/Runtime/ModelViewViewModel/Source/ModelViewViewModel/Public/MVVMViewModelBase.h
+++ b/Engine/Plugins/Runtime/ModelViewViewModel/Source/ModelViewViewModel/Public/MVVMViewModelBase.h
@@ -31,7 +31,7 @@ namespace UMMVVMViewModelBase
 
 
 /** Base class for MVVM viewmodel. */
-UCLASS(Blueprintable, Abstract, DisplayName="MVVM ViewModel")
+UCLASS(Blueprintable, Abstract, DisplayName="MVVM Base Viewmodel")
 class MODELVIEWVIEWMODEL_API UMVVMViewModelBase : public UObject, public INotifyFieldValueChanged
 {
 	GENERATED_BODY()
@@ -84,7 +84,7 @@ protected:
 
 	bool SetPropertyValue(FText& Value, const FText& NewValue, UE::FieldNotification::FFieldId FieldId)
 	{
-		if (Value.EqualTo(NewValue))
+		if (Value.IdenticalTo(Value))
 		{
 			return false;
 		}
