MVVM: Remove needs to have a SourcePath for ComplexConversionFunction. The source data is inside the function itself and the fieldnotify name is already saved in the class binding.
Update GetBindingInfoFromFieldPath to support viewmodel inside another viewmodel. Add unittest for the function.
Add DoesInitializeAtConstruction on the view class. That is to support manual initialization. When the flag is false, the user is responsible for initializing the view.
Add ReevaluateSourceCreator. A source that was discovered at compile time and added to the list of SourceCreator may need to be evaluated if the path contains notify. That function is similar to SetViewModel but finds the new value from the SourceCreator.
Fix issue when a binding destination can't be evaluated at initialization. Since we only reactivate the bindings link to a source in SetViewModel, some bindings were never registered and never executed after the SetViewModel properly set the binding destination.
Update the view class ToString functions with the new flags.
Prevent adding a new binding when the editor data contains a conversion function and a regular function. That may occur when the user started with a regular binding and changed it afterward for a conversion function. We may not want to reset the binding to be non destructive.
#prefligth 6419e837bb1320be412d274d, 641b068aee0aed13479e6ea3

diff --git a/Engine/Plugins/Runtime/ModelViewViewModel/Source/ModelViewViewModel/ModelViewViewModel.Build.cs b/Engine/Plugins/Runtime/ModelViewViewModel/Source/ModelViewViewModel/ModelViewViewModel.Build.cs
--- a/Engine/Plugins/Runtime/ModelViewViewModel/Source/ModelViewViewModel/ModelViewViewModel.Build.cs
+++ b/Engine/Plugins/Runtime/ModelViewViewModel/Source/ModelViewViewModel/ModelViewViewModel.Build.cs
@@ -23,5 +23,14 @@ public class ModelViewViewModel : ModuleRules
 				"Slate",
 				"UMG",
 			});
+
+		if (Target.Configuration == UnrealTargetConfiguration.Test || Target.Configuration == UnrealTargetConfiguration.Shipping)
+		{
+			PublicDefinitions.Add("UE_WITH_MVVM_DEBUGGING=0");
+		}
+		else
+		{
+			PublicDefinitions.Add("UE_WITH_MVVM_DEBUGGING=1");
+		}
 	}
 }
diff --git a/Engine/Plugins/Runtime/ModelViewViewModel/Source/ModelViewViewModel/Private/Bindings/MVVMFieldPathHelper.cpp b/Engine/Plugins/Runtime/ModelViewViewModel/Source/ModelViewViewModel/Private/Bindings/MVVMFieldPathHelper.cpp
--- a/Engine/Plugins/Runtime/ModelViewViewModel/Source/ModelViewViewModel/Private/Bindings/MVVMFieldPathHelper.cpp
+++ b/Engine/Plugins/Runtime/ModelViewViewModel/Source/ModelViewViewModel/Private/Bindings/MVVMFieldPathHelper.cpp
@@ -272,8 +272,12 @@ TValueOrError<TArray<FMVVMConstFieldVariant>, FText> GenerateFieldPathList(TArra
 }
 
 
-TValueOrError<FParsedBindingInfo, FText> GetBindingInfoFromFieldPath(const UClass* InAccessor, TArrayView<const FMVVMConstFieldVariant> InFieldPath)
+TValueOrError<FParsedNotifyBindingInfo, FText> GetNotifyBindingInfoFromFieldPath(const UClass* InAccessor, TArrayView<const FMVVMConstFieldVariant> InFieldPath)
 {
+	// The class is the UserWidget.
+	// We assume that the first property is the viewmodel/widget added in the editor
+	// To be a valid oneway path, all properties after the initial viewmodel has to be "notify"
+	// or after the last "notify" all the path after are onetime (example vectornotify.x)
 	if (InAccessor == nullptr)
 	{
 		return MakeError(LOCTEXT("AccessorClassNull", "The accessor class is null."));
@@ -288,8 +292,8 @@ TValueOrError<FParsedBindingInfo, FText> GetBindingInfoFromFieldPath(const UClas
 		return MakeError(LOCTEXT("FieldPathDoesntStartWithAccessor", "Field path that doesn't start with the accessor is not supported."));
 	}
 
-	FParsedBindingInfo Result;
-
+	FParsedNotifyBindingInfo Result;
+	bool bFoundViewModel = false;
 	for (int32 Index = InFieldPath.Num() - 1; Index >= 0; --Index)
 	{
 		FMVVMConstFieldVariant Field = InFieldPath[Index];
@@ -341,11 +345,22 @@ TValueOrError<FParsedBindingInfo, FText> GetBindingInfoFromFieldPath(const UClas
 			FMVVMAvailableBinding AvailableBinding = UMVVMSubsystem::GetAvailableBindingForField(Field, InAccessor);
 			if (AvailableBinding.IsValid() && AvailableBinding.HasNotify())
 			{
-				Result.NotifyFieldClass = const_cast<UClass*>(FieldOwnerAsClass);
-				Result.NotifyFieldId = FFieldNotificationId(AvailableBinding.GetBindingName().ToName());
-				Result.ViewModelIndex = Index-1;
-				break;
+				if (!bFoundViewModel)
+				{
+					Result.NotifyFieldClass = const_cast<UClass*>(FieldOwnerAsClass);
+					Result.NotifyFieldId = FFieldNotificationId(AvailableBinding.GetBindingName().ToName());
+					Result.ViewModelIndex = Index - 1;
+					bFoundViewModel = true;
+				}
 			}
+			else if (bFoundViewModel && Index != 0)
+			{
+				Result = FParsedNotifyBindingInfo();
+			}
+		}
+		else if (bFoundViewModel)
+		{
+			Result = FParsedNotifyBindingInfo();
 		}
 	}
 
diff --git a/Engine/Plugins/Runtime/ModelViewViewModel/Source/ModelViewViewModel/Private/MVVMSubsystem.cpp b/Engine/Plugins/Runtime/ModelViewViewModel/Source/ModelViewViewModel/Private/MVVMSubsystem.cpp
--- a/Engine/Plugins/Runtime/ModelViewViewModel/Source/ModelViewViewModel/Private/MVVMSubsystem.cpp
+++ b/Engine/Plugins/Runtime/ModelViewViewModel/Source/ModelViewViewModel/Private/MVVMSubsystem.cpp
@@ -335,30 +335,40 @@ TValueOrError<bool, FText> UMVVMSubsystem::IsBindingValid(FDirectionalBindingArg
 
 TValueOrError<bool, FText> UMVVMSubsystem::IsBindingValid(FConstDirectionalBindingArgs Args) const
 {
+	const bool bIsSimpleConversionFunction = Args.ConversionFunction && UE::MVVM::BindingHelper::IsValidForSimpleRuntimeConversion(Args.ConversionFunction);
+	const bool bIsComplexConversionFunction = Args.ConversionFunction && UE::MVVM::BindingHelper::IsValidForComplexRuntimeConversion(Args.ConversionFunction);
+
 	// Test Source
-	TValueOrError<const FProperty*, FText> SourceResult = UE::MVVM::BindingHelper::TryGetPropertyTypeForSourceBinding(Args.SourceBinding);
-	if (SourceResult.HasError())
+	const FProperty* SourceProperty = nullptr;
+	if (!bIsComplexConversionFunction)
 	{
-		return MakeError(SourceResult.StealError());
-	}
+		TValueOrError<const FProperty*, FText> SourceResult = UE::MVVM::BindingHelper::TryGetPropertyTypeForSourceBinding(Args.SourceBinding);
+		if (SourceResult.HasError())
+		{
+			return MakeError(SourceResult.StealError());
+		}
 
-	const FProperty* SourceProperty = SourceResult.StealValue();
-	if (SourceProperty == nullptr)
-	{
-		return MakeError(LOCTEXT("NoValueToBindAtSource", "There is no value to bind at the source."));
+		SourceProperty = SourceResult.StealValue();
+		if (SourceProperty == nullptr)
+		{
+			return MakeError(LOCTEXT("NoValueToBindAtSource", "There is no value to bind at the source."));
+		}
 	}
 
 	// Test Destination
-	TValueOrError<const FProperty*, FText> DestinationResult = UE::MVVM::BindingHelper::TryGetPropertyTypeForDestinationBinding(Args.DestinationBinding);
-	if (DestinationResult.HasError())
+	const FProperty* DestinationProperty = nullptr;
 	{
-		return MakeError(DestinationResult.StealError());
-	}
+		TValueOrError<const FProperty*, FText> DestinationResult = UE::MVVM::BindingHelper::TryGetPropertyTypeForDestinationBinding(Args.DestinationBinding);
+		if (DestinationResult.HasError())
+		{
+			return MakeError(DestinationResult.StealError());
+		}
 
-	const FProperty* DestinationProperty = DestinationResult.StealValue();
-	if (DestinationProperty == nullptr)
-	{
-		return MakeError(LOCTEXT("NoValueToBindAtDestination", "There is no value to bind at the destination."));
+		DestinationProperty = DestinationResult.StealValue();
+		if (DestinationProperty == nullptr)
+		{
+			return MakeError(LOCTEXT("NoValueToBindAtDestination", "There is no value to bind at the destination."));
+		}
 	}
 
 	auto GetPropertyType = [](const FProperty* Property)
@@ -381,8 +391,6 @@ TValueOrError<bool, FText> UMVVMSubsystem::IsBindingValid(FConstDirectionalBindi
 			return MakeError(ReturnResult.StealError());
 		}
 
-		const bool bIsSimpleConversionFunction = UE::MVVM::BindingHelper::IsValidForSimpleRuntimeConversion(Args.ConversionFunction);
-		const bool bIsComplexConversionFunction = UE::MVVM::BindingHelper::IsValidForComplexRuntimeConversion(Args.ConversionFunction);
 		if (!bIsSimpleConversionFunction && !bIsComplexConversionFunction)
 		{
 			return MakeError(FText::Format(LOCTEXT("ConversionFunctionNotValid", "The conversion function '{0}' is not valid as a conversion function"), FText::FromString(Args.ConversionFunction->GetName())));
diff --git a/Engine/Plugins/Runtime/ModelViewViewModel/Source/ModelViewViewModel/Private/View/MVVMView.cpp b/Engine/Plugins/Runtime/ModelViewViewModel/Source/ModelViewViewModel/Private/View/MVVMView.cpp
--- a/Engine/Plugins/Runtime/ModelViewViewModel/Source/ModelViewViewModel/Private/View/MVVMView.cpp
+++ b/Engine/Plugins/Runtime/ModelViewViewModel/Source/ModelViewViewModel/Private/View/MVVMView.cpp
@@ -43,6 +43,18 @@ void UMVVMView::Construct()
 	check(ClassExtension);
 	check(bConstructed == false);
 
+	if (ClassExtension->DoesInitializeAtConstruction())
+	{
+		InitializeBindings();
+	}
+
+	bConstructed = true;
+}
+
+void UMVVMView::InitializeBindings()
+{
+	check(bInitialized == false);
+
 	// Init ViewModel instances
 	UUserWidget* UserWidget = GetUserWidget();
 	const TArrayView<const FMVVMViewClass_SourceCreator> AllViewModelCreators = ClassExtension->GetViewModelCreators();
@@ -91,7 +103,7 @@ void UMVVMView::Construct()
 		}
 	}
 
-	bConstructed = true;
+	bInitialized = true;
 
 	if (bHasEveryTickBinding)
 	{
@@ -109,12 +121,23 @@ void UMVVMView::Destruct()
 	UE::MVVM::FDebugging::BroadcastViewBeginDestruction(this);
 #endif
 
+	if (bInitialized)
+	{
+		DeintializeBindings();
+	}
+}
+
+void UMVVMView::DeintializeBindings()
+{
+	check(bInitialized == true);
+	bInitialized = false;
+
 	UUserWidget* UserWidget = GetUserWidget();
 	check(UserWidget);
 
 	for (FMVVMViewSource& Source : Sources)
 	{
-		if (Source.RegisteredCout > 0 && Source.Source)
+		if (Source.RegisteredCount > 0 && Source.Source)
 		{
 			TScriptInterface<INotifyFieldValueChanged> SourceAsInterface = Source.Source;
 			checkf(SourceAsInterface.GetInterface(), TEXT("It was added as a INotifyFieldValueChanged. It should still be."));
@@ -137,7 +160,7 @@ void UMVVMView::Destruct()
 		}
 	}
 
-	EnabledLibraryBindings.Reset();
+	RegisteredLibraryBindings.Reset();
 
 	if (bHasEveryTickBinding)
 	{
@@ -156,111 +179,147 @@ TScriptInterface<INotifyFieldValueChanged> UMVVMView::GetViewModel(FName ViewMod
 
 bool UMVVMView::SetViewModel(FName ViewModelName, TScriptInterface<INotifyFieldValueChanged> NewValue)
 {
-	if (!ViewModelName.IsNone() && ClassExtension != nullptr)
+	return SetSourceInternal(ViewModelName, NewValue, false);
+}
+
+
+bool UMVVMView::EvaluateSourceCreator(int32 SourceCreatorIndex)
+{
+	if (!ClassExtension->GetViewModelCreators().IsValidIndex(SourceCreatorIndex))
 	{
-		const int32 AllCreatedSourcesIndex = ClassExtension->GetViewModelCreators().IndexOfByPredicate([ViewModelName](const FMVVMViewClass_SourceCreator& Item)
-			{
-				return Item.GetSourceName() == ViewModelName;
-			});
-		{
-			if (AllCreatedSourcesIndex == INDEX_NONE)
-			{
-				UE::MVVM::FMessageLog Log(GetUserWidget());
-				Log.Error(LOCTEXT("SetViewModelViewModelNameNotFound", "The viewmodel name could not be found."));
-				return false;
-			}
-			else if (NewValue.GetObject()
-				&& !NewValue.GetObject()->GetClass()->IsChildOf(ClassExtension->GetViewModelCreators()[AllCreatedSourcesIndex].GetSourceClass()))
-			{
-				UE::MVVM::FMessageLog Log(GetUserWidget());
-				Log.Error(LOCTEXT("SetViewModelInvalidValueType", "The new viewmodel is not of the expected type."));
-				UE_LOG(LogMVVM, Error, TEXT("The viewmodel name '%s' is invalid for the view '%s'"), *ViewModelName.ToString(), *GetFullName());
-				return false;
-			}
-			else if (!Sources.IsValidIndex(AllCreatedSourcesIndex) || Sources[AllCreatedSourcesIndex].SourceName != ViewModelName)
-			{
-				UE::MVVM::FMessageLog Log(GetUserWidget());
-				Log.Error(LOCTEXT("SetViewModelInvalidCreatedSource", "An internal error occurs. The new viewmodel is not found."));
-				UE_LOG(LogMVVM, Error, TEXT("The viewmodel name '%s' is invalid for the view '%s'"), *ViewModelName.ToString(), *GetFullName());
-				return false;
-			}
-		}
+		UE::MVVM::FMessageLog Log(GetUserWidget());
+		Log.Error(FText::Format(LOCTEXT("ReevaluateDynamicSourceViewModelNameNotFound", "Internal error. The dynamic source {0} cound not be found."), FText::AsNumber(SourceCreatorIndex)));
+		return false;
+	}
+
+	UUserWidget* UserWidget = GetUserWidget();
+	const FMVVMViewClass_SourceCreator& SourceCreator = ClassExtension->GetViewModelCreators()[SourceCreatorIndex];
+	UObject* NewSource = SourceCreator.CreateInstance(ClassExtension, this, UserWidget);
+	ensureMsgf((NewSource == nullptr || NewSource->GetClass()->ImplementsInterface(UNotifyFieldValueChanged::StaticClass())), TEXT("The source has implement the interface. It should be check at compile time."));
+	return SetSourceInternal(SourceCreator.GetSourceName(), NewSource, true);
+}
+
+
+bool UMVVMView::SetSourceInternal(FName ViewModelName, TScriptInterface<INotifyFieldValueChanged> NewValue, bool bForDynamicSource)
+{
+	if (ViewModelName.IsNone())
+	{
+		UE::MVVM::FMessageLog Log(GetUserWidget());
+		Log.Error(LOCTEXT("SetViewModelInvalidName", "The viewmodel name is empty."));
+		return false;
+	}
 
-		FMVVMViewSource& ViewSource = Sources[AllCreatedSourcesIndex];
-		UObject* PreviousValue = Sources[AllCreatedSourcesIndex].Source;
-		if (PreviousValue != NewValue.GetObject())
+	if (ClassExtension == nullptr)
+	{
+		UE::MVVM::FMessageLog Log(GetUserWidget());
+		Log.Error(LOCTEXT("SetViewModelInvalidClass", "The view is not constructed."));
+		return false;
+	}
+
+	const int32 AllCreatedSourcesIndex = ClassExtension->GetViewModelCreators().IndexOfByPredicate([ViewModelName](const FMVVMViewClass_SourceCreator& Item)
 		{
-			const TArrayView<const FMVVMViewClass_CompiledBinding> CompiledBindings = ClassExtension->GetCompiledBindings();
+			return Item.GetSourceName() == ViewModelName;
+		});
+		
+	if (AllCreatedSourcesIndex == INDEX_NONE)
+	{
+		UE::MVVM::FMessageLog Log(GetUserWidget());
+		Log.Error(LOCTEXT("SetViewModelViewModelNameNotFound", "The viewmodel name could not be found."));
+		return false;
+	}
 
-			FMemMark Mark(FMemStack::Get());
-			TArray<int32, TMemStackAllocator<>> BindingToReenable;
-			BindingToReenable.Reserve(ViewSource.RegisteredCout);
+	const FMVVMViewClass_SourceCreator& SourceCreator = ClassExtension->GetViewModelCreators()[AllCreatedSourcesIndex];
 
-			// Unregister any bindings from that source
-			for (int32 Index = 0; Index < CompiledBindings.Num(); ++Index)
+	if (bForDynamicSource && !SourceCreator.CanBeEvaluated())
+	{
+		UE::MVVM::FMessageLog Log(GetUserWidget());
+		Log.Error(FText::Format(LOCTEXT("SetViewModelCannotBeEvalaute", "The new viewmodel {0} cannot be evaluated again."), FText::FromName(ViewModelName)));
+		return false;
+	}
+	else if (!bForDynamicSource && !SourceCreator.CanBeSet())
+	{
+		UE::MVVM::FMessageLog Log(GetUserWidget());
+		Log.Error(FText::Format(LOCTEXT("SetViewModelCannotBeSet", "The new viewmodel {0} cannot be set."), FText::FromName(ViewModelName)));
+		return false;
+	}
+			
+	if (NewValue.GetObject()
+		&& !NewValue.GetObject()->GetClass()->IsChildOf(SourceCreator.GetSourceClass()))
+	{
+		UE::MVVM::FMessageLog Log(GetUserWidget());
+		Log.Error(LOCTEXT("SetViewModelInvalidValueType", "The new viewmodel is not of the expected type."));
+		UE_LOG(LogMVVM, Error, TEXT("The viewmodel name '%s' is invalid for the view '%s'"), *ViewModelName.ToString(), *GetFullName());
+		return false;
+	}
+			
+	if (!Sources.IsValidIndex(AllCreatedSourcesIndex) || Sources[AllCreatedSourcesIndex].SourceName != ViewModelName)
+	{
+		UE::MVVM::FMessageLog Log(GetUserWidget());
+		Log.Error(LOCTEXT("SetViewModelInvalidCreatedSource", "An internal error occurs. The new viewmodel is not found."));
+		UE_LOG(LogMVVM, Error, TEXT("The viewmodel name '%s' is invalid for the view '%s'"), *ViewModelName.ToString(), *GetFullName());
+		return false;
+	}
+
+	FMVVMViewSource& ViewSource = Sources[AllCreatedSourcesIndex];
+	UObject* PreviousValue = Sources[AllCreatedSourcesIndex].Source;
+	if (PreviousValue != NewValue.GetObject())
+	{
+		const TArrayView<const FMVVMViewClass_CompiledBinding> CompiledBindings = ClassExtension->GetCompiledBindings();
+
+		// Unregister any bindings from that source
+		for (int32 Index = 0; Index < CompiledBindings.Num(); ++Index)
+		{
+			const FMVVMViewClass_CompiledBinding& Binding = CompiledBindings[Index];
+			if (IsLibraryBindingEnabled(Index) && Binding.GetSourceName() == ViewModelName)
 			{
-				const FMVVMViewClass_CompiledBinding& Binding = CompiledBindings[Index];
-				if (Binding.GetSourceName() == ViewModelName && IsLibraryBindingEnabled(Index))
-				{
-					DisableLibraryBinding(Binding, Index);
-					BindingToReenable.Add(Index);
-				}
+				DisableLibraryBinding(Binding, Index);
 			}
+		}
 
-			ViewSource.Source = NewValue.GetObject();
-			ViewSource.bSetManually = true;
-			if (ClassExtension->GetViewModelCreators()[AllCreatedSourcesIndex].IsSourceAUserWidgetProperty())
+		ViewSource.Source = NewValue.GetObject();
+		ViewSource.bSetManually = bForDynamicSource;
+		if (SourceCreator.IsSourceAUserWidgetProperty())
+		{
+			FObjectPropertyBase* FoundObjectProperty = FindFProperty<FObjectPropertyBase>(GetUserWidget()->GetClass(), ViewModelName);
+			if (ensure(FoundObjectProperty))
 			{
-				FObjectPropertyBase* FoundObjectProperty = FindFProperty<FObjectPropertyBase>(GetUserWidget()->GetClass(), ViewModelName);
-				if (ensure(FoundObjectProperty))
-				{
-					FoundObjectProperty->SetObjectPropertyValue_InContainer(GetUserWidget(), NewValue.GetObject());
-					ViewSource.bAssignedToUserWidgetProperty = true;
-				}
+				FoundObjectProperty->SetObjectPropertyValue_InContainer(GetUserWidget(), NewValue.GetObject());
+				ViewSource.bAssignedToUserWidgetProperty = true;
 			}
+		}
 
-			bool bPreviousEveryTickBinding = bHasEveryTickBinding;
-			bHasEveryTickBinding = false;
-			if (NewValue.GetObject())
+		bool bPreviousEveryTickBinding = bHasEveryTickBinding;
+		bHasEveryTickBinding = false;
+		if (NewValue.GetObject())
+		{
+			// Register back any binding that was previously enabled
+			if (bInitialized)
 			{
-				// Register back any binding that was previously enabled
-				if (BindingToReenable.Num() > 0)
+				// Enabled the default bindings
+				for (int32 Index = 0; Index < CompiledBindings.Num(); ++Index)
 				{
-					for (int32 Index : BindingToReenable)
+					const FMVVMViewClass_CompiledBinding& Binding = CompiledBindings[Index];
+					if (Binding.IsEnabledByDefault() && Binding.GetSourceName() == ViewModelName)
 					{
-						const FMVVMViewClass_CompiledBinding& Binding = CompiledBindings[Index];
 						EnableLibraryBinding(Binding, Index);
 					}
 				}
-				else if (bConstructed)
-				{
-					// Enabled the default bindings
-					for (int32 Index = 0; Index < CompiledBindings.Num(); ++Index)
-					{
-						const FMVVMViewClass_CompiledBinding& Binding = CompiledBindings[Index];
-						if (Binding.IsEnabledByDefault() && Binding.GetSourceName() == ViewModelName)
-						{
-							EnableLibraryBinding(Binding, Index);
-						}
-					}
-				}
 			}
+		}
 
-			if (bPreviousEveryTickBinding != bHasEveryTickBinding)
+		if (bPreviousEveryTickBinding != bHasEveryTickBinding)
+		{
+			if (bHasEveryTickBinding)
 			{
-				if (bHasEveryTickBinding)
-				{
-					GEngine->GetEngineSubsystem<UMVVMBindingSubsystem>()->AddViewWithEveryTickBinding(this);
-				}
-				else
-				{
-					GEngine->GetEngineSubsystem<UMVVMBindingSubsystem>()->RemoveViewWithEveryTickBinding(this);
-				}
+				GEngine->GetEngineSubsystem<UMVVMBindingSubsystem>()->AddViewWithEveryTickBinding(this);
+			}
+			else
+			{
+				GEngine->GetEngineSubsystem<UMVVMBindingSubsystem>()->RemoveViewWithEveryTickBinding(this);
 			}
 		}
-		return true;
 	}
-	return false;
+	return true;
 }
 
 
@@ -304,7 +363,7 @@ void UMVVMView::HandledLibraryBindingValueChanged(UObject* InViewModelOrWidget,
 			{
 				UE::MVVM::Private::RecursiveDetector.Emplace(this, InCompiledBindingIndex);
 
-				ExecuteLibraryBinding(Binding, InViewModelOrWidget);
+				ExecuteLibraryBinding(Binding, InCompiledBindingIndex);
 
 				UE::MVVM::Private::RecursiveDetector.Pop();
 			}
@@ -345,7 +404,7 @@ void UMVVMView::ExecuteDelayedBinding(const FMVVMViewDelayedBinding& DelayedBind
 				{
 					UE::MVVM::Private::RecursiveDetector.Emplace(this, CompiledBindingIndex);
 
-					ExecuteLibraryBinding(Binding);
+					ExecuteLibraryBinding(Binding, DelayedBinding.GetCompiledBindingIndex());
 
 					UE::MVVM::Private::RecursiveDetector.Pop();
 				}
@@ -368,85 +427,75 @@ void UMVVMView::ExecuteEveryTickBindings() const
 			const FMVVMViewClass_CompiledBinding& Binding = CompiledBindings[Index];
 			if (Binding.GetExecuteMode() == EMVVMExecutionMode::Tick && IsLibraryBindingEnabled(Index))
 			{
-				ExecuteLibraryBinding(Binding);
+				ExecuteLibraryBinding(Binding, Index);
 			}
 		}
 	}
 }
 
 
-void UMVVMView::ExecuteLibraryBinding(const FMVVMViewClass_CompiledBinding& Binding) const
+void UMVVMView::ExecuteLibraryBinding(const FMVVMViewClass_CompiledBinding& Binding, int32 BindingIndex) const
 {
 	check(ClassExtension);
 	check(GetUserWidget());
 
-	FMVVMCompiledBindingLibrary::EConversionFunctionType FunctionType = Binding.IsConversionFunctionComplex() ? FMVVMCompiledBindingLibrary::EConversionFunctionType::Complex : FMVVMCompiledBindingLibrary::EConversionFunctionType::Simple;
-	TValueOrError<void, FMVVMCompiledBindingLibrary::EExecutionFailingReason> ExecutionResult = ClassExtension->GetBindingLibrary().Execute(GetUserWidget(), Binding.GetBinding(), FunctionType);
+	bool bError = false;
+	if (!Binding.IsEvaluateSourceCreatorBinding())
+	{
+		FMVVMCompiledBindingLibrary::EConversionFunctionType FunctionType = Binding.IsConversionFunctionComplex() ? FMVVMCompiledBindingLibrary::EConversionFunctionType::Complex : FMVVMCompiledBindingLibrary::EConversionFunctionType::Simple;
+		TValueOrError<void, FMVVMCompiledBindingLibrary::EExecutionFailingReason> ExecutionResult = ClassExtension->GetBindingLibrary().Execute(GetUserWidget(), Binding.GetBinding(), FunctionType);
 
 #if UE_WITH_MVVM_DEBUGGING
-	if (ExecutionResult.HasError())
-	{
-		UE::MVVM::FDebugging::BroadcastLibraryBindingExecuted(this, Binding, ExecutionResult.GetError());
+		if (ExecutionResult.HasError())
+		{
+			UE::MVVM::FDebugging::BroadcastLibraryBindingExecuted(this, Binding, ExecutionResult.GetError());
+		}
+#endif
+
+		if (ExecutionResult.HasError())
+		{
+			UE::MVVM::FMessageLog Log(GetUserWidget());
+			Log.Error(FText::Format(LOCTEXT("ExecuteBindingFailGeneric", "The binding '{0}' was not executed. {1}.")
+#if UE_WITH_MVVM_DEBUGGING
+				, FText::FromString(Binding.ToString(ClassExtension->GetBindingLibrary(), FMVVMViewClass_CompiledBinding::FToStringArgs::Short()))
+#else
+				, FText::AsNumber(BindingIndex)
+#endif
+				, FMVVMCompiledBindingLibrary::LexToText(ExecutionResult.GetError())
+			));
+			bError = true;
+		}
 	}
 	else
 	{
-		UE::MVVM::FDebugging::BroadcastLibraryBindingExecuted(this, Binding);
-	}
+		if (!const_cast<UMVVMView*>(this)->EvaluateSourceCreator(Binding.GetEvaluateSourceCreatorBindingIndex()))
+		{
+			UE::MVVM::FMessageLog Log(GetUserWidget());
+			Log.Error(FText::Format(LOCTEXT("ExecuteBindingFailEvaluate", "The evaluate source creator binding '{0}' was not executed.")
+#if UE_WITH_MVVM_DEBUGGING
+				, FText::FromString(Binding.ToString(ClassExtension->GetBindingLibrary(), FMVVMViewClass_CompiledBinding::FToStringArgs::Short()))
+#else
+				, FText::AsNumber(BindingIndex)
 #endif
-
-	if (ExecutionResult.HasError())
-	{
-		UE::MVVM::FMessageLog Log(GetUserWidget());
-		Log.Error(FText::Format(LOCTEXT("ExecuteBindingGenericExecute", "The binding '{0}' was not executed. {1}.")
-			, FText::FromString(Binding.ToString(ClassExtension->GetBindingLibrary(), FMVVMViewClass_CompiledBinding::FToStringArgs::Short()))
-			, FMVVMCompiledBindingLibrary::LexToText(ExecutionResult.GetError())
 			));
+			bError = true;
+		}
 	}
-	else if (bLogBinding)
+
+	if (bLogBinding && !bError)
 	{
 		UE::MVVM::FMessageLog Log(GetUserWidget());
-		Log.Info(FText::Format(LOCTEXT("ExecuteLibraryBinding", "Execute binding '{0}'.")
+		Log.Info(FText::Format(LOCTEXT("ExecuteBindingGeneric", "Execute binding '{0}'.")
+#if UE_WITH_MVVM_DEBUGGING
 			, FText::FromString(Binding.ToString(ClassExtension->GetBindingLibrary(), FMVVMViewClass_CompiledBinding::FToStringArgs::All()))
-			));
-	}
-}
-
-
-void UMVVMView::ExecuteLibraryBinding(const FMVVMViewClass_CompiledBinding& Binding, UObject* Source) const
-{
-	check(ClassExtension);
-	check(GetUserWidget());
-	check(Source);
-
-	TValueOrError<void, FMVVMCompiledBindingLibrary::EExecutionFailingReason> ExecutionResult = Binding.IsConversionFunctionComplex()
-		? ClassExtension->GetBindingLibrary().Execute(GetUserWidget(), Binding.GetBinding(), FMVVMCompiledBindingLibrary::EConversionFunctionType::Complex)
-		: ClassExtension->GetBindingLibrary().ExecuteWithSource(GetUserWidget(), Binding.GetBinding(), Source);
+#else
+			, FText::AsNumber(BindingIndex)
+#endif
+		));
 
 #if UE_WITH_MVVM_DEBUGGING
-	if (ExecutionResult.HasError())
-	{
-		UE::MVVM::FDebugging::BroadcastLibraryBindingExecuted(this, Binding, ExecutionResult.GetError());
-	}
-	else
-	{
 		UE::MVVM::FDebugging::BroadcastLibraryBindingExecuted(this, Binding);
-	}
 #endif
-
-	if (ExecutionResult.HasError())
-	{
-		UE::MVVM::FMessageLog Log(GetUserWidget());
-		Log.Error(FText::Format(LOCTEXT("ExecuteBindingGenericExecute", "The binding '{0}' was not executed. {1}.")
-			, FText::FromString(Binding.ToString(ClassExtension->GetBindingLibrary(), FMVVMViewClass_CompiledBinding::FToStringArgs::Short()))
-			, FMVVMCompiledBindingLibrary::LexToText(ExecutionResult.GetError())
-			));
-	}
-	else if (bLogBinding)
-	{
-		UE::MVVM::FMessageLog Log(GetUserWidget());
-		Log.Info(FText::Format(LOCTEXT("ExecuteLibraryBinding", "Execute binding '{0}'.")
-			, FText::FromString(Binding.ToString(ClassExtension->GetBindingLibrary(), FMVVMViewClass_CompiledBinding::FToStringArgs::All()))
-			));
 	}
 }
 
@@ -489,35 +538,41 @@ void UMVVMView::ExecuteLibraryBinding(const FMVVMViewClass_CompiledBinding& Bind
 
 bool UMVVMView::IsLibraryBindingEnabled(int32 InBindindIndex) const
 {
-	return EnabledLibraryBindings.IsValidIndex(InBindindIndex) && EnabledLibraryBindings[InBindindIndex];
+	return RegisteredLibraryBindings.IsValidIndex(InBindindIndex) && RegisteredLibraryBindings[InBindindIndex];
 }
 
 
 void UMVVMView::EnableLibraryBinding(const FMVVMViewClass_CompiledBinding& Binding, int32 BindingIndex)
 {
-	EnabledLibraryBindings.PadToNum(BindingIndex + 1, false);
-	check(EnabledLibraryBindings[BindingIndex] == false);
+	RegisteredLibraryBindings.PadToNum(BindingIndex + 1, false);
+	check(RegisteredLibraryBindings[BindingIndex] == false);
 
 	EMVVMExecutionMode ExecutionMode = Binding.GetExecuteMode();
-	bool bRegistered = true;
+	bool bCanExecute = true;
+	bool bRegistered = false;
 	if (!Binding.IsOneTime() && ExecutionMode != EMVVMExecutionMode::Tick)
 	{
 		bRegistered = RegisterLibraryBinding(Binding, BindingIndex);
+		bCanExecute = bRegistered;
 	}
 
 	if (bRegistered && bLogBinding)
 	{
 		UE::MVVM::FMessageLog Log(GetUserWidget());
 		Log.Info(FText::Format(LOCTEXT("EnableLibraryBinding", "Enable binding '{0}'.")
+#if UE_WITH_MVVM_DEBUGGING
 			, FText::FromString(Binding.ToString(ClassExtension->GetBindingLibrary(), FMVVMViewClass_CompiledBinding::FToStringArgs::All()))
+#else
+			, FText::AsNumber(BindingIndex)
+#endif
 		));
 	}
 
-	EnabledLibraryBindings[BindingIndex] = bRegistered;
+	RegisteredLibraryBindings[BindingIndex] = bRegistered;
 
-	if (bRegistered && Binding.NeedsExecutionAtInitialization())
+	if (bCanExecute && Binding.NeedsExecutionAtInitialization())
 	{
-		ExecuteLibraryBinding(Binding);
+		ExecuteLibraryBinding(Binding, BindingIndex);
 	}
 
 	bHasEveryTickBinding = bHasEveryTickBinding || ExecutionMode == EMVVMExecutionMode::Tick;
@@ -530,17 +585,21 @@ void UMVVMView::DisableLibraryBinding(const FMVVMViewClass_CompiledBinding& Bind
 
 
 	EMVVMExecutionMode ExecutionMode = Binding.GetExecuteMode();
-	EnabledLibraryBindings[BindingIndex] = false;
+	RegisteredLibraryBindings[BindingIndex] = false;
 	if (!Binding.IsOneTime() && ExecutionMode != EMVVMExecutionMode::Tick)
 	{
-		UnregisterLibraryBinding(Binding);
+		UnregisterLibraryBinding(Binding, BindingIndex);
 	}
 
 	if (bLogBinding)
 	{
 		UE::MVVM::FMessageLog Log(GetUserWidget());
 		Log.Info(FText::Format(LOCTEXT("DisableLibraryBinding", "Disable binding '{0}'.")
+#if UE_WITH_MVVM_DEBUGGING
 			, FText::FromString(Binding.ToString(ClassExtension->GetBindingLibrary(), FMVVMViewClass_CompiledBinding::FToStringArgs::All()))
+#else
+			, FText::AsNumber(BindingIndex)
+#endif
 			));
 	}
 }
@@ -550,12 +609,16 @@ bool UMVVMView::RegisterLibraryBinding(const FMVVMViewClass_CompiledBinding& Bin
 {
 	check(ClassExtension);
 
-	auto LogMessage = [this, &Binding](const FText& Message)
+	auto LogMessage = [this, &Binding, BindingIndex](const FText& Message)
 	{
 		UE::MVVM::FMessageLog Log(GetUserWidget());
-		Log.Error(FText::Format(LOCTEXT("RegisterBindingFailed_Format", "Widget '{0}' can't register binding '{1}'. {2}.")
+		Log.Error(FText::Format(LOCTEXT("RegisterBindingFailed_Format", "Widget '{0}' can't register binding '{1}'. {2}")
 			, FText::FromString(GetFullName())
+#if UE_WITH_MVVM_DEBUGGING
 			, FText::FromString(Binding.ToString(ClassExtension->GetBindingLibrary(), FMVVMViewClass_CompiledBinding::FToStringArgs::Short()))
+#else
+			, FText::AsNumber(BindingIndex)
+#endif
 			, Message
 			));
 	};
@@ -583,7 +646,7 @@ bool UMVVMView::RegisterLibraryBinding(const FMVVMViewClass_CompiledBinding& Bin
 	}
 
 	// The source may not have been created because the property path was wrong.
-	TScriptInterface<INotifyFieldValueChanged> Source = FindSource(Binding, true);
+	TScriptInterface<INotifyFieldValueChanged> Source = FindSource(Binding, BindingIndex, true);
 	if (Source.GetInterface() == nullptr)
 	{
 		if (!Binding.IsRegistrationOptional())
@@ -597,33 +660,12 @@ bool UMVVMView::RegisterLibraryBinding(const FMVVMViewClass_CompiledBinding& Bin
 		return false;
 	}
 
-	// Only bind if the source and the destination are valid.
-	if (ClassExtension->GetBindingLibrary().EvaluateFieldPath(GetUserWidget(), Binding.GetBinding().GetSourceFieldPath()).HasError())
-	{
-		LogMessage(LOCTEXT("RegisterBindingFailed_InvalidSource", "The source could not be evaluated."));
-
-#if UE_WITH_MVVM_DEBUGGING
-		UE::MVVM::FDebugging::BroadcastLibraryBindingRegistered(this, Binding, UE::MVVM::FDebugging::ERegisterLibraryBindingResult::Failed_InvalidSourceField);
-#endif
-		return false;
-	}
-	if (ClassExtension->GetBindingLibrary().EvaluateFieldPath(GetUserWidget(), Binding.GetBinding().GetDestinationFieldPath()).HasError())
-	{
-		LogMessage(LOCTEXT("RegisterBindingFailed_InvalidDestination", "The destination could not be evaluated."));
-
-#if UE_WITH_MVVM_DEBUGGING
-		UE::MVVM::FDebugging::BroadcastLibraryBindingRegistered(this, Binding, UE::MVVM::FDebugging::ERegisterLibraryBindingResult::Failed_InvalidDestinationField);
-#endif
-		return false;
-	}
-
-
 	UE::FieldNotification::FFieldMulticastDelegate::FDelegate Delegate = UE::FieldNotification::FFieldMulticastDelegate::FDelegate::CreateUObject(this, &UMVVMView::HandledLibraryBindingValueChanged, BindingIndex);
 	if (Source->AddFieldValueChangedDelegate(FieldId, MoveTemp(Delegate)).IsValid())
 	{
 		if (FMVVMViewSource* RegisteredSource = FindViewSource(Binding.GetSourceName()))
 		{
-			++RegisteredSource->RegisteredCout;
+			++RegisteredSource->RegisteredCount;
 		}
 		else
 		{
@@ -631,7 +673,7 @@ bool UMVVMView::RegisterLibraryBinding(const FMVVMViewClass_CompiledBinding& Bin
 			FMVVMViewSource& NewSource = Sources.AddDefaulted_GetRef();
 			NewSource.Source = Source.GetObject();
 			NewSource.SourceName = Binding.GetSourceName();
-			NewSource.RegisteredCout = 1;
+			NewSource.RegisteredCount = 1;
 		}
 
 #if UE_WITH_MVVM_DEBUGGING
@@ -649,9 +691,9 @@ bool UMVVMView::RegisterLibraryBinding(const FMVVMViewClass_CompiledBinding& Bin
 }
 
 
-void UMVVMView::UnregisterLibraryBinding(const FMVVMViewClass_CompiledBinding& Binding)
+void UMVVMView::UnregisterLibraryBinding(const FMVVMViewClass_CompiledBinding& Binding, int32 BindingIndex)
 {
-	TScriptInterface<INotifyFieldValueChanged> Source = FindSource(Binding, true);
+	TScriptInterface<INotifyFieldValueChanged> Source = FindSource(Binding, BindingIndex, true);
 	if (Source.GetInterface() && ClassExtension)
 	{
 		TValueOrError<UE::FieldNotification::FFieldId, void> FieldIdResult = ClassExtension->GetBindingLibrary().GetFieldId(Binding.GetSourceFieldId());
@@ -663,14 +705,14 @@ void UMVVMView::UnregisterLibraryBinding(const FMVVMViewClass_CompiledBinding& B
 		FMVVMViewSource* RegisteredSource = FindViewSource(Binding.GetSourceName());
 		if (ensureMsgf(RegisteredSource, TEXT("If the binding was enabled, then the source should also be there.")))
 		{
-			ensureMsgf(RegisteredSource->RegisteredCout > 0, TEXT("The count should match the number of RegisterLibraryBinding and UnregisterLibraryBinding"));
-			--RegisteredSource->RegisteredCout;
+			ensureMsgf(RegisteredSource->RegisteredCount > 0, TEXT("The count should match the number of RegisterLibraryBinding and UnregisterLibraryBinding"));
+			--RegisteredSource->RegisteredCount;
 		}
 	}
 }
 
 
-TScriptInterface<INotifyFieldValueChanged> UMVVMView::FindSource(const FMVVMViewClass_CompiledBinding& Binding, bool bAllowNull) const
+TScriptInterface<INotifyFieldValueChanged> UMVVMView::FindSource(const FMVVMViewClass_CompiledBinding& Binding, int32 BindingIndex, bool bAllowNull) const
 {
 	UUserWidget* UserWidget = GetUserWidget();
 	check(UserWidget);
@@ -682,12 +724,16 @@ TScriptInterface<INotifyFieldValueChanged> UMVVMView::FindSource(const FMVVMView
 	}
 	else
 	{
-		auto LogMessage = [this, &Binding](const FText& Message)
+		auto LogMessage = [this, &Binding, BindingIndex](const FText& Message)
 		{
 			UE::MVVM::FMessageLog Log(GetUserWidget());
 			Log.Error(FText::Format(LOCTEXT("FindSourceFailed_Format", "Widget '{0}' can't evaluate the source for binding '{1}'. {2}.")
 				, FText::FromString(GetFullName())
+#if UE_WITH_MVVM_DEBUGGING
 				, FText::FromString(Binding.ToString(ClassExtension->GetBindingLibrary(), FMVVMViewClass_CompiledBinding::FToStringArgs::Short()))
+#else
+				, FText::AsNumber(BindingIndex)
+#endif
 				, Message
 			));
 		};
diff --git a/Engine/Plugins/Runtime/ModelViewViewModel/Source/ModelViewViewModel/Private/View/MVVMViewClass.cpp b/Engine/Plugins/Runtime/ModelViewViewModel/Source/ModelViewViewModel/Private/View/MVVMViewClass.cpp
--- a/Engine/Plugins/Runtime/ModelViewViewModel/Source/ModelViewViewModel/Private/View/MVVMViewClass.cpp
+++ b/Engine/Plugins/Runtime/ModelViewViewModel/Source/ModelViewViewModel/Private/View/MVVMViewClass.cpp
@@ -22,76 +22,6 @@
 ///////////////////////////////////////////////////////////////////////
 // 
 ///////////////////////////////////////////////////////////////////////
-
-FMVVMViewClass_SourceCreator FMVVMViewClass_SourceCreator::MakeManual(FName InName, UClass* InNotifyFieldValueChangedClass)
-{
-	FMVVMViewClass_SourceCreator Result;
-	if (ensure(InNotifyFieldValueChangedClass && InNotifyFieldValueChangedClass->ImplementsInterface(UNotifyFieldValueChanged::StaticClass())))
-	{
-		Result.PropertyName = InName;
-		Result.ExpectedSourceType = InNotifyFieldValueChangedClass;
-		Result.bCreateInstance = false;
-		Result.bOptional = true;
-	}
-	return Result;
-}
-
-FMVVMViewClass_SourceCreator FMVVMViewClass_SourceCreator::MakeInstance(FName InName, UClass* InNotifyFieldValueChangedClass)
-{
-	FMVVMViewClass_SourceCreator Result;
-	if (ensure(InNotifyFieldValueChangedClass && InNotifyFieldValueChangedClass->ImplementsInterface(UNotifyFieldValueChanged::StaticClass())))
-	{
-		Result.PropertyName = InName;
-		Result.ExpectedSourceType = InNotifyFieldValueChangedClass;
-		Result.bCreateInstance = true;
-		Result.bOptional = false;
-	}
-	return Result;
-}
-
-FMVVMViewClass_SourceCreator FMVVMViewClass_SourceCreator::MakeFieldPath(FName InName, UClass* InNotifyFieldValueChangedClass, FMVVMVCompiledFieldPath InFieldPath, bool bOptional)
-{
-	FMVVMViewClass_SourceCreator Result;
-	if (ensure(InNotifyFieldValueChangedClass&& InNotifyFieldValueChangedClass->ImplementsInterface(UNotifyFieldValueChanged::StaticClass())))
-	{
-		if (ensure(InFieldPath.IsValid()))
-		{
-			Result.PropertyName = InName;
-			Result.ExpectedSourceType = InNotifyFieldValueChangedClass;
-			Result.FieldPath = InFieldPath;
-			Result.bOptional = bOptional;
-		}
-	}
-	return Result;
-}
-
-FMVVMViewClass_SourceCreator FMVVMViewClass_SourceCreator::MakeGlobalContext(FName InName, FMVVMViewModelContext InContext, bool bOptional)
-{
-	FMVVMViewClass_SourceCreator Result;
-	if (ensure(InContext.ContextClass && InContext.ContextClass->ImplementsInterface(UNotifyFieldValueChanged::StaticClass())))
-	{
-		Result.PropertyName = InName;
-		Result.ExpectedSourceType = InContext.ContextClass;
-		Result.GlobalViewModelInstance = MoveTemp(InContext);
-		Result.bOptional = bOptional;
-	}
-	return Result;
-}
-
-FMVVMViewClass_SourceCreator FMVVMViewClass_SourceCreator::MakeResolver(FName InName, UClass* InNotifyFieldValueChangedClass, UMVVMViewModelContextResolver* InResolver, bool bOptional)
-{
-	FMVVMViewClass_SourceCreator Result;
-	if (ensure(InResolver && InNotifyFieldValueChangedClass && InNotifyFieldValueChangedClass->ImplementsInterface(UNotifyFieldValueChanged::StaticClass())))
-	{
-		Result.PropertyName = InName;
-		Result.ExpectedSourceType = InNotifyFieldValueChangedClass;
-		Result.Resolver = InResolver;
-		Result.bOptional = bOptional;
-		ensure(InResolver->GetPackage() != GetTransientPackage() && !InResolver->HasAnyFlags(RF_Transient));
-	}
-	return Result;
-}
-
 UObject* FMVVMViewClass_SourceCreator::CreateInstance(const UMVVMViewClass* InViewClass, UMVVMView* InView, UUserWidget* InUserWidget) const
 {
 	check(InViewClass);
@@ -99,8 +29,9 @@ UObject* FMVVMViewClass_SourceCreator::CreateInstance(const UMVVMViewClass* InVi
 	check(InUserWidget);
 
 	UObject* Result = nullptr;
+	const bool bOptional = (Flags & (uint8)ESourceFlags::IsOptional) != 0;
 
-	if (bCreateInstance)
+	if ((Flags & (uint8)ESourceFlags::TypeCreateInstance) != 0)
 	{
 		if (ExpectedSourceType.Get() != nullptr)
 		{
@@ -184,6 +115,134 @@ UObject* FMVVMViewClass_SourceCreator::CreateInstance(const UMVVMViewClass* InVi
 	return Result;
 }
 
+#if UE_WITH_MVVM_DEBUGGING
+FMVVMViewClass_SourceCreator::FToStringArgs FMVVMViewClass_SourceCreator::FToStringArgs::Short()
+{
+	FToStringArgs Result;
+	Result.bUseDisplayName = false;
+	Result.bAddCreationMode = false;
+	Result.bAddFlags = false;
+	return Result;
+}
+
+FMVVMViewClass_SourceCreator::FToStringArgs FMVVMViewClass_SourceCreator::FToStringArgs::All()
+{
+	return FToStringArgs();
+}
+
+FString FMVVMViewClass_SourceCreator::ToString(const FMVVMCompiledBindingLibrary& BindingLibrary, FToStringArgs Args) const
+{
+	TStringBuilder<512> StringBuilder;
+	StringBuilder << TEXT("Type: ");
+#if WITH_EDITOR
+	if (Args.bUseDisplayName)
+	{
+		StringBuilder << ExpectedSourceType->GetDisplayNameText().ToString();
+	}
+	else
+#endif
+	{
+		StringBuilder << ExpectedSourceType->GetPathName();
+	}
+	
+	StringBuilder << TEXT(", SourceName: ");
+	StringBuilder << GetSourceName();
+
+	if (Args.bAddFlags)
+	{
+		bool bHasFlag = false;
+		auto AddPipe = [&StringBuilder, &bHasFlag]()
+		{
+			if (bHasFlag)
+			{
+				StringBuilder << TEXT('|');
+			}
+			bHasFlag = true;
+		};
+
+		StringBuilder << TEXT(", Flags: ");
+		if ((Flags & (uint8)ESourceFlags::IsOptional) != 0)
+		{
+			AddPipe();
+			StringBuilder << TEXT("Optional");
+		}
+		if (CanBeSet())
+		{
+			AddPipe();
+			StringBuilder << TEXT("CanBeSet");
+		}
+		if (CanBeEvaluated())
+		{
+			AddPipe();
+			StringBuilder << TEXT("CanBeEvaluated");
+		}
+		if (IsSourceAUserWidgetProperty())
+		{
+			AddPipe();
+			StringBuilder << TEXT("IsUserWidgetProperty");
+		}
+	}
+
+	if (Args.bAddCreationMode)
+	{
+		StringBuilder << TEXT("\n    CreationType: ");
+		if ((Flags & (uint8)ESourceFlags::TypeCreateInstance) != 0)
+		{
+			StringBuilder << TEXT("CreateInsance");
+		}
+		else if (Resolver)
+		{
+			StringBuilder << Resolver->GetClass()->GetFullName();
+		}
+		else if (GlobalViewModelInstance.IsValid())
+		{
+			StringBuilder << TEXT("GlobalCollection={");
+			if (GlobalViewModelInstance.ContextClass)
+			{
+				StringBuilder << GlobalViewModelInstance.ContextClass->GetPathName();
+			}
+			else
+			{
+				StringBuilder << TEXT("<Invalid>");
+			}
+			StringBuilder << TEXT(", ");
+			StringBuilder << GlobalViewModelInstance.ContextName;
+			StringBuilder << TEXT("}");
+		}
+		else if (FieldPath.IsValid())
+		{
+			StringBuilder << TEXT("Path=");
+			TValueOrError<FString, FString> FieldPathString = BindingLibrary.FieldPathToString(FieldPath, Args.bUseDisplayName);
+			if (FieldPathString.HasValue())
+			{
+				StringBuilder << FieldPathString.GetValue();
+			}
+			else
+			{
+				StringBuilder << TEXT("<Error>");
+				StringBuilder << FieldPathString.GetError();
+			}
+		}
+	}
+
+	return StringBuilder.ToString();
+}
+#endif
+
+#if WITH_EDITOR
+void FMVVMViewClass_SourceCreator::PostSerialize(const FArchive& Ar)
+{
+	if (Ar.IsLoading())
+	{
+		if (Flags == 0)
+		{
+			Flags |= bCreateInstance_DEPRECATED ? (uint8)ESourceFlags::TypeCreateInstance : 0;
+			Flags |= bIsUserWidgetProperty_DEPRECATED ? (uint8)ESourceFlags::IsUserWidgetProperty : 0;
+			Flags |= bOptional_DEPRECATED ? (uint8)ESourceFlags::IsOptional : 0;
+		}
+	}
+}
+#endif
 
 ///////////////////////////////////////////////////////////////////////
 // 
@@ -204,6 +263,7 @@ EMVVMExecutionMode FMVVMViewClass_CompiledBinding::GetExecuteMode() const
 	return (Flags & EBindingFlags::OverrideExecuteMode) == 0 ? DefaultMode : ExecutionMode;
 }
 
+#if UE_WITH_MVVM_DEBUGGING
 FMVVMViewClass_CompiledBinding::FToStringArgs FMVVMViewClass_CompiledBinding::FToStringArgs::Short()
 {
 	FToStringArgs Result;
@@ -242,71 +302,103 @@ FString FMVVMViewClass_CompiledBinding::ToString(const FMVVMCompiledBindingLibra
 
 	if (Args.bAddFieldPath)
 	{
-		StringBuilder << TEXT("Binding: ");
-		TValueOrError<FString, FString> DestinationString = BindingLibrary.FieldPathToString(GetBinding().GetDestinationFieldPath(), Args.bUseDisplayName);
-		TValueOrError<FString, FString> SourceString = BindingLibrary.FieldPathToString(GetBinding().GetSourceFieldPath(), Args.bUseDisplayName);
-		TValueOrError<FString, FString> ConversionString = MakeValue(FString());
-		if (GetBinding().GetConversionFunctionFieldPath().IsValid())
+		if (GetBinding().IsValid())
 		{
-			ConversionString = BindingLibrary.FieldPathToString(GetBinding().GetConversionFunctionFieldPath(), Args.bUseDisplayName);
-		}
+			StringBuilder << TEXT("Binding: ");
 
-		if (DestinationString.HasError() || SourceString.HasError() || ConversionString.HasError())
-		{
-			StringBuilder << TEXT("Error: ");
-		}
+			FString SourceString;
+			FString DestinationString;
+			FString ConversionString;
+			bool bErrorStrings = false;
+			if (!IsConversionFunctionComplex())
+			{
+				TValueOrError<FString, FString> SourceFieldPathString = BindingLibrary.FieldPathToString(GetBinding().GetSourceFieldPath(), Args.bUseDisplayName);
+				SourceString = SourceFieldPathString.HasValue() ? SourceFieldPathString.StealValue() : SourceFieldPathString.StealError();
+				bErrorStrings = bErrorStrings || SourceFieldPathString.HasError();
+			}
+			{
+				TValueOrError<FString, FString> DestinationFieldPathString = BindingLibrary.FieldPathToString(GetBinding().GetDestinationFieldPath(), Args.bUseDisplayName);
+				DestinationString = DestinationFieldPathString.HasValue() ? DestinationFieldPathString.StealValue() : DestinationFieldPathString.StealError();
+				bErrorStrings = bErrorStrings || DestinationFieldPathString.HasError();
+			}
+			if (GetBinding().GetConversionFunctionFieldPath().IsValid())
+			{
+				TValueOrError<FString, FString> ConversionFieldPathString = BindingLibrary.FieldPathToString(GetBinding().GetConversionFunctionFieldPath(), Args.bUseDisplayName);
+				ConversionString = ConversionFieldPathString.HasValue() ? ConversionFieldPathString.StealValue() : ConversionFieldPathString.StealError();
+				bErrorStrings = bErrorStrings || ConversionFieldPathString.HasError();
+			}
+
+			if (bErrorStrings)
+			{
+				StringBuilder << TEXT("Error: ");
+			}
 
-		if (Args.bUseDisplayName)
-		{
-			StringBuilder << TEXT('"');
-		}
-		StringBuilder << (DestinationString.HasValue() ? DestinationString.GetValue() : DestinationString.GetError());
-		if (Args.bUseDisplayName)
-		{
-			StringBuilder << TEXT('"');
-		}
-		StringBuilder << TEXT(" = ");
-		if (GetBinding().GetConversionFunctionFieldPath().IsValid())
-		{
 			if (Args.bUseDisplayName)
 			{
 				StringBuilder << TEXT('"');
 			}
-			StringBuilder << (ConversionString.HasValue() ? ConversionString.GetValue() : ConversionString.GetError());
+			StringBuilder << DestinationString;
 			if (Args.bUseDisplayName)
 			{
 				StringBuilder << TEXT('"');
 			}
-			StringBuilder << TEXT(" ( ");
-		}
+			StringBuilder << TEXT(" = ");
+			if (GetBinding().GetConversionFunctionFieldPath().IsValid())
+			{
+				if (Args.bUseDisplayName)
+				{
+					StringBuilder << TEXT('"');
+				}
+				StringBuilder << ConversionString;
+				if (Args.bUseDisplayName)
+				{
+					StringBuilder << TEXT('"');
+				}
+				StringBuilder << TEXT(" ( ");
+			}
 
-		if (Args.bUseDisplayName)
-		{
-			StringBuilder << TEXT('"');
+			if (!IsConversionFunctionComplex())
+			{
+				if (Args.bUseDisplayName)
+				{
+					StringBuilder << TEXT('"');
+				}
+				StringBuilder << SourceString;
+				if (Args.bUseDisplayName)
+				{
+					StringBuilder << TEXT('"');
+				}
+			}
+
+			if (GetBinding().GetConversionFunctionFieldPath().IsValid())
+			{
+				StringBuilder << TEXT(" )");
+			}
+
+			StringBuilder << TEXT("\n    ");
 		}
-		StringBuilder << (SourceString.HasValue() ? SourceString.GetValue() : SourceString.GetError());
-		if (Args.bUseDisplayName)
+		else if (Args.bAddFieldPath && IsEvaluateSourceCreatorBinding())
 		{
-			StringBuilder << TEXT('"');
+			StringBuilder << TEXT("Evaluate Source Creator: [");
+			StringBuilder << EvaluateSourceCreatorIndex;
+			StringBuilder << TEXT("]\n    ");
 		}
-
-		if (GetBinding().GetConversionFunctionFieldPath().IsValid())
+		else
 		{
-			StringBuilder << TEXT(" )");
+			StringBuilder << TEXT("!!Invalid Binding Error!!");
+			StringBuilder << TEXT("\n    ");
 		}
-
-		StringBuilder << TEXT("\n    ");
 	}
 
-	StringBuilder << TEXT("PropertyName: '");
-	StringBuilder << SourcePropertyName;
+	StringBuilder << TEXT("SourceName: ");
+	StringBuilder << GetSourceName();
 
 	if (Args.bAddFieldId)
 	{
 		if (GetSourceFieldId().IsValid())
 		{
 			TValueOrError<UE::FieldNotification::FFieldId, void> SourceFieldId = BindingLibrary.GetFieldId(GetSourceFieldId());
-			StringBuilder << TEXT("', FieldId: '");
+			StringBuilder << TEXT(", FieldId: ");
 			StringBuilder << (SourceFieldId.HasValue() ? SourceFieldId.GetValue().GetName() : FName());
 		}
 		else
@@ -317,24 +409,29 @@ FString FMVVMViewClass_CompiledBinding::ToString(const FMVVMCompiledBindingLibra
 
 	if (Args.bAddFlags)
 	{
-		StringBuilder << TEXT("`, Mode: '");
-		StringBuilder << StaticEnum<EMVVMExecutionMode>()->GetNameByValue((int64)ExecutionMode);
-		StringBuilder << TEXT("`, Flags: '");
+		if ((Flags & EBindingFlags::OverrideExecuteMode) != 0)
+		{
+			StringBuilder << TEXT(", Mode: ");
+			StringBuilder << StaticEnum<EMVVMExecutionMode>()->GetNameByValue((int64)ExecutionMode);
+		}
 
-		auto AddPipe = [&StringBuilder]()
+		StringBuilder << TEXT(", Flags: ");
+
+		bool bAddPipe = false;
+		auto AddPipe = [&StringBuilder, &bAddPipe]()
 		{
-			StringBuilder << TEXT('|');
+			if (bAddPipe)
+			{
+				StringBuilder << TEXT('|');
+			}
+			bAddPipe = true;
 		};
 
-		if ((Flags & EBindingFlags::ForwardBinding) != 0)
-		{
-			StringBuilder << TEXT("Forward");
-		}
-		else
+		if ((Flags & EBindingFlags::ExecuteAtInitialization) != 0)
 		{
-			StringBuilder << TEXT("Backward");
+			AddPipe();
+			StringBuilder << TEXT("ExecAtInit");
 		}
-
 		if ((Flags & EBindingFlags::OneTime) != 0)
 		{
 			AddPipe();
@@ -350,7 +447,7 @@ FString FMVVMViewClass_CompiledBinding::ToString(const FMVVMCompiledBindingLibra
 			AddPipe();
 			StringBuilder << TEXT("Optional");
 		}
-		if ((Flags & EBindingFlags::ConversionFunctionIsComplex) != 0)
+		if (IsConversionFunctionComplex())
 		{
 			AddPipe();
 			StringBuilder << TEXT("Complex");
@@ -365,11 +462,11 @@ FString FMVVMViewClass_CompiledBinding::ToString(const FMVVMCompiledBindingLibra
 			AddPipe();
 			StringBuilder << TEXT("Self");
 		}
-		StringBuilder << TEXT("'");
 	}
 
 	return StringBuilder.ToString();
 }
+#endif
 
 
 ///////////////////////////////////////////////////////////////////////
@@ -392,5 +489,46 @@ void UMVVMViewClass::Initialize(UUserWidget* UserWidget)
 	}
 }
 
+#if UE_WITH_MVVM_DEBUGGING
+void UMVVMViewClass::Log(FMVVMViewClass_SourceCreator::FToStringArgs SourceArgs, FMVVMViewClass_CompiledBinding::FToStringArgs BindingArgs) const
+{
+	bool bPreviousLoad = bLoaded;
+	if (!bLoaded)
+	{
+		const_cast<UMVVMViewClass*>(this)->BindingLibrary.Load();
+		const_cast<UMVVMViewClass*>(this)->bLoaded = true;
+	}
+
+	TStringBuilder<2048> Builder;
+	Builder << TEXT("Compiled Sources for: ");
+	Builder << GetOutermost()->GetFName();
+	for (const FMVVMViewClass_SourceCreator& Sources : GetViewModelCreators())
+	{
+		Builder << TEXT("\n");
+		Builder << Sources.ToString(GetBindingLibrary(), SourceArgs);
+	}
+	UE_LOG(LogMVVM, Log, TEXT("%s"), Builder.ToString());
+
+	Builder.Reset();
+
+	Builder << TEXT("Compiled Bindings for: ");
+	Builder << GetOutermost()->GetFName();
+	for (const FMVVMViewClass_CompiledBinding& Binding : GetCompiledBindings())
+	{
+		Builder << TEXT("\n");
+		Builder << Binding.ToString(GetBindingLibrary(), BindingArgs);
+	}
+	UE_LOG(LogMVVM, Log, TEXT("%s"), Builder.ToString());
+
+	if (!bPreviousLoad)
+	{
+#if WITH_EDITOR
+		const_cast<UMVVMViewClass*>(this)->BindingLibrary.Unload();
+		const_cast<UMVVMViewClass*>(this)->bLoaded = false;
+#endif
+	}
+}
+#endif
+
 #undef LOCTEXT_NAMESPACE
 
diff --git a/Engine/Plugins/Runtime/ModelViewViewModel/Source/ModelViewViewModel/Public/Bindings/MVVMCompiledBindingLibrary.h b/Engine/Plugins/Runtime/ModelViewViewModel/Source/ModelViewViewModel/Public/Bindings/MVVMCompiledBindingLibrary.h
--- a/Engine/Plugins/Runtime/ModelViewViewModel/Source/ModelViewViewModel/Public/Bindings/MVVMCompiledBindingLibrary.h
+++ b/Engine/Plugins/Runtime/ModelViewViewModel/Source/ModelViewViewModel/Public/Bindings/MVVMCompiledBindingLibrary.h
@@ -197,7 +197,10 @@ struct FMVVMVCompiledBinding
 public:
 	bool IsValid() const
 	{
-		return SourceFieldPath.IsValid() && DestinationFieldPath.IsValid();
+		const bool bSourceValid = SourceFieldPath.IsValid() || IsComplexFunction();
+		const bool bDestinationValid = DestinationFieldPath.IsValid();
+		const bool bFunctionValid = ConversionFunctionFieldPath.IsValid() || !HasConversionFunction();
+		return bSourceValid && bDestinationValid && bFunctionValid;
 	}
 
 	FMVVMVCompiledFieldPath GetSourceFieldPath() const
@@ -215,11 +218,22 @@ public:
 		return ConversionFunctionFieldPath;
 	}
 
+	bool HasConversionFunction() const
+	{
+		return (Flags & (uint8)EFlags::HasConversionFunction) != 0;
+	}
+
+	bool IsComplexFunction() const
+	{
+		uint8 AllFunctionFlags = (uint8)EFlags::IsConversionFunctionComplex | (uint8)EFlags::HasConversionFunction;
+		return (Flags & AllFunctionFlags) == AllFunctionFlags;
+	}
+
 private:
 	using IndexType = int16;
 
 	UPROPERTY()
-	FMVVMVCompiledFieldPath SourceFieldPath; // todo: make this an array to support multiple input to conversion functions
+	FMVVMVCompiledFieldPath SourceFieldPath;
 
 	UPROPERTY()
 	FMVVMVCompiledFieldPath DestinationFieldPath;
@@ -227,6 +241,16 @@ private:
 	UPROPERTY()
 	FMVVMVCompiledFieldPath ConversionFunctionFieldPath;
 
+	enum class EFlags : uint8
+	{
+		None = 0,
+		HasConversionFunction = 1 << 0,
+		IsConversionFunctionComplex = 1 << 1,
+	};
+
+	UPROPERTY()
+	uint8 Flags;
+
 #if WITH_EDITORONLY_DATA
 	UPROPERTY(meta = (IgnoreForMemberInitializationTest))
 	FGuid CompiledBindingLibraryId;
@@ -239,6 +263,7 @@ USTRUCT()
 struct MODELVIEWVIEWMODEL_API FMVVMCompiledBindingLibrary
 {
 	GENERATED_BODY()
+
 	friend UE::MVVM::FCompiledBindingLibraryCompiler;
 
 public:
diff --git a/Engine/Plugins/Runtime/ModelViewViewModel/Source/ModelViewViewModel/Public/Bindings/MVVMFieldPathHelper.h b/Engine/Plugins/Runtime/ModelViewViewModel/Source/ModelViewViewModel/Public/Bindings/MVVMFieldPathHelper.h
--- a/Engine/Plugins/Runtime/ModelViewViewModel/Source/ModelViewViewModel/Public/Bindings/MVVMFieldPathHelper.h
+++ b/Engine/Plugins/Runtime/ModelViewViewModel/Source/ModelViewViewModel/Public/Bindings/MVVMFieldPathHelper.h
@@ -14,28 +14,27 @@ template <typename ValueType, typename ErrorType> class TValueOrError;
 namespace UE::MVVM::FieldPathHelper
 {
 	/**
-	 * Info about the viewmodel and field.
+	 * Info about the viewmodel and field notify.
+	 * It can has invalid values and still be a valid OneTime binding.
 	 */
-	struct FParsedBindingInfo
+	struct FParsedNotifyBindingInfo
 	{
-		using FFieldVariantArray = TArray<FMVVMConstFieldVariant, TInlineAllocator<4>>;
-
 		/**
-		 * The viewmodel that is used to create the binding. When it's invalid the binding can be valid but only OneTime mode is supported.
-		 * @return the class of "ViewmodelB" when the path is "ViewmodelA.ViewmodelB.Vector.X"
+		 * The viewmodel that is used to create the binding.
+		 * Value is the class of "ViewmodelB" when the path is "ViewmodelA.ViewmodelB.Vector.X"
 		 */
-		TSubclassOf<UObject> NotifyFieldClass;
+		const UClass* NotifyFieldClass = nullptr;
 
 		/**
-		 * The property we will bind to. When it's invalid the binding can be valid but only OneTime mode is supported.
-		 * @return "Vector" when the path is "ViewmodelA.ViewmodelB.Vector.X"
+		 * The property we will bind to
+		 * Value is "Vector" when the path is "ViewmodelA.ViewmodelB.Vector.X"
 		 */
 		FFieldNotificationId NotifyFieldId;
 
 		/**
 		 * The viewmodel index in the original path.
 		 * When it's INDEX_NONE but the ParsedBindingInfo is valid, the class accessor is the Accessor class.
-		 * @return 1 when the path is "ViewmodelA.ViewmodelB.Vector.X"
+		 * Value 1 when the path is "ViewmodelA.ViewmodelB.Vector.X"
 		 */
 		int32 ViewModelIndex = INDEX_NONE;
 	};
@@ -47,11 +46,11 @@ namespace UE::MVVM::FieldPathHelper
 	UE_NODISCARD MODELVIEWVIEWMODEL_API TValueOrError<TArray<FMVVMConstFieldVariant>, FText> GenerateFieldPathList(TArrayView<const FMVVMConstFieldVariant> Fields, bool bForSourceBinding);
 
 	/**
-	 * Generate the FParsedBindingInfo from the parsed field path.
-	 * This gives us which FieldId and Viewmodel should be used by the path.
+	 * Generate the FParsedNotifyBindingInfo from the parsed field path.
+	 * This gives us which FieldId and Viewmodel that will be used by the path.
 	 * Note. The code may have changed since the binding was created in the editor. As an example, a sub-object can be a viewmodel now.
 	 */
-	UE_NODISCARD MODELVIEWVIEWMODEL_API TValueOrError<FParsedBindingInfo, FText> GetBindingInfoFromFieldPath(const UClass* Accessor, TArrayView<const FMVVMConstFieldVariant> Fields);
+	UE_NODISCARD MODELVIEWVIEWMODEL_API TValueOrError<FParsedNotifyBindingInfo, FText> GetNotifyBindingInfoFromFieldPath(const UClass* Accessor, TArrayView<const FMVVMConstFieldVariant> Fields);
 
 	/** */
 	UE_NODISCARD MODELVIEWVIEWMODEL_API FString ToString(TArrayView<const FMVVMFieldVariant> Fields);
diff --git a/Engine/Plugins/Runtime/ModelViewViewModel/Source/ModelViewViewModel/Public/Debugging/MVVMDebugging.h b/Engine/Plugins/Runtime/ModelViewViewModel/Source/ModelViewViewModel/Public/Debugging/MVVMDebugging.h
--- a/Engine/Plugins/Runtime/ModelViewViewModel/Source/ModelViewViewModel/Public/Debugging/MVVMDebugging.h
+++ b/Engine/Plugins/Runtime/ModelViewViewModel/Source/ModelViewViewModel/Public/Debugging/MVVMDebugging.h
@@ -72,7 +72,6 @@ public:
 		Failed_FieldIdNotFound,
 		Failed_InvalidSource,
 		Failed_InvalidSourceField,
-		Failed_InvalidDestinationField,
 		Failed_Unknown,
 	};
 
diff --git a/Engine/Plugins/Runtime/ModelViewViewModel/Source/ModelViewViewModel/Public/View/MVVMView.h b/Engine/Plugins/Runtime/ModelViewViewModel/Source/ModelViewViewModel/Public/View/MVVMView.h
--- a/Engine/Plugins/Runtime/ModelViewViewModel/Source/ModelViewViewModel/Public/View/MVVMView.h
+++ b/Engine/Plugins/Runtime/ModelViewViewModel/Source/ModelViewViewModel/Public/View/MVVMView.h
@@ -34,7 +34,7 @@ struct FMVVMViewSource
 
 	// Number of bindings connected to the source.
 	UPROPERTY(VisibleAnywhere, Category = "Viewmodel")
-	int32 RegisteredCout = 0;
+	int32 RegisteredCount = 0;
 
 	// The source is defined in the ClassExtension.
 	UPROPERTY(VisibleAnywhere, Category = "Viewmodel")
@@ -63,10 +63,14 @@ public:
 	void ConstructView(const UMVVMViewClass* ClassExtension);
 
 	//~ Begin UUserWidgetExtension implementation
+	//virtual void Initialize() override;
 	virtual void Construct() override;
 	virtual void Destruct() override;
 	//~ End UUserWidgetExtension implementation
 
+	void InitializeBindings();
+	void DeintializeBindings();
+
 	const UMVVMViewClass* GetViewClass() const
 	{
 		return ClassExtension;
@@ -75,57 +79,58 @@ public:
 	void ExecuteDelayedBinding(const FMVVMViewDelayedBinding& DelayedBinding) const;
 	void ExecuteEveryTickBindings() const;
 
-// todo a way to identify a binding from outside. maybe a unique name in the editor?
-	//UFUNCTION(BlueprintCallable, Category = "Viewmodel")
-	//void SetLibraryBindingEnabled(FGuid ViewModelId, FMVVMBindingName BindingName, bool bEnable);
-
-	//UFUNCTION(BlueprintCallable, Category = "Viewmodel")
-	//bool IsLibraryBindingEnabled(FGuid ViewModelId, FMVVMBindingName BindingName) const;
-
+	/** Find and return the viewmodel with the specified name. */
 	UFUNCTION(BlueprintCallable, Category = "Viewmodel")
 	TScriptInterface<INotifyFieldValueChanged> GetViewModel(FName ViewModelName) const;
 
+	/**
+	 * Set the viewmodel of the specified name.
+	 * The viewmodel needs to be settable and the type should match (child of the defined viewmodel).
+	 * If the view is initialized, all bindings that uses that viewmodel will be re-executed with the new viewmodel instance.
+	 */
 	UFUNCTION(BlueprintCallable, Category = "Viewmodel")
 	bool SetViewModel(FName ViewModelName, TScriptInterface<INotifyFieldValueChanged> ViewModel);
 
 private:
+	bool EvaluateSourceCreator(int32 SourceIndex);
+	bool SetSourceInternal(FName ViewModelName, TScriptInterface<INotifyFieldValueChanged> ViewModel, bool bForDynamicSource);
+
 	void HandledLibraryBindingValueChanged(UObject* InViewModel, UE::FieldNotification::FFieldId InFieldId, int32 InCompiledBindingIndex) const;
 
-	void ExecuteLibraryBinding(const FMVVMViewClass_CompiledBinding& Binding) const;
-	void ExecuteLibraryBinding(const FMVVMViewClass_CompiledBinding& Binding, UObject* Source) const;
+	void ExecuteLibraryBinding(const FMVVMViewClass_CompiledBinding& Binding, int32 BindingIndex) const;
 
 	void EnableLibraryBinding(const FMVVMViewClass_CompiledBinding& Item, int32 BindingIndex);
 	void DisableLibraryBinding(const FMVVMViewClass_CompiledBinding& Item, int32 BindingIndex);
 	bool IsLibraryBindingEnabled(int32 InBindindIndex) const;
 
 	bool RegisterLibraryBinding(const FMVVMViewClass_CompiledBinding& Binding, int32 BindingIndex);
-	void UnregisterLibraryBinding(const FMVVMViewClass_CompiledBinding& Binding);
+	void UnregisterLibraryBinding(const FMVVMViewClass_CompiledBinding& Binding, int32 BindingIndex);
 
-	TScriptInterface<INotifyFieldValueChanged> FindSource(const FMVVMViewClass_CompiledBinding& Binding, bool bAllowNull) const;
+	TScriptInterface<INotifyFieldValueChanged> FindSource(const FMVVMViewClass_CompiledBinding& Binding, int32 BindingIndex, bool bAllowNull) const;
 	FMVVMViewSource* FindViewSource(const FName SourceName);
 	const FMVVMViewSource* FindViewSource(const FName SourceName) const;
 
 private:
-// todo support dynamic runtime binding.
-	/** Binding that are added dynamically at runtime. */
-	//TArray<FMVVMView_Binding> RegisteredDynamicBindings;
-
 	UPROPERTY(Transient)
 	TObjectPtr<const UMVVMViewClass> ClassExtension;
 
 	UPROPERTY(VisibleAnywhere, Transient, Category = "Viewmodel")
 	TArray<FMVVMViewSource> Sources;
 
-	/** The binding that are enabled for the instance. */
-	TBitArray<> EnabledLibraryBindings;
+	/** The binding that are registered by the view to the sources. */
+	TBitArray<> RegisteredLibraryBindings;
 
 	/** Should log when a binding is executed. */
 	UPROPERTY(EditAnywhere, Transient, Category = "Viewmodel")
 	bool bLogBinding = false;
 
-	/** Is the Construct method was called. */
+	/** Is the Construct method called. */
 	UPROPERTY(VisibleAnywhere, Transient, Category = "Viewmodel")
 	bool bConstructed = false;
+	
+	/** Is the Initialize method called. */
+	UPROPERTY(VisibleAnywhere, Transient, Category = "Viewmodel")
+	bool bInitialized = false;
 
 	/** The view has at least one binding that need to be ticked every frame. */
 	UPROPERTY(VisibleAnywhere, Transient, Category = "Viewmodel")
diff --git a/Engine/Plugins/Runtime/ModelViewViewModel/Source/ModelViewViewModel/Public/View/MVVMViewClass.h b/Engine/Plugins/Runtime/ModelViewViewModel/Source/ModelViewViewModel/Public/View/MVVMViewClass.h
--- a/Engine/Plugins/Runtime/ModelViewViewModel/Source/ModelViewViewModel/Public/View/MVVMViewClass.h
+++ b/Engine/Plugins/Runtime/ModelViewViewModel/Source/ModelViewViewModel/Public/View/MVVMViewClass.h
@@ -35,12 +35,6 @@ struct FMVVMViewClass_SourceCreator
 	friend UE::MVVM::Private::FMVVMViewBlueprintCompiler;
 
 public:
-	MODELVIEWVIEWMODEL_API static FMVVMViewClass_SourceCreator MakeManual(FName Name, UClass* NotifyFieldValueChangedClass);
-	MODELVIEWVIEWMODEL_API static FMVVMViewClass_SourceCreator MakeInstance(FName Name, UClass* NotifyFieldValueChangedClass);
-	MODELVIEWVIEWMODEL_API static FMVVMViewClass_SourceCreator MakeFieldPath(FName Name, UClass* NotifyFieldValueChangedClass, FMVVMVCompiledFieldPath FieldPath, bool bOptional);
-	MODELVIEWVIEWMODEL_API static FMVVMViewClass_SourceCreator MakeGlobalContext(FName Name, FMVVMViewModelContext Context, bool bOptional);
-	MODELVIEWVIEWMODEL_API static FMVVMViewClass_SourceCreator MakeResolver(FName Name, UClass* InNotifyFieldValueChangedClass, UMVVMViewModelContextResolver* Source, bool bOptional);
-
 	UObject* CreateInstance(const UMVVMViewClass* ViewClass, UMVVMView* View, UUserWidget* UserWidget) const;
 
 	UClass* GetSourceClass() const
@@ -50,7 +44,17 @@ public:
 
 	bool IsSourceAUserWidgetProperty() const
 	{
-		return bIsUserWidgetProperty;
+		return (Flags & (uint8)ESourceFlags::IsUserWidgetProperty) != 0;
+	}
+	
+	bool CanBeSet() const
+	{
+		return (Flags & (uint8)ESourceFlags::CanBeSet) != 0;
+	}
+
+	bool CanBeEvaluated() const
+	{
+		return (Flags & (uint8)ESourceFlags::CanBeEvaluated) != 0;
 	}
 
 	FName GetSourceName() const
@@ -58,6 +62,25 @@ public:
 		return PropertyName;
 	}
 
+#if UE_WITH_MVVM_DEBUGGING
+	struct FToStringArgs
+	{
+		bool bUseDisplayName = true;
+		bool bAddCreationMode = true;
+		bool bAddFlags = true;
+
+		MODELVIEWVIEWMODEL_API static FToStringArgs Short();
+		MODELVIEWVIEWMODEL_API static FToStringArgs All();
+	};
+
+	/** @return a human readable version of the source that can be use for debugging purposes. */
+	MODELVIEWVIEWMODEL_API FString ToString(const FMVVMCompiledBindingLibrary& CompiledBindingLibrary, FToStringArgs Args) const;
+#endif
+
+#if WITH_EDITOR
+	MODELVIEWVIEWMODEL_API void PostSerialize(const FArchive& Ar);
+#endif
+
 private:
 	/** Class type to create a source at runtime. */
 	UPROPERTY()
@@ -81,16 +104,42 @@ private:
 
 	UPROPERTY()
 	FName PropertyName;
+
+	enum class ESourceFlags : uint8
+	{
+		None = 0,
+		TypeCreateInstance = 1 << 0,
+		IsUserWidgetProperty = 1 << 1,
+		IsOptional = 1 << 2,
+		CanBeSet = 1 << 3,
+		CanBeEvaluated = 1 << 4,
+	};
+
+	UPROPERTY()
+	uint8 Flags = (uint8)ESourceFlags::None;
 	
+#if WITH_EDITORONLY_DATA
 	UPROPERTY()
-	bool bCreateInstance = false;
+	bool bCreateInstance_DEPRECATED = false;
 
 	UPROPERTY()
-	bool bIsUserWidgetProperty = true;
+	bool bIsUserWidgetProperty_DEPRECATED = true;
 
 	UPROPERTY()
-	bool bOptional = false;
+	bool bOptional_DEPRECATED = false;
+#endif
+};
+
+#if WITH_EDITOR
+template<>
+struct TStructOpsTypeTraits<FMVVMViewClass_SourceCreator> : public TStructOpsTypeTraitsBase2<FMVVMViewClass_SourceCreator>
+{
+	enum
+	{
+		WithPostSerialize = true,
+	};
 };
+#endif
 
 
 /**
@@ -120,6 +169,12 @@ public:
 		return SourcePropertyName;
 	}
 
+	/** @return The SourceCreator index when the binding is of the EvaluateSourceCreator type. */
+	int32 GetEvaluateSourceCreatorBindingIndex() const
+	{
+		return EvaluateSourceCreatorIndex;
+	}
+
 	/** @return true if SourceFieldId is not from a property but the object is the UserWidget. */
 	bool IsSourceUserWidget() const
 	{
@@ -135,7 +190,7 @@ public:
 	/** @return true if this Binding should be executed at initialization. */
 	bool NeedsExecutionAtInitialization() const
 	{
-		return (Flags & EBindingFlags::ForwardBinding) != 0;
+		return (Flags & EBindingFlags::ExecuteAtInitialization) != 0;
 	}
 
 	/** @return true if this Binding should be executed once (at initialization) but should not be executed when the SourceFieldId value changes. */
@@ -155,16 +210,26 @@ public:
 	{
 		return (Flags & EBindingFlags::ViewModelOptional) != 0;
 	}
+	
+	/** @return true if the binding is not valid and we should only evaluate the source creator. */
+	bool IsEvaluateSourceCreatorBinding() const
+	{
+		return EvaluateSourceCreatorIndex != INDEX_NONE;
+	}
 
-	/** @return true if the binding use a conversion function and that the conversion function is complex. */
+	/**
+	 * @return true if the binding use a conversion function and that the conversion function is complex.
+	 * The conversion function is complex, there is no input. The inputs are calculated in the BP function.
+	 */
 	bool IsConversionFunctionComplex() const
 	{
-		return (Flags & EBindingFlags::ConversionFunctionIsComplex) != 0;
+		return Binding.IsComplexFunction();
 	}
 	
 	/** How the binding should be executed. */
 	EMVVMExecutionMode GetExecuteMode() const;
 
+#if UE_WITH_MVVM_DEBUGGING
 	struct FToStringArgs
 	{
 		bool bUseDisplayName = true;
@@ -177,11 +242,13 @@ public:
 		MODELVIEWVIEWMODEL_API static FToStringArgs All();
 	};
 
+
 	UE_DEPRECATED(5.3, "ToString with no argument is deprecated.")
 	FString ToString() const;
 
 	/** @return a human readable version of the binding that can be use for debugging purposes. */
 	FString ToString(const FMVVMCompiledBindingLibrary& CompiledBindingLibrary, FToStringArgs Args) const;
+#endif
 
 #if WITH_EDITOR
 	FGuid GetEditorId() const
@@ -203,17 +270,19 @@ private:
 	UPROPERTY(VisibleAnywhere, Category = "Viewmodel")
 	EMVVMExecutionMode ExecutionMode = EMVVMExecutionMode::Immediate;
 
+	UPROPERTY(VisibleAnywhere, Category = "Viewmodel")
+	int8 EvaluateSourceCreatorIndex = INDEX_NONE;
+
 	enum EBindingFlags
 	{
 		None = 0,
-		/** True when the binding goes from Source to Destination. */
-		ForwardBinding = 1 << 0,
+		ExecuteAtInitialization = 1 << 0,
+		Unused01 = 1 << 1,
 		OneTime = 1 << 2,
 		EnabledByDefault = 1 << 3,
 		/** The source (viewmodel) can be nullptr and the binding could failed and should not log a warning. */
 		ViewModelOptional = 1 << 4,
-		/** The conversion function is complex, there is no input. The inputs are calculated in the BP function. */
-		ConversionFunctionIsComplex = 1 << 5,
+		Unused02 = 1 << 5,
 		/** In development, (when the Blueprint maybe not be compiled with the latest data), the ExecutionMode may not reflect the default project setting value. */
 		OverrideExecuteMode = 1 << 6,
 		/** When the source object is the object itself. */
@@ -246,6 +315,12 @@ public:
 	//~ End UWidgetBlueprintGeneratedClassExtension
 
 public:
+	/** Should it automatically execute the bindings when the view is constructed or they will be executed manually. */
+	bool DoesInitializeAtConstruction() const
+	{
+		return bInitializeAtConstruction;
+	}
+	
 	/** Get the list of the needed ViewModel. */
 	const TArrayView<const FMVVMViewClass_SourceCreator> GetViewModelCreators() const
 	{
@@ -270,6 +345,10 @@ public:
 		return CompiledBindings[Index];
 	}
 
+#if UE_WITH_MVVM_DEBUGGING
+	void Log(FMVVMViewClass_SourceCreator::FToStringArgs SourceArgs, FMVVMViewClass_CompiledBinding::FToStringArgs BindingArgs) const;
+#endif
+
 private:
 	/**  */
 	TArrayView<FMVVMViewClass_CompiledBinding> GetCompiledBindings()
@@ -290,6 +369,10 @@ private:
 	UPROPERTY()
 	FMVVMCompiledBindingLibrary BindingLibrary;
 
+	/** */
+	UPROPERTY()
+	bool bInitializeAtConstruction = true;
+
 	/** */
 	bool bLoaded = false;
 };
diff --git a/Engine/Plugins/Runtime/ModelViewViewModel/Source/ModelViewViewModelBlueprint/Private/Bindings/MVVMCompiledBindingLibraryCompiler.cpp b/Engine/Plugins/Runtime/ModelViewViewModel/Source/ModelViewViewModelBlueprint/Private/Bindings/MVVMCompiledBindingLibraryCompiler.cpp
--- a/Engine/Plugins/Runtime/ModelViewViewModel/Source/ModelViewViewModelBlueprint/Private/Bindings/MVVMCompiledBindingLibraryCompiler.cpp
+++ b/Engine/Plugins/Runtime/ModelViewViewModel/Source/ModelViewViewModelBlueprint/Private/Bindings/MVVMCompiledBindingLibraryCompiler.cpp
@@ -20,7 +20,7 @@ static const FName NAME_BlueprintGetter = "BlueprintGetter";
 /** */
 struct FRawFieldId
 {
-	UClass* NotifyFieldValueChangedClass;
+	const UClass* NotifyFieldValueChangedClass;
 	UE::FieldNotification::FFieldId FieldId;
 
 	int32 LoadedFieldIdIndex = INDEX_NONE;
@@ -66,17 +66,18 @@ public:
 /** */
 struct FRawBinding
 {
-	TArray<FCompiledBindingLibraryCompiler::FFieldPathHandle> SourcePathHandles;
+	FCompiledBindingLibraryCompiler::FFieldPathHandle SourcePathHandle;
 	FCompiledBindingLibraryCompiler::FFieldPathHandle DestinationPathHandle;
 	FCompiledBindingLibraryCompiler::FFieldPathHandle ConversionFunctionPathHandle;
 
 	FCompiledBindingLibraryCompiler::FBindingHandle BindingHandle;
 	FMVVMVCompiledBinding CompiledBinding;
+	bool bIsConversionFunctionComplex = false;
 
 public:
 	bool IsSameBinding(const FRawBinding& Binding) const
 	{
-		return Binding.SourcePathHandles == SourcePathHandles
+		return Binding.SourcePathHandle == SourcePathHandle
 			&& Binding.DestinationPathHandle == DestinationPathHandle
 			&& Binding.ConversionFunctionPathHandle == ConversionFunctionPathHandle;
 	}
@@ -137,7 +138,7 @@ FCompiledBindingLibraryCompiler::FCompiledBindingLibraryCompiler()
 
 }
 
-TValueOrError<FCompiledBindingLibraryCompiler::FFieldIdHandle, FText> FCompiledBindingLibraryCompiler::AddFieldId(TSubclassOf<UObject> SourceClass, FName FieldId)
+TValueOrError<FCompiledBindingLibraryCompiler::FFieldIdHandle, FText> FCompiledBindingLibraryCompiler::AddFieldId(const UClass* SourceClass, FName FieldId)
 {
 	Impl->bCompiled = false;
 
@@ -165,12 +166,12 @@ TValueOrError<FCompiledBindingLibraryCompiler::FFieldIdHandle, FText> FCompiledB
 
 		int32 FoundFieldIdIndex = Impl->FieldIds.IndexOfByPredicate([FoundFieldId, SourceClass](const Private::FRawFieldId& Other)
 			{
-				return Other.FieldId == FoundFieldId && Other.NotifyFieldValueChangedClass == SourceClass.Get();
+				return Other.FieldId == FoundFieldId && Other.NotifyFieldValueChangedClass == SourceClass;
 			});
 		if (FoundFieldIdIndex == INDEX_NONE)
 		{
 			Private::FRawFieldId RawFieldId;
-			RawFieldId.NotifyFieldValueChangedClass = SourceClass.Get();
+			RawFieldId.NotifyFieldValueChangedClass = SourceClass;
 			RawFieldId.FieldId = FoundFieldId;
 			RawFieldId.IdHandle = FFieldIdHandle::MakeHandle();
 			FoundFieldIdIndex = Impl->FieldIds.Add(MoveTemp(RawFieldId));
@@ -360,7 +361,7 @@ TValueOrError<FCompiledBindingLibraryCompiler::FFieldPathHandle, FText> FCompile
 }
 
 
-TValueOrError<FCompiledBindingLibraryCompiler::FFieldPathHandle, FText> FCompiledBindingLibraryCompiler::AddConversionFunctionFieldPath(TSubclassOf<UObject> SourceClass, const UFunction* Function)
+TValueOrError<FCompiledBindingLibraryCompiler::FFieldPathHandle, FText> FCompiledBindingLibraryCompiler::AddConversionFunctionFieldPath(const UClass* SourceClass, const UFunction* Function)
 {
 	Impl->bCompiled = false;
 
@@ -417,29 +418,34 @@ TValueOrError<FCompiledBindingLibraryCompiler::FFieldPathHandle, FText> FCompile
 
 TValueOrError<FCompiledBindingLibraryCompiler::FBindingHandle, FText> FCompiledBindingLibraryCompiler::AddBinding(FFieldPathHandle InSourceHandle, FFieldPathHandle InDestinationHandle)
 {
-	return AddBinding(InSourceHandle, InDestinationHandle, FFieldPathHandle());
+	return AddBindingImpl(InSourceHandle, InDestinationHandle, FFieldPathHandle(), false);
 }
 
 
 TValueOrError<FCompiledBindingLibraryCompiler::FBindingHandle, FText> FCompiledBindingLibraryCompiler::AddBinding(FFieldPathHandle InSourceHandle, FFieldPathHandle InDestinationHandle, FFieldPathHandle InConversionFunctionHandle)
 {
-	return AddBinding(MakeArrayView(&InSourceHandle, 1), InDestinationHandle, InConversionFunctionHandle);
+	return AddBindingImpl(InSourceHandle, InDestinationHandle, InConversionFunctionHandle, false);
 }
 
 
-TValueOrError<FCompiledBindingLibraryCompiler::FBindingHandle, FText> FCompiledBindingLibraryCompiler::AddBinding(TArrayView<const FFieldPathHandle> InSourceHandles, FFieldPathHandle InDestinationHandle, FFieldPathHandle InConversionFunctionHandle)
+TValueOrError<FCompiledBindingLibraryCompiler::FBindingHandle, FText> FCompiledBindingLibraryCompiler::AddComplexBinding(FFieldPathHandle InDestinationHandle, FFieldPathHandle InConversionFunctionHandle)
+{
+	return AddBindingImpl(FFieldPathHandle(), InDestinationHandle, InConversionFunctionHandle, true);
+}
+
+
+TValueOrError<FCompiledBindingLibraryCompiler::FBindingHandle, FText> FCompiledBindingLibraryCompiler::AddBindingImpl(FFieldPathHandle InSourceHandle, FFieldPathHandle InDestinationHandle, FFieldPathHandle InConversionFunctionHandle, bool bInIsComplexBinding)
 {
 	Impl->bCompiled = false;
 
 	UMVVMSubsystem::FConstDirectionalBindingArgs DirectionBindingArgs;
 
-	ensureMsgf(InSourceHandles.Num() == 1, TEXT("Conversion function with more than one arguements is not yet supported."));
-
+	// Complex Conversion function do not have input arguments.
+	if (!bInIsComplexBinding)
 	{
-		FFieldPathHandle FirstSourceHandle = InSourceHandles[0];
-		const int32 FoundSourceFieldPath = Impl->FieldPaths.IndexOfByPredicate([FirstSourceHandle](const Private::FRawFieldPath& Other)
+		const int32 FoundSourceFieldPath = Impl->FieldPaths.IndexOfByPredicate([InSourceHandle](const Private::FRawFieldPath& Other)
 			{
-				return Other.PathHandle == FirstSourceHandle;
+				return Other.PathHandle == InSourceHandle;
 			});
 		if (FoundSourceFieldPath == INDEX_NONE)
 		{
@@ -527,19 +533,20 @@ TValueOrError<FCompiledBindingLibraryCompiler::FBindingHandle, FText> FCompiledB
 	}
 
 	Private::FRawBinding NewBinding;
-	NewBinding.SourcePathHandles = InSourceHandles;
+	NewBinding.SourcePathHandle = InSourceHandle;
 	NewBinding.DestinationPathHandle = InDestinationHandle;
 	NewBinding.ConversionFunctionPathHandle = InConversionFunctionHandle;
+	NewBinding.bIsConversionFunctionComplex = bInIsComplexBinding;
 	const int32 FoundSameBindingIndex = Impl->Bindings.IndexOfByPredicate([&NewBinding](const Private::FRawBinding& Binding)
 		{
 			return NewBinding.IsSameBinding(Binding);
 		});
 
-	//if (FoundSameBindingIndex != INDEX_NONE)
-	//{
-	//	return MakeError(LOCTEXT("BindingAlreadyAdded", "The binding already exist."));
-	//}
-	//else
+	if (FoundSameBindingIndex != INDEX_NONE)
+	{
+		return MakeError(LOCTEXT("BindingAlreadyAdded", "The binding already exist."));
+	}
+	else
 	{
 		FCompiledBindingLibraryCompiler::FBindingHandle ResultBindingHandle = FBindingHandle::MakeHandle();
 		NewBinding.BindingHandle = ResultBindingHandle;
@@ -769,14 +776,18 @@ TValueOrError<FCompiledBindingLibraryCompiler::FCompileResult, FText> FCompiledB
 		Binding.CompiledBinding.CompiledBindingLibraryId = Result.Library.CompiledBindingLibraryId;
 		check(Binding.CompiledBinding.CompiledBindingLibraryId.IsValid());
 
-		Binding.CompiledBinding.SourceFieldPath = GetCompiledFieldPath(Binding.SourcePathHandles[0]);
-		check(Binding.CompiledBinding.SourceFieldPath.IsValid());
+		Binding.CompiledBinding.SourceFieldPath = GetCompiledFieldPath(Binding.SourcePathHandle);
+		check(Binding.bIsConversionFunctionComplex || Binding.CompiledBinding.SourceFieldPath.IsValid());
 
 		Binding.CompiledBinding.DestinationFieldPath = GetCompiledFieldPath(Binding.DestinationPathHandle);
 		check(Binding.CompiledBinding.DestinationFieldPath.IsValid());
 
 		Binding.CompiledBinding.ConversionFunctionFieldPath = GetCompiledFieldPath(Binding.ConversionFunctionPathHandle);
 
+		Binding.CompiledBinding.Flags = 0;
+		Binding.CompiledBinding.Flags |= (Binding.ConversionFunctionPathHandle.IsValid()) ? (uint8)FMVVMVCompiledBinding::EFlags::HasConversionFunction : 0;
+		Binding.CompiledBinding.Flags |= (Binding.bIsConversionFunctionComplex) ? (uint8)FMVVMVCompiledBinding::EFlags::IsConversionFunctionComplex : 0;
+
 		Result.Bindings.Add(Binding.BindingHandle, Binding.CompiledBinding);
 	}
 
diff --git a/Engine/Plugins/Runtime/ModelViewViewModel/Source/ModelViewViewModelBlueprint/Private/MVVMViewBlueprintCompiler.cpp b/Engine/Plugins/Runtime/ModelViewViewModel/Source/ModelViewViewModelBlueprint/Private/MVVMViewBlueprintCompiler.cpp
--- a/Engine/Plugins/Runtime/ModelViewViewModel/Source/ModelViewViewModelBlueprint/Private/MVVMViewBlueprintCompiler.cpp
+++ b/Engine/Plugins/Runtime/ModelViewViewModel/Source/ModelViewViewModelBlueprint/Private/MVVMViewBlueprintCompiler.cpp
@@ -22,11 +22,11 @@
 
 namespace UE::MVVM::Private
 {
-	FAutoConsoleVariable CVarOnCompliedLogBinding(
-		TEXT("MVVM.OnCompliedLogBinding"),
-		false,
-		TEXT("After the view is compiled log the compiled bindings.")
-	);
+FAutoConsoleVariable CVarLogViewCompliedResult(
+	TEXT("MVVM.LogViewCompliedResult"),
+	false,
+	TEXT("After the view is compiled log the compiled bindings and sources.")
+);
 
 FString PropertyPathToString(const UMVVMBlueprintView* BlueprintView, const FMVVMBlueprintPropertyPath& PropertyPath)
 {
@@ -75,6 +75,26 @@ FText GetViewModelIdText(const FMVVMBlueprintPropertyPath& PropertyPath)
 	return FText::FromString(PropertyPath.GetViewModelId().ToString(EGuidFormats::DigitsWithHyphensInBraces));
 }
 
+TValueOrError<FCompiledBindingLibraryCompiler::FFieldPathHandle, FText> AddObjectFieldPath(FCompiledBindingLibraryCompiler& BindingLibraryCompiler, const UWidgetBlueprintGeneratedClass* Class, FStringView ObjectPath, UClass* ExpectedType)
+{
+	// Generate a path to read the value at runtime
+	static const FText InvalidGetterFormat = LOCTEXT("ViewModelInvalidGetterWithReason", "Viewmodel has an invalid Getter. {0}");
+
+	TValueOrError<TArray<FMVVMConstFieldVariant>, FText> GeneratedField = FieldPathHelper::GenerateFieldPathList(Class, ObjectPath, true);
+	if (GeneratedField.HasError())
+	{
+		return MakeError(FText::Format(InvalidGetterFormat, GeneratedField.GetError()));
+	}
+
+	TValueOrError<FCompiledBindingLibraryCompiler::FFieldPathHandle, FText> ReadFieldPathResult = BindingLibraryCompiler.AddObjectFieldPath(GeneratedField.GetValue(), ExpectedType, true);
+	if (ReadFieldPathResult.HasError())
+	{
+		return MakeError(FText::Format(InvalidGetterFormat, ReadFieldPathResult.GetError()));
+	}
+
+	return MakeValue(ReadFieldPathResult.StealValue());
+}
+
 void FMVVMViewBlueprintCompiler::AddMessageForBinding(FMVVMBlueprintViewBinding& Binding, UMVVMBlueprintView* BlueprintView, const FText& MessageText, EBindingMessageType MessageType, FName ArgumentName) const
 {
 	const FText BindingName = FText::FromString(Binding.GetDisplayNameString(WidgetBlueprintCompilerContext.WidgetBlueprint()));
@@ -332,7 +352,6 @@ void FMVVMViewBlueprintCompiler::CreateSourceLists(const FWidgetBlueprintCompile
 {
 	CompilerUserWidgetPropertyContexts.Reset();
 	CompilerSourceCreatorContexts.Reset();
-	ViewModelDynamicCounter = 0;
 
 	TSet<FGuid> ViewModelGuids;
 	TSet<FName> WidgetSources;
@@ -674,19 +693,16 @@ bool FMVVMViewBlueprintCompiler::Compile(UWidgetBlueprintGeneratedClass* Class,
 	{
 		ViewExtension->BindingLibrary = MoveTemp(CompileResult.GetValue().Library);
 
-		if (CVarOnCompliedLogBinding->GetBool())
+#if UE_WITH_MVVM_DEBUGGING
+		if (CVarLogViewCompliedResult->GetBool())
 		{
-			ViewExtension->BindingLibrary.Load();
-			UE_LOG(LogMVVM, Log, TEXT("Compiled Bindings for %s"), *ViewExtension->GetOutermost()->GetName());
-
-			FMVVMViewClass_CompiledBinding::FToStringArgs ToStringArgs = FMVVMViewClass_CompiledBinding::FToStringArgs::All();
-			ToStringArgs.bUseDisplayName = false;
-			for (const FMVVMViewClass_CompiledBinding& Binding : ViewExtension->GetCompiledBindings())
-			{
-				UE_LOG(LogMVVM, Log, TEXT("%s"), *Binding.ToString(ViewExtension->GetBindingLibrary(), ToStringArgs));
-			}
-			ViewExtension->BindingLibrary.Unload();
+			FMVVMViewClass_SourceCreator::FToStringArgs CreatorsToStringArgs = FMVVMViewClass_SourceCreator::FToStringArgs::All();
+			CreatorsToStringArgs.bUseDisplayName = false;
+			FMVVMViewClass_CompiledBinding::FToStringArgs BindingToStringArgs = FMVVMViewClass_CompiledBinding::FToStringArgs::All();
+			BindingToStringArgs.bUseDisplayName = false;
+			ViewExtension->Log(CreatorsToStringArgs, BindingToStringArgs);
 		}
+#endif
 	}
 
 	return bResult;
@@ -695,7 +711,7 @@ bool FMVVMViewBlueprintCompiler::Compile(UWidgetBlueprintGeneratedClass* Class,
 
 bool FMVVMViewBlueprintCompiler::PreCompileBindingSources(UWidgetBlueprintGeneratedClass* Class, UMVVMBlueprintView* BlueprintView)
 {
-	const int32 NumBindings = BlueprintView->GetNumBindings();
+	const int32 NumBindings = BlueprintView->GetNumBindings(); // NB Binding can be added when creating a dynamic vm
 	for (int32 Index = 0; Index < NumBindings; ++Index)
 	{
 		FMVVMBlueprintViewBinding* BindingPtr = BlueprintView->GetBindingAt(Index);
@@ -712,9 +728,11 @@ bool FMVVMViewBlueprintCompiler::PreCompileBindingSources(UWidgetBlueprintGenera
 			continue;
 		}
 
-		auto CreateSourceContextForPropertyPath = [this, &Binding, BlueprintView, Class, Index](const FMVVMBlueprintPropertyPath& Path, bool bForwardBinding, FName ArgumentName = FName()) -> bool
+		bool bIsOneTimeBinding = IsOneTimeBinding(Binding.BindingType);
+
+		auto CreateSourceContextForPropertyPath = [this, &Binding, BlueprintView, Class, Index, bIsOneTimeBinding](const FMVVMBlueprintPropertyPath& Path, bool bForwardBinding, FName ArgumentName) -> bool
 		{
-			const TValueOrError<FBindingSourceContext, FText> CreatedBindingSourceContext = CreateBindingSourceContext(BlueprintView, Class, Path);
+			const TValueOrError<FBindingSourceContext, FText> CreatedBindingSourceContext = CreateBindingSourceContext(BlueprintView, Class, Path, bIsOneTimeBinding);
 			if (CreatedBindingSourceContext.HasError())
 			{
 				AddMessageForBinding(Binding, BlueprintView,
@@ -756,18 +774,38 @@ bool FMVVMViewBlueprintCompiler::PreCompileBindingSources(UWidgetBlueprintGenera
 
 			BindingSourceContext.BindingIndex = Index;
 			BindingSourceContext.bIsForwardBinding = bForwardBinding;
+			BindingSourceContext.bIsComplexBinding = !ArgumentName.IsNone();
 
 			this->BindingSourceContexts.Add(MoveTemp(BindingSourceContext));
 			return true;
 		};
 
-		auto CreateSourcesForConversionFunction = [this, &Binding, &CreateSourceContextForPropertyPath](bool bForwardBinding)
+		enum class ECreateSourcesForConversionFunctionResult : uint8 { Valid, Failed, Continue };
+		auto CreateSourcesForConversionFunction = [this, &Binding, BlueprintView, &CreateSourceContextForPropertyPath](bool bForwardBinding)
 		{
+			ECreateSourcesForConversionFunctionResult Result = ECreateSourcesForConversionFunctionResult::Continue;
 			TMap<FName, FMVVMBlueprintPropertyPath> ArgumentPaths = ConversionFunctionHelper::GetAllArgumentPropertyPaths(WidgetBlueprintCompilerContext.WidgetBlueprint(), Binding, bForwardBinding, true);
 			for (const TPair<FName, FMVVMBlueprintPropertyPath>& Pair : ArgumentPaths)
 			{
-				CreateSourceContextForPropertyPath(Pair.Value, bForwardBinding, Pair.Key);
+				if (Pair.Key.IsNone())
+				{
+					AddMessageForBinding(Binding, BlueprintView,
+						FText::Format(LOCTEXT("InvalidArgumentPathName", "The conversion function {0} has an invalid argument"), FText::FromString(Binding.GetDisplayNameString(WidgetBlueprintCompilerContext.WidgetBlueprint()))),
+						EBindingMessageType::Error
+					);
+					Result = ECreateSourcesForConversionFunctionResult::Failed;
+				}
+				else if (CreateSourceContextForPropertyPath(Pair.Value, bForwardBinding, Pair.Key))
+				{
+					Result = ECreateSourcesForConversionFunctionResult::Valid;
+				}
+				else
+				{
+					Result = ECreateSourcesForConversionFunctionResult::Failed;
+					break;
+				}
 			}
+			return Result;
 		};
 
 		auto AddWarningForPropertyWithMVVMAndLegacyBinding = [this, &Binding, &BlueprintView, Class](const FMVVMBlueprintPropertyPath& Path)
@@ -822,40 +860,68 @@ bool FMVVMViewBlueprintCompiler::PreCompileBindingSources(UWidgetBlueprintGenera
 			}
 		};
 
+		// Add the forward binding. If the binding has a conversion function, use it instead of the regular binding.
 		if (IsForwardBinding(Binding.BindingType))
 		{
-			CreateSourcesForConversionFunction(true);
-
-			if (!Binding.SourcePath.IsEmpty())
+			ECreateSourcesForConversionFunctionResult ConversionFunctionResult = CreateSourcesForConversionFunction(true);
+			if (ConversionFunctionResult == ECreateSourcesForConversionFunctionResult::Continue)
 			{
-				if (!Binding.DestinationPath.IsEmpty())
+				if (!Binding.SourcePath.IsEmpty())
 				{
-					AddWarningForPropertyWithMVVMAndLegacyBinding(Binding.DestinationPath);
+					if (!Binding.DestinationPath.IsEmpty())
+					{
+						AddWarningForPropertyWithMVVMAndLegacyBinding(Binding.DestinationPath);
+					}
+					if (!CreateSourceContextForPropertyPath(Binding.SourcePath, true, FName()))
+					{
+						bIsBindingsValid = false;
+						continue;
+					}
 				}
-				if (!CreateSourceContextForPropertyPath(Binding.SourcePath, true))
+				else
 				{
+					AddMessageForBinding(Binding, BlueprintView, LOCTEXT("BindingEmptySourcePath", "The binding doesn't have a Source or a Conversion function."), EBindingMessageType::Error);
 					bIsBindingsValid = false;
 					continue;
 				}
 			}
+			else if (ConversionFunctionResult == ECreateSourcesForConversionFunctionResult::Failed)
+			{
+				bIsBindingsValid = false;
+				continue;
+			}
 		}
 
+		// Add the backward binding. If the binding has a conversion function, use it instead of the regular binding.
 		if (IsBackwardBinding(Binding.BindingType))
 		{
-			CreateSourcesForConversionFunction(false);
-
-			if (!Binding.DestinationPath.IsEmpty() && Binding.DestinationPath.HasPaths())
+			ECreateSourcesForConversionFunctionResult ConversionFunctionResult = CreateSourcesForConversionFunction(false);
+			if (ConversionFunctionResult == ECreateSourcesForConversionFunctionResult::Continue)
 			{
-				if (!Binding.SourcePath.IsEmpty())
+				if (!Binding.DestinationPath.IsEmpty())
 				{
-					AddWarningForPropertyWithMVVMAndLegacyBinding(Binding.SourcePath);
+					if (!Binding.SourcePath.IsEmpty())
+					{
+						AddWarningForPropertyWithMVVMAndLegacyBinding(Binding.SourcePath);
+					}
+					if (!CreateSourceContextForPropertyPath(Binding.DestinationPath, false, FName()))
+					{
+						bIsBindingsValid = false;
+						continue;
+					}
 				}
-				if (!CreateSourceContextForPropertyPath(Binding.DestinationPath, false))
+				else
 				{
+					AddMessageForBinding(Binding, BlueprintView, LOCTEXT("BindingEmptyDestinationPath", "The binding doesn't have a Destination or a Conversion function."), EBindingMessageType::Error);
 					bIsBindingsValid = false;
 					continue;
 				}
 			}
+			else if (ConversionFunctionResult == ECreateSourcesForConversionFunctionResult::Failed)
+			{
+				bIsBindingsValid = false;
+				continue;
+			}
 		}
 	}
 
@@ -918,21 +984,10 @@ bool FMVVMViewBlueprintCompiler::PreCompileSourceCreators(UWidgetBlueprintGenera
 					continue;
 				}
 
-				// Generate a path to read the value at runtime
-				static const FText InvalidGetterFormat = LOCTEXT("ViewModelInvalidGetterWithReason", "Viewmodel has an invalid Getter. {0}");
-
-				TValueOrError<TArray<FMVVMConstFieldVariant>, FText> GeneratedField = FieldPathHelper::GenerateFieldPathList(Class, ViewModelContext.ViewModelPropertyPath, true);
-				if (GeneratedField.HasError())
-				{
-					AddErrorForViewModel(ViewModelContext, FText::Format(InvalidGetterFormat, GeneratedField.GetError()));
-					bAreSourcesCreatorValid = false;
-					continue;
-				}
-
-				TValueOrError<FCompiledBindingLibraryCompiler::FFieldPathHandle, FText> ReadFieldPathResult = BindingLibraryCompiler.AddObjectFieldPath(GeneratedField.GetValue(), ViewModelContext.GetViewModelClass(), true);
+				TValueOrError<FCompiledBindingLibraryCompiler::FFieldPathHandle, FText> ReadFieldPathResult = AddObjectFieldPath(BindingLibraryCompiler, Class, ViewModelContext.ViewModelPropertyPath, ViewModelContext.GetViewModelClass());
 				if (ReadFieldPathResult.HasError())
 				{
-					AddErrorForViewModel(ViewModelContext, FText::Format(InvalidGetterFormat	, ReadFieldPathResult.GetError()));
+					AddErrorForViewModel(ViewModelContext, ReadFieldPathResult.GetError());
 					bAreSourcesCreatorValid = false;
 					continue;
 				}
@@ -979,19 +1034,32 @@ bool FMVVMViewBlueprintCompiler::CompileSourceCreators(const FCompiledBindingLib
 
 	for (const FCompilerSourceCreatorContext& SourceCreatorContext : CompilerSourceCreatorContexts)
 	{
+		const FMVVMBlueprintViewModelContext& ViewModelContext = SourceCreatorContext.ViewModelContext;
+		FMVVMViewClass_SourceCreator CompiledSourceCreator;
+
+		ensure(ViewModelContext.GetViewModelClass() && ViewModelContext.GetViewModelClass()->ImplementsInterface(UNotifyFieldValueChanged::StaticClass()));
+		CompiledSourceCreator.ExpectedSourceType = ViewModelContext.GetViewModelClass();
+		CompiledSourceCreator.PropertyName = ViewModelContext.GetViewModelName();
+
+		bool bCanBeSet = false;
+		bool CanBeEvaluated = false;
+		bool bIsOptional = false;
+		bool bCreateInstance = false;
+		bool IsUserWidgetProperty = false;
+
 		if (SourceCreatorContext.Type == ECompilerSourceCreatorType::ViewModel)
 		{
-			const FMVVMBlueprintViewModelContext& ViewModelContext = SourceCreatorContext.ViewModelContext;
-
-			FMVVMViewClass_SourceCreator CompiledSourceCreator;
+			IsUserWidgetProperty = true;
+			bCanBeSet = ViewModelContext.bCreateSetterFunction;
 
 			if (ViewModelContext.CreationType == EMVVMBlueprintViewModelContextCreationType::Manual)
 			{
-				CompiledSourceCreator = FMVVMViewClass_SourceCreator::MakeManual(ViewModelContext.GetViewModelName(), ViewModelContext.GetViewModelClass());
+				bCanBeSet = true;
+				bIsOptional = true;
 			}
 			else if (ViewModelContext.CreationType == EMVVMBlueprintViewModelContextCreationType::CreateInstance)
 			{
-				CompiledSourceCreator = FMVVMViewClass_SourceCreator::MakeInstance(ViewModelContext.GetViewModelName(), ViewModelContext.GetViewModelClass());
+				bCreateInstance = true;
 			}
 			else if (ViewModelContext.CreationType == EMVVMBlueprintViewModelContextCreationType::PropertyPath)
 			{
@@ -1002,7 +1070,9 @@ bool FMVVMViewBlueprintCompiler::CompileSourceCreators(const FCompiledBindingLib
 					bAreSourcesCreatorValid = false;
 					continue;
 				}
-				CompiledSourceCreator = FMVVMViewClass_SourceCreator::MakeFieldPath(ViewModelContext.GetViewModelName(), ViewModelContext.GetViewModelClass(), *CompiledFieldPath, ViewModelContext.bOptional);
+
+				CompiledSourceCreator.FieldPath = *CompiledFieldPath;
+				bIsOptional = ViewModelContext.bOptional;
 			}
 			else if (ViewModelContext.CreationType == EMVVMBlueprintViewModelContextCreationType::GlobalViewModelCollection)
 			{
@@ -1023,7 +1093,8 @@ bool FMVVMViewBlueprintCompiler::CompileSourceCreators(const FCompiledBindingLib
 					continue;
 				}
 
-				CompiledSourceCreator = FMVVMViewClass_SourceCreator::MakeGlobalContext(ViewModelContext.GetViewModelName(), MoveTemp(GlobalViewModelInstance), ViewModelContext.bOptional);
+				CompiledSourceCreator.GlobalViewModelInstance = MoveTemp(GlobalViewModelInstance);
+				bIsOptional = ViewModelContext.bOptional;
 			}
 			else if (ViewModelContext.CreationType == EMVVMBlueprintViewModelContextCreationType::Resolver)
 			{
@@ -1035,7 +1106,8 @@ bool FMVVMViewBlueprintCompiler::CompileSourceCreators(const FCompiledBindingLib
 					continue;
 				}
 
-				CompiledSourceCreator = FMVVMViewClass_SourceCreator::MakeResolver(ViewModelContext.GetViewModelName(), ViewModelContext.GetViewModelClass(), Resolver, ViewModelContext.bOptional);
+				CompiledSourceCreator.Resolver = Resolver;
+				bIsOptional = ViewModelContext.bOptional;
 			}
 			else
 			{
@@ -1043,9 +1115,29 @@ bool FMVVMViewBlueprintCompiler::CompileSourceCreators(const FCompiledBindingLib
 				bAreSourcesCreatorValid = false;
 				continue;
 			}
+		}
+		else if (SourceCreatorContext.Type == ECompilerSourceCreatorType::ViewModelDynamic)
+		{
+			const FMVVMVCompiledFieldPath* CompiledFieldPath = CompileResult.FieldPaths.Find(SourceCreatorContext.ReadPropertyPath);
+			if (CompiledFieldPath == nullptr)
+			{
+				AddErrorForViewModel(ViewModelContext, LOCTEXT("ViewModelInvalidInitializationBindingNotGenerated", "The viewmodel initialization binding was not generated."));
+				bAreSourcesCreatorValid = false;
+				continue;
+			}
 
-			ViewExtension->SourceCreators.Add(MoveTemp(CompiledSourceCreator));
+			CanBeEvaluated = true;
+			CompiledSourceCreator.FieldPath = *CompiledFieldPath;
 		}
+
+		CompiledSourceCreator.Flags = 0;
+		CompiledSourceCreator.Flags |= bCreateInstance ? (uint8)FMVVMViewClass_SourceCreator::ESourceFlags::TypeCreateInstance : 0;
+		CompiledSourceCreator.Flags |= IsUserWidgetProperty ? (uint8)FMVVMViewClass_SourceCreator::ESourceFlags::IsUserWidgetProperty : 0;
+		CompiledSourceCreator.Flags |= bIsOptional ? (uint8)FMVVMViewClass_SourceCreator::ESourceFlags::IsOptional : 0;
+		CompiledSourceCreator.Flags |= bCanBeSet ? (uint8)FMVVMViewClass_SourceCreator::ESourceFlags::CanBeSet : 0;
+		CompiledSourceCreator.Flags |= CanBeEvaluated ? (uint8)FMVVMViewClass_SourceCreator::ESourceFlags::CanBeEvaluated : 0;
+
+		ViewExtension->SourceCreators.Add(MoveTemp(CompiledSourceCreator));
 	}
 
 	return bAreSourcesCreatorValid;
@@ -1066,90 +1158,7 @@ bool FMVVMViewBlueprintCompiler::PreCompileBindings(UWidgetBlueprintGeneratedCla
 		FMVVMBlueprintViewBinding& Binding = *BindingPtr;
 
 		FMVVMViewBlueprintCompiler* Self = this;
-		auto AddBinding = [Self](UWidgetBlueprintGeneratedClass* Class, TArrayView<const UE::MVVM::FMVVMConstFieldVariant> GetterFields, TArrayView<const UE::MVVM::FMVVMConstFieldVariant> SetterFields, const UFunction* ConversionFunction) -> TValueOrError<FCompilerBinding, FText>
-		{
-			FCompilerBinding Result;
-
-			if (GetterFields.Num() > 0)
-			{
-				static const FText CouldNotCreateSourceFieldPathFormat = LOCTEXT("CouldNotCreateSourceFieldPath", "Couldn't create the source field path '{0}'. {1}");
-
-
-				// Generate a path to read the value at runtime
-				TValueOrError<TArray<FMVVMConstFieldVariant>, FText> GeneratedField = FieldPathHelper::GenerateFieldPathList(GetterFields, true);
-				if (GeneratedField.HasError())
-				{
-					return MakeError(FText::Format(CouldNotCreateSourceFieldPathFormat
-						, ::UE::MVVM::FieldPathHelper::ToText(GetterFields)
-						, GeneratedField.GetError()));
-				}
-
-				TValueOrError<FCompiledBindingLibraryCompiler::FFieldPathHandle, FText> FieldPathResult = Self->BindingLibraryCompiler.AddFieldPath(GeneratedField.GetValue(), true);
-				if (FieldPathResult.HasError())
-				{
-					return MakeError(FText::Format(CouldNotCreateSourceFieldPathFormat
-						, ::UE::MVVM::FieldPathHelper::ToText(GetterFields)
-						, FieldPathResult.GetError()));
-				}
-				Result.SourceRead = FieldPathResult.StealValue();
-			}
-
-			static const FText CouldNotCreateDestinationFieldPathFormat = LOCTEXT("CouldNotCreateDestinationFieldPath", "Couldn't create the destination field path '{0}'. {1}");
-
-			{
-				TValueOrError<TArray<FMVVMConstFieldVariant>, FText> GeneratedField = FieldPathHelper::GenerateFieldPathList(SetterFields, false);
-				if (GeneratedField.HasError())
-				{
-					return MakeError(FText::Format(CouldNotCreateDestinationFieldPathFormat
-						, ::UE::MVVM::FieldPathHelper::ToText(SetterFields)
-						, GeneratedField.GetError()));
-				}
-
-				TValueOrError<FCompiledBindingLibraryCompiler::FFieldPathHandle, FText> FieldPathResult = Self->BindingLibraryCompiler.AddFieldPath(GeneratedField.GetValue(), false);
-				if (FieldPathResult.HasError())
-				{
-					return MakeError(FText::Format(CouldNotCreateDestinationFieldPathFormat
-						, ::UE::MVVM::FieldPathHelper::ToText(SetterFields)
-						, FieldPathResult.GetError()));
-				}
-				Result.DestinationWrite = FieldPathResult.StealValue();
-			}
-
-			if (ConversionFunction != nullptr)
-			{
-				TValueOrError<FCompiledBindingLibraryCompiler::FFieldPathHandle, FText> FieldPathResult = Self->BindingLibraryCompiler.AddConversionFunctionFieldPath(Class, ConversionFunction);
-				if (FieldPathResult.HasError())
-				{
-					return MakeError(FText::Format(LOCTEXT("CouldNotCreateConversionFunctionFieldPath", "Couldn't create the conversion function field path '{0}'. {1}")
-						, FText::FromString(ConversionFunction->GetPathName())
-						, FieldPathResult.GetError()));
-				}
-				Result.ConversionFunction = FieldPathResult.StealValue();
-				Result.bIsConversionFunctionComplex = BindingHelper::IsValidForComplexRuntimeConversion(ConversionFunction);
-			}
-
-			// Sanity check.
-			if (GetterFields.Num() == 0 && (ConversionFunction == nullptr || !BindingHelper::IsValidForComplexRuntimeConversion(ConversionFunction)))
-			{
-				return MakeError(LOCTEXT("BindingShouldHaveGetter", "The binding should have a getter."));
-			}
-			if (GetterFields.Num() != 0 && ConversionFunction == nullptr && BindingHelper::IsValidForComplexRuntimeConversion(ConversionFunction))
-			{
-				return MakeError(LOCTEXT("BindingHasConversionFunctionAndGetter", "The binding has both a complex conversion function and a getter."));
-			}
-
-			// Generate the binding
-			TValueOrError<FCompiledBindingLibraryCompiler::FBindingHandle, FText> BindingResult = Self->BindingLibraryCompiler.AddBinding(Result.SourceRead, Result.DestinationWrite, Result.ConversionFunction);
-			if (BindingResult.HasError())
-			{
-				return MakeError(BindingResult.StealError());
-			}
-			Result.BindingHandle = BindingResult.StealValue();
-
-			return MakeValue(Result);
-		};
-
-		auto AddFieldId = [Self](UClass* SourceContextClass, bool bNotifyFieldValueChangedRequired, EMVVMBindingMode BindingMode, FName FieldToListenTo) -> TValueOrError<FCompiledBindingLibraryCompiler::FFieldIdHandle, FText>
+		auto AddFieldId = [Self](const UClass* SourceContextClass, bool bNotifyFieldValueChangedRequired, EMVVMBindingMode BindingMode, FName FieldToListenTo) -> TValueOrError<FCompiledBindingLibraryCompiler::FFieldIdHandle, FText>
 		{
 			if (!IsOneTimeBinding(BindingMode) && bNotifyFieldValueChangedRequired)
 			{
@@ -1208,7 +1217,7 @@ bool FMVVMViewBlueprintCompiler::PreCompileBindings(UWidgetBlueprintGeneratedCla
 			continue;
 		}
 
-		TValueOrError<FCompilerBinding, FText> AddBindingResult = AddBinding(Class, BindingSourceContext.PropertyPath, SetterPath, ConversionFunction);
+		TValueOrError<FCompiledBinding, FText> AddBindingResult = CreateCompiledBinding(Class, BindingSourceContext.PropertyPath, SetterPath, ConversionFunction, BindingSourceContext.bIsComplexBinding);
 		if (AddBindingResult.HasError())
 		{
 			AddMessageForBinding(Binding, BlueprintView,
@@ -1219,13 +1228,17 @@ bool FMVVMViewBlueprintCompiler::PreCompileBindings(UWidgetBlueprintGeneratedCla
 			continue;
 		}
 
-		FCompilerBinding NewBinding = AddBindingResult.StealValue();
+		FCompilerBinding NewBinding;
 		NewBinding.BindingIndex = BindingSourceContext.BindingIndex;
 		NewBinding.UserWidgetPropertyContextIndex = BindingSourceContext.UserWidgetPropertyContextIndex;
 		NewBinding.SourceCreatorContextIndex = BindingSourceContext.SourceCreatorContextIndex;
-		NewBinding.FieldIdHandle = AddFieldResult.StealValue();
 		NewBinding.bSourceIsUserWidget = BindingSourceContext.bIsRootWidget;
 		NewBinding.bFieldIdNeeded = !IsOneTimeBinding(Binding.BindingType);
+		NewBinding.bIsForwardBinding = BindingSourceContext.bIsForwardBinding;
+
+		NewBinding.CompiledBinding = AddBindingResult.StealValue();
+		NewBinding.FieldIdHandle = AddFieldResult.StealValue();
+
 		CompilerBindings.Emplace(NewBinding);
 	}
 
@@ -1249,7 +1262,12 @@ bool FMVVMViewBlueprintCompiler::CompileBindings(const FCompiledBindingLibraryCo
 
 	for (const FCompilerBinding& CompileBinding : CompilerBindings)
 	{
-		FMVVMBlueprintViewBinding& ViewBinding = *BlueprintView->GetBindingAt(CompileBinding.BindingIndex);
+		// PropertyBinding needs a valid CompileBinding.BindingIndex.
+		check(CompileBinding.Type != ECompilerBindingType::PropertyBinding || (CompileBinding.BindingIndex >= 0 && CompileBinding.BindingIndex < BlueprintView->GetNumBindings()));
+		FMVVMBlueprintViewBinding* ViewBinding = CompileBinding.Type == ECompilerBindingType::PropertyBinding
+			? BlueprintView->GetBindingAt(CompileBinding.BindingIndex)
+			: nullptr;
+
 
 		FMVVMViewClass_CompiledBinding NewBinding;
 		const bool bIsSourceSelf = CompileBinding.bSourceIsUserWidget;
@@ -1270,21 +1288,30 @@ bool FMVVMViewBlueprintCompiler::CompileBindings(const FCompiledBindingLibraryCo
 			}
 		}
 
+		auto AddErroMessage = [this, ViewBinding, BlueprintView](const FText& ErrorMessage)
+		{
+				if (ViewBinding)
+				{
+					AddMessageForBinding(*ViewBinding, BlueprintView, ErrorMessage, EBindingMessageType::Error);
+				}
+				else
+				{
+					WidgetBlueprintCompilerContext.MessageLog.Error(*ErrorMessage.ToString());
+				}
+		};
+
 		const FMVVMVCompiledFieldId* CompiledFieldId = CompileResult.FieldIds.Find(CompileBinding.FieldIdHandle);
 		if (CompiledFieldId == nullptr && CompileBinding.bFieldIdNeeded)
 		{
-			AddMessageForBinding(ViewBinding, BlueprintView,
-				FText::Format(LOCTEXT("FieldIdNotGenerated", "Could not generate field ID for property '{0}'."), FText::FromName(NewBinding.SourcePropertyName)),
-				EBindingMessageType::Error
-			);
+			AddErroMessage(FText::Format(LOCTEXT("FieldIdNotGenerated", "Could not generate field ID for property '{0}'."), FText::FromName(NewBinding.SourcePropertyName)));
 			bIsBindingsValid = false;
 			continue;
 		}
 
-		const FMVVMVCompiledBinding* CompiledBinding = CompileResult.Bindings.Find(CompileBinding.BindingHandle);
-		if (CompiledBinding == nullptr)
+		const FMVVMVCompiledBinding* CompiledBinding = CompileResult.Bindings.Find(CompileBinding.CompiledBinding.BindingHandle);
+		if (CompiledBinding == nullptr && CompileBinding.Type != ECompilerBindingType::ViewModelDynamic)
 		{
-			AddMessageForBinding(ViewBinding, BlueprintView, LOCTEXT("CompiledBindingNotGenerated", "Could not generate compiled binding."), EBindingMessageType::Error);
+			AddErroMessage(LOCTEXT("CompiledBindingNotGenerated", "Could not generate compiled binding."));
 			bIsBindingsValid = false;
 			continue;
 		}
@@ -1292,44 +1319,89 @@ bool FMVVMViewBlueprintCompiler::CompileBindings(const FCompiledBindingLibraryCo
 		bool bIsOptional = false;
 		if (!bIsSourceSelf)
 		{
-			FGuid ViewModelId;
+			const FMVVMBlueprintViewModelContext* ViewModelContext = nullptr;
 			if (CompilerUserWidgetPropertyContexts.IsValidIndex(CompileBinding.UserWidgetPropertyContextIndex))
 			{
-				ViewModelId = CompilerUserWidgetPropertyContexts[CompileBinding.UserWidgetPropertyContextIndex].ViewModelId;
+				if (CompilerUserWidgetPropertyContexts[CompileBinding.UserWidgetPropertyContextIndex].ViewModelId.IsValid())
+				{
+					ViewModelContext = BlueprintView->FindViewModel(CompilerUserWidgetPropertyContexts[CompileBinding.UserWidgetPropertyContextIndex].ViewModelId);
+					if (ViewModelContext == nullptr)
+					{
+						AddErroMessage(LOCTEXT("CompiledBindingWithInvalidIVewModelId", "Internal error: the viewmodel became invalid."));
+						bIsBindingsValid = false;
+						continue;
+					}
+				}
 			}
 			else if (CompilerSourceCreatorContexts.IsValidIndex(CompileBinding.SourceCreatorContextIndex))
 			{
-				ViewModelId = CompilerSourceCreatorContexts[CompileBinding.SourceCreatorContextIndex].ViewModelContext.GetViewModelId();
+				ViewModelContext = &CompilerSourceCreatorContexts[CompileBinding.SourceCreatorContextIndex].ViewModelContext;
 			}
 
-			if (ViewModelId.IsValid())
+			if (ViewModelContext)
 			{
-				const FMVVMBlueprintViewModelContext* ModelContext = BlueprintView->FindViewModel(ViewModelId);
-				if (ModelContext == nullptr)
+				bIsOptional = ViewModelContext->bOptional || ViewModelContext->CreationType == EMVVMBlueprintViewModelContextCreationType::Manual;
+			}
+		}
+
+		if (CompileBinding.Type == ECompilerBindingType::ViewModelDynamic)
+		{
+			int32 FoundSourceCreatorIndex = ViewExtension->SourceCreators.IndexOfByPredicate([LookFor = CompileBinding.DynamicViewModelName](const FMVVMViewClass_SourceCreator& Other)
 				{
-					AddMessageForBinding(ViewBinding, BlueprintView, LOCTEXT("CompiledBindingWithInvalidIVewModelId", "Internal error: the view model became invalid."), EBindingMessageType::Error);
-					bIsBindingsValid = false;
-					continue;
-				}
+					return Other.GetSourceName() == LookFor;
+				});
 
-				bIsOptional = ModelContext->bOptional || ModelContext->CreationType == EMVVMBlueprintViewModelContextCreationType::Manual;
+			check(FoundSourceCreatorIndex < std::numeric_limits<int8>::max()); // the index is saved as a int8
+			NewBinding.EvaluateSourceCreatorIndex = FoundSourceCreatorIndex;
+			if (NewBinding.EvaluateSourceCreatorIndex == INDEX_NONE)
+			{
+				FText ErrorMessage = FText::Format(LOCTEXT("CompiledViewModelDynamicBindingNotGenerated", "Was not able to find the source for {0}."), FText::FromName(CompileBinding.DynamicViewModelName));
+				WidgetBlueprintCompilerContext.MessageLog.Error(*ErrorMessage.ToString());
+				bIsBindingsValid = false;
+				continue;
 			}
 		}
 
-		NewBinding.FieldId = CompiledFieldId  ? *CompiledFieldId : FMVVMVCompiledFieldId();
-		NewBinding.Binding = *CompiledBinding;
-		NewBinding.ExecutionMode = ViewBinding.bOverrideExecutionMode ? ViewBinding.OverrideExecutionMode : (EMVVMExecutionMode)CVarDefaultExecutionMode->GetInt();
-		NewBinding.EditorId = ViewBinding.BindingId;
-
+		FText WrongBindingTypeInternalErrorMessage = LOCTEXT("WrongBindingTypeInternalError", "Internal Error. Wrong binding type.");
+		// FieldId can be null if it's a one time binding
+		NewBinding.FieldId = CompiledFieldId ? *CompiledFieldId : FMVVMVCompiledFieldId();
+		NewBinding.Binding = CompiledBinding ? *CompiledBinding : FMVVMVCompiledBinding();
 		NewBinding.Flags = 0;
-		NewBinding.Flags |= (ViewBinding.bEnabled) ? FMVVMViewClass_CompiledBinding::EBindingFlags::EnabledByDefault : 0;
-		NewBinding.Flags |= (IsForwardBinding(ViewBinding.BindingType)) ? FMVVMViewClass_CompiledBinding::EBindingFlags::ForwardBinding : 0;
-		//NewBinding.Flags |= (ViewBinding.BindingType == EMVVMBindingMode::TwoWay) ? FMVVMViewClass_CompiledBinding::EBindingFlags::TwoWayBinding : 0;
-		NewBinding.Flags |= (IsOneTimeBinding(ViewBinding.BindingType)) ? FMVVMViewClass_CompiledBinding::EBindingFlags::OneTime : 0;
-		NewBinding.Flags |= (bIsOptional) ? FMVVMViewClass_CompiledBinding::EBindingFlags::ViewModelOptional : 0;
-		NewBinding.Flags |= (CompileBinding.bIsConversionFunctionComplex) ? FMVVMViewClass_CompiledBinding::EBindingFlags::ConversionFunctionIsComplex : 0;
-		NewBinding.Flags |= (ViewBinding.bOverrideExecutionMode) ? FMVVMViewClass_CompiledBinding::EBindingFlags::OverrideExecuteMode : 0;
-		NewBinding.Flags |= (bIsSourceSelf) ? FMVVMViewClass_CompiledBinding::EBindingFlags::SourceObjectIsSelf : 0;
+		if (ViewBinding)
+		{
+			if (CompileBinding.Type != ECompilerBindingType::PropertyBinding)
+			{
+				WidgetBlueprintCompilerContext.MessageLog.Error(*WrongBindingTypeInternalErrorMessage.ToString());
+				bIsBindingsValid = false;
+				continue;
+			}
+
+			NewBinding.ExecutionMode = ViewBinding->bOverrideExecutionMode ? ViewBinding->OverrideExecutionMode : (EMVVMExecutionMode)CVarDefaultExecutionMode->GetInt();;
+			NewBinding.EditorId = ViewBinding->BindingId;
+
+			NewBinding.Flags |= (ViewBinding->bEnabled) ? FMVVMViewClass_CompiledBinding::EBindingFlags::EnabledByDefault : 0;
+			NewBinding.Flags |= (CompileBinding.bIsForwardBinding) ? FMVVMViewClass_CompiledBinding::EBindingFlags::ExecuteAtInitialization : 0;
+			NewBinding.Flags |= (IsOneTimeBinding(ViewBinding->BindingType)) ? FMVVMViewClass_CompiledBinding::EBindingFlags::OneTime : 0;
+			NewBinding.Flags |= (bIsOptional) ? FMVVMViewClass_CompiledBinding::EBindingFlags::ViewModelOptional : 0;
+			NewBinding.Flags |= (ViewBinding->bOverrideExecutionMode) ? FMVVMViewClass_CompiledBinding::EBindingFlags::OverrideExecuteMode : 0;
+			NewBinding.Flags |= (bIsSourceSelf) ? FMVVMViewClass_CompiledBinding::EBindingFlags::SourceObjectIsSelf : 0;
+		}
+		else
+		{
+			if (CompileBinding.Type != ECompilerBindingType::ViewModelDynamic)
+			{
+				WidgetBlueprintCompilerContext.MessageLog.Error(*WrongBindingTypeInternalErrorMessage.ToString());
+				bIsBindingsValid = false;
+				continue;
+			}
+			NewBinding.ExecutionMode = EMVVMExecutionMode::Immediate;
+			NewBinding.EditorId = FGuid();
+
+			NewBinding.Flags |= FMVVMViewClass_CompiledBinding::EBindingFlags::EnabledByDefault;
+			NewBinding.Flags |= (bIsOptional) ? FMVVMViewClass_CompiledBinding::EBindingFlags::ViewModelOptional : 0;
+			NewBinding.Flags |= FMVVMViewClass_CompiledBinding::EBindingFlags::OverrideExecuteMode; // The mode needs to be Immediate.
+			NewBinding.Flags |= (bIsSourceSelf) ? FMVVMViewClass_CompiledBinding::EBindingFlags::SourceObjectIsSelf : 0;
+		}
 
 		ViewExtension->CompiledBindings.Emplace(MoveTemp(NewBinding));
 	}
@@ -1338,6 +1410,87 @@ bool FMVVMViewBlueprintCompiler::CompileBindings(const FCompiledBindingLibraryCo
 }
 
 
+TValueOrError<FMVVMViewBlueprintCompiler::FCompiledBinding, FText> FMVVMViewBlueprintCompiler::CreateCompiledBinding(const UWidgetBlueprintGeneratedClass* Class, TArrayView<const UE::MVVM::FMVVMConstFieldVariant> GetterFields, TArrayView<const UE::MVVM::FMVVMConstFieldVariant> SetterFields, const UFunction* ConversionFunction, bool bIsComplexBinding)
+{
+	FCompiledBinding Result;
+
+	if (ConversionFunction != nullptr)
+	{
+		TValueOrError<FCompiledBindingLibraryCompiler::FFieldPathHandle, FText> FieldPathResult = BindingLibraryCompiler.AddConversionFunctionFieldPath(Class, ConversionFunction);
+		if (FieldPathResult.HasError())
+		{
+			return MakeError(FText::Format(LOCTEXT("CouldNotCreateConversionFunctionFieldPath", "Couldn't create the conversion function field path '{0}'. {1}")
+				, FText::FromString(ConversionFunction->GetPathName())
+				, FieldPathResult.GetError()));
+		}
+		Result.ConversionFunction = FieldPathResult.StealValue();
+		Result.bIsConversionFunctionComplex = bIsComplexBinding;
+
+		// Sanity check
+		if (bIsComplexBinding && !BindingHelper::IsValidForComplexRuntimeConversion(ConversionFunction))
+		{
+			return MakeError(LOCTEXT("ConversionFunctionIsNotComplex", "Internal Error. The complex conversion function does not respect the prerequist."));
+		}
+	}
+
+	if (!Result.bIsConversionFunctionComplex)
+	{
+		static const FText CouldNotCreateSourceFieldPathFormat = LOCTEXT("CouldNotCreateSourceFieldPath", "Couldn't create the source field path '{0}'. {1}");
+
+		// Generate a path to read the value at runtime
+		TValueOrError<TArray<FMVVMConstFieldVariant>, FText> GeneratedField = FieldPathHelper::GenerateFieldPathList(GetterFields, true);
+		if (GeneratedField.HasError())
+		{
+			return MakeError(FText::Format(CouldNotCreateSourceFieldPathFormat
+				, ::UE::MVVM::FieldPathHelper::ToText(GetterFields)
+				, GeneratedField.GetError()));
+		}
+
+		TValueOrError<FCompiledBindingLibraryCompiler::FFieldPathHandle, FText> FieldPathResult = BindingLibraryCompiler.AddFieldPath(GeneratedField.GetValue(), true);
+		if (FieldPathResult.HasError())
+		{
+			return MakeError(FText::Format(CouldNotCreateSourceFieldPathFormat
+				, ::UE::MVVM::FieldPathHelper::ToText(GetterFields)
+				, FieldPathResult.GetError()));
+		}
+		Result.SourceRead = FieldPathResult.StealValue();
+	}
+
+	{
+		static const FText CouldNotCreateDestinationFieldPathFormat = LOCTEXT("CouldNotCreateDestinationFieldPath", "Couldn't create the destination field path '{0}'. {1}");
+
+		TValueOrError<TArray<FMVVMConstFieldVariant>, FText> GeneratedField = FieldPathHelper::GenerateFieldPathList(SetterFields, false);
+		if (GeneratedField.HasError())
+		{
+			return MakeError(FText::Format(CouldNotCreateDestinationFieldPathFormat
+				, ::UE::MVVM::FieldPathHelper::ToText(SetterFields)
+				, GeneratedField.GetError()));
+		}
+
+		TValueOrError<FCompiledBindingLibraryCompiler::FFieldPathHandle, FText> FieldPathResult = BindingLibraryCompiler.AddFieldPath(GeneratedField.GetValue(), false);
+		if (FieldPathResult.HasError())
+		{
+			return MakeError(FText::Format(CouldNotCreateDestinationFieldPathFormat
+				, ::UE::MVVM::FieldPathHelper::ToText(SetterFields)
+				, FieldPathResult.GetError()));
+		}
+		Result.DestinationWrite = FieldPathResult.StealValue();
+	}
+
+	// Generate the binding
+	TValueOrError<FCompiledBindingLibraryCompiler::FBindingHandle, FText> BindingResult = Result.bIsConversionFunctionComplex
+		? BindingLibraryCompiler.AddComplexBinding(Result.DestinationWrite, Result.ConversionFunction)
+		: BindingLibraryCompiler.AddBinding(Result.SourceRead, Result.DestinationWrite, Result.ConversionFunction);
+	if (BindingResult.HasError())
+	{
+		return MakeError(BindingResult.StealError());
+	}
+	Result.BindingHandle = BindingResult.StealValue();
+
+	return MakeValue(Result);
+}
+
+
 const FMVVMViewBlueprintCompiler::FCompilerSourceCreatorContext* FMVVMViewBlueprintCompiler::FindViewModelSource(FGuid Id) const
 {
 	return CompilerSourceCreatorContexts.FindByPredicate([Id](const FCompilerSourceCreatorContext& Other)
@@ -1347,7 +1500,7 @@ const FMVVMViewBlueprintCompiler::FCompilerSourceCreatorContext* FMVVMViewBluepr
 }
 
 
-TValueOrError<FMVVMViewBlueprintCompiler::FBindingSourceContext, FText> FMVVMViewBlueprintCompiler::CreateBindingSourceContext(const UMVVMBlueprintView* BlueprintView, const UWidgetBlueprintGeneratedClass* Class, const FMVVMBlueprintPropertyPath& PropertyPath)
+TValueOrError<FMVVMViewBlueprintCompiler::FBindingSourceContext, FText> FMVVMViewBlueprintCompiler::CreateBindingSourceContext(const UMVVMBlueprintView* BlueprintView, const UWidgetBlueprintGeneratedClass* Class, const FMVVMBlueprintPropertyPath& PropertyPath, bool bIsOneTimeBinding)
 {
 	if (PropertyPath.IsEmpty())
 	{
@@ -1393,63 +1546,135 @@ TValueOrError<FMVVMViewBlueprintCompiler::FBindingSourceContext, FText> FMVVMVie
 	}
 
 	// The path may contains another INotifyFieldValueChanged
-	TValueOrError<FieldPathHelper::FParsedBindingInfo, FText> BindingInfoResult = FieldPathHelper::GetBindingInfoFromFieldPath(Class, Result.PropertyPath);
+	TValueOrError<FieldPathHelper::FParsedNotifyBindingInfo, FText> BindingInfoResult = FieldPathHelper::GetNotifyBindingInfoFromFieldPath(Class, Result.PropertyPath);
 	if (BindingInfoResult.HasError())
 	{
 		return MakeError(BindingInfoResult.StealError());
 	}
-
-	const FieldPathHelper::FParsedBindingInfo& BindingInfo = BindingInfoResult.GetValue();
+	const FieldPathHelper::FParsedNotifyBindingInfo& BindingInfo = BindingInfoResult.GetValue();
 	Result.FieldId = BindingInfo.NotifyFieldId;
 
-	const bool bCreateADynamicSourceEntry = BindingInfo.ViewModelIndex > 1
-		&& BindingInfo.NotifyFieldClass != nullptr;
-
-	// The INotifyFieldValueChanged/viewmodel is not the first property. Creates the ViaBinding source entries.
-	if (bCreateADynamicSourceEntry)
+	if (BindingInfo.ViewModelIndex < 1 && BindingInfo.NotifyFieldClass && Result.SourceClass)
 	{
-		return MakeError(LOCTEXT("DynamicSourceEntryNotSupport", "Long source entry is not yet fully supported. Add the viewmodel manually."));
+		if (!Result.SourceClass->IsChildOf(BindingInfo.NotifyFieldClass))
+		{
+			return MakeError(LOCTEXT("InvalidNotifyFieldClassInternal", "Internal error. The viewmodel class doesn't matches."));
+		}
+	}
 
-		if (!Result.PropertyPath.IsValidIndex(BindingInfo.ViewModelIndex))
+	// The INotifyFieldValueChanged/viewmodel is not the first and only INotifyFieldValueChanged/viewmodel property path.
+	//Create a new source in PropertyPath creator mode. Create a special binding to update the viewmodel when it changes.
+	//This binding (calling this function) will use the new source.
+	if (BindingInfo.ViewModelIndex >= 1 && !bIsOneTimeBinding)
+	{
+		if (!GetDefault<UMVVMDeveloperProjectSettings>()->bAllowLongSourcePath)
 		{
-			return MakeError(LOCTEXT("DynamicSourceEntryInternalIndex", "Internal error. The source index is not valid."));
+			return MakeError(LOCTEXT("DynamicSourceEntryNotSupport", "Long source entry is not supported. Add the viewmodel manually."));
 		}
-		//for (0 to index) if they are all interface, then create new BindingCompilerSourceCreatorContexts
-		//create new dynamic vm with property path to access it
-		TStringBuilder<512> PropertyPathBuilder;
-		for (int32 Index = 0; Index < BindingInfo.ViewModelIndex; ++Index)
+
+		int32 SourceCreatorContextIndex = INDEX_NONE;
+		for (int32 DynamicIndex = 1; DynamicIndex <= BindingInfo.ViewModelIndex; ++DynamicIndex)
 		{
-			if (Index > 0)
+			if (!Result.PropertyPath.IsValidIndex(DynamicIndex))
 			{
-				PropertyPathBuilder << TEXT('.');
+				return MakeError(LOCTEXT("DynamicSourceEntryInternalIndex", "Internal error. The source index is not valid."));
 			}
-			PropertyPathBuilder << Result.PropertyPath[Index].GetName();
-		}
 
-		int32 SourceCreatorContextIndex = CompilerSourceCreatorContexts.IndexOfByPredicate([&PropertyPathBuilder](const FCompilerSourceCreatorContext& Other)
+			FName NewSourceName;
+			FString NewSourcePropertyPath;
 			{
-				return Other.Type == ECompilerSourceCreatorType::ViewModelDynamic
-					&& Other.ViewModelContext.ViewModelPropertyPath == PropertyPathBuilder.ToString();
-			});
-
-		if (SourceCreatorContextIndex == INDEX_NONE)
-		{
-			FCompilerSourceCreatorContext SourceCreatorContext;
-			SourceCreatorContext.Type = ECompilerSourceCreatorType::ViewModelDynamic;
-			SourceCreatorContext.ViewModelContext = FMVVMBlueprintViewModelContext(BindingInfo.NotifyFieldClass, *FString::Printf(TEXT("__Dynamic__%d"), ViewModelDynamicCounter));
-			SourceCreatorContext.ViewModelContext.bCreateSetterFunction = false;
-			SourceCreatorContext.ViewModelContext.bOptional = false;
-			SourceCreatorContext.ViewModelContext.CreationType = EMVVMBlueprintViewModelContextCreationType::PropertyPath;
-			SourceCreatorContext.ViewModelContext.ViewModelPropertyPath = PropertyPathBuilder.ToString();
-			++ViewModelDynamicCounter;
-
-			SourceCreatorContextIndex = CompilerSourceCreatorContexts.Add(MoveTemp(SourceCreatorContext));
+				TStringBuilder<512> PropertyPathBuilder;
+				TStringBuilder<512> DynamicNameBuilder;
+				for (int32 Index = 0; Index <= DynamicIndex; ++Index)
+				{
+					if (Index > 0)
+					{
+						PropertyPathBuilder << TEXT('.');
+						DynamicNameBuilder << TEXT('_');
+					}
+					PropertyPathBuilder << Result.PropertyPath[Index].GetName();
+					DynamicNameBuilder << Result.PropertyPath[Index].GetName();
+				}
+
+				NewSourceName = FName(DynamicNameBuilder.ToString());
+				NewSourcePropertyPath = PropertyPathBuilder.ToString();
+			}
+
+			// Did we already create the new source?
+			int32 PreviousSourceCreatorContextIndex = SourceCreatorContextIndex;
+			SourceCreatorContextIndex = CompilerSourceCreatorContexts.IndexOfByPredicate([NewSourceName](const FCompilerSourceCreatorContext& Other)
+				{
+					return Other.Type == ECompilerSourceCreatorType::ViewModelDynamic
+						&& Other.ViewModelContext.GetViewModelName() == NewSourceName;
+				});
+			if (SourceCreatorContextIndex == INDEX_NONE)
+			{
+				// Create the new source
+				{
+					const UClass* ViewModelClass = Cast<const UClass>(Result.PropertyPath[DynamicIndex+1].GetOwner());
+					if (!ViewModelClass)
+					{
+						return MakeError(FText::GetEmpty());
+					}
+
+					FCompilerSourceCreatorContext SourceCreatorContext;
+					SourceCreatorContext.Type = ECompilerSourceCreatorType::ViewModelDynamic;
+					SourceCreatorContext.ViewModelContext = FMVVMBlueprintViewModelContext(ViewModelClass, NewSourceName);
+					SourceCreatorContext.ViewModelContext.bCreateSetterFunction = false;
+					SourceCreatorContext.ViewModelContext.bOptional = false;
+					SourceCreatorContext.ViewModelContext.CreationType = EMVVMBlueprintViewModelContextCreationType::PropertyPath;
+					SourceCreatorContext.ViewModelContext.ViewModelPropertyPath = NewSourcePropertyPath;
+
+					TValueOrError<FCompiledBindingLibraryCompiler::FFieldPathHandle, FText> ReadFieldPathResult = AddObjectFieldPath(BindingLibraryCompiler, Class, SourceCreatorContext.ViewModelContext.ViewModelPropertyPath, SourceCreatorContext.ViewModelContext.GetViewModelClass());
+					if (ReadFieldPathResult.HasError())
+					{
+						return MakeError(FText::Format(LOCTEXT("DynamicSourceEntryInvalidPath", "Internal error. {0}."), ReadFieldPathResult.StealError()));
+					}
+					SourceCreatorContext.ReadPropertyPath = ReadFieldPathResult.StealValue();
+
+					SourceCreatorContextIndex = CompilerSourceCreatorContexts.Add(MoveTemp(SourceCreatorContext));
+				}
+
+				// Create the binding to update the source when it changes.
+				{
+					FCompilerBinding NewCompilerBinding;
+					NewCompilerBinding.Type = ECompilerBindingType::ViewModelDynamic;
+					NewCompilerBinding.bSourceIsUserWidget = false;
+					NewCompilerBinding.bIsForwardBinding = true;
+					NewCompilerBinding.bFieldIdNeeded = true;
+					NewCompilerBinding.DynamicViewModelName = NewSourceName;
+
+					if (PreviousSourceCreatorContextIndex == INDEX_NONE)
+					{
+						NewCompilerBinding.UserWidgetPropertyContextIndex = Result.UserWidgetPropertyContextIndex;
+					}
+					else
+					{
+						NewCompilerBinding.SourceCreatorContextIndex = PreviousSourceCreatorContextIndex;
+					}
+
+					{
+						const UClass* OwnerClass = Cast<const UClass>(Result.PropertyPath[DynamicIndex].GetOwner());
+						if (!OwnerClass)
+						{
+							return MakeError(FText::GetEmpty());
+						}
+
+						TValueOrError<FCompiledBindingLibraryCompiler::FFieldIdHandle, FText> AddFieldResult = BindingLibraryCompiler.AddFieldId(OwnerClass, Result.PropertyPath[DynamicIndex].GetName());
+						if (AddFieldResult.HasError())
+						{
+							return MakeError(AddFieldResult.StealError());
+						}
+						NewCompilerBinding.FieldIdHandle = AddFieldResult.StealValue();
+					}
+
+					CompilerBindings.Add(MoveTemp(NewCompilerBinding));
+				}
+			}
 		}
 
 		Result.SourceClass = BindingInfo.NotifyFieldClass;
 		Result.FieldId = BindingInfo.NotifyFieldId;
-		int32 NumberToRemove = Result.PropertyPath.Num() - BindingInfo.ViewModelIndex;
-		Result.PropertyPath.RemoveAt(BindingInfo.ViewModelIndex, NumberToRemove);
 		Result.UserWidgetPropertyContextIndex = INDEX_NONE;
 		Result.SourceCreatorContextIndex = SourceCreatorContextIndex;
 		Result.bIsRootWidget = false;
diff --git a/Engine/Plugins/Runtime/ModelViewViewModel/Source/ModelViewViewModelBlueprint/Private/MVVMViewBlueprintCompiler.h b/Engine/Plugins/Runtime/ModelViewViewModel/Source/ModelViewViewModelBlueprint/Private/MVVMViewBlueprintCompiler.h
--- a/Engine/Plugins/Runtime/ModelViewViewModel/Source/ModelViewViewModelBlueprint/Private/MVVMViewBlueprintCompiler.h
+++ b/Engine/Plugins/Runtime/ModelViewViewModel/Source/ModelViewViewModelBlueprint/Private/MVVMViewBlueprintCompiler.h
@@ -29,6 +29,7 @@ struct FMVVMViewBlueprintCompiler
 private:
 	struct FCompilerUserWidgetPropertyContext;
 	struct FCompilerSourceCreatorContext;
+	struct FCompiledBinding;
 	struct FCompilerBinding;
 	struct FBindingSourceContext;
 
@@ -76,8 +77,10 @@ private:
 	void AddMessageForBinding(FMVVMBlueprintViewBinding& Binding, UMVVMBlueprintView* BlueprintView, const FText& MessageText, EBindingMessageType MessageType, FName ArgumentName = FName()) const;
 	void AddErrorForViewModel(const FMVVMBlueprintViewModelContext& ViewModel, const FText& Message) const;
 
-	TValueOrError<FBindingSourceContext, FText> CreateBindingSourceContext(const UMVVMBlueprintView* BlueprintView, const UWidgetBlueprintGeneratedClass* Class, const FMVVMBlueprintPropertyPath& PropertyPath);
+	TValueOrError<FBindingSourceContext, FText> CreateBindingSourceContext(const UMVVMBlueprintView* BlueprintView, const UWidgetBlueprintGeneratedClass* Class, const FMVVMBlueprintPropertyPath& PropertyPath, bool bIsOneTimeBinding);
 	TArray<FMVVMConstFieldVariant> CreateBindingDestinationPath(const UMVVMBlueprintView* BlueprintView, const UWidgetBlueprintGeneratedClass* Class, const FMVVMBlueprintPropertyPath& PropertyPath) const;
+	TValueOrError<FCompiledBinding, FText> CreateCompiledBinding(const UWidgetBlueprintGeneratedClass* Class, TArrayView<const UE::MVVM::FMVVMConstFieldVariant> GetterFields, TArrayView<const UE::MVVM::FMVVMConstFieldVariant> SetterFields, const UFunction* ConversionFunction, bool bIsComplexBinding);
+
 
 	TArray<FMVVMConstFieldVariant> CreatePropertyPath(const UClass* Class, FName PropertyName, TArray<FMVVMConstFieldVariant> Properties) const;
 	bool IsPropertyPathValid(TArrayView<const FMVVMConstFieldVariant> PropertyPath) const;
@@ -120,7 +123,28 @@ private:
 		UEdGraph* SetterGraph = nullptr;
 	};
 	TArray<FCompilerSourceCreatorContext> CompilerSourceCreatorContexts;
-	int32 ViewModelDynamicCounter = 0;
+
+	/** The compiled binding from the binding compiler. */
+	struct FCompiledBinding
+	{
+		FCompiledBindingLibraryCompiler::FBindingHandle BindingHandle;
+
+		// At runtime, OnFieldValueChanged
+		//if the ConversionFunction exist
+		// then we read the Source
+		// then we call ConversionFunction with the source as an input
+		// then we copy the result of the ConversionFunction to DestinationWrite
+		//else if the ConversionFunction is complex
+		// then we call ConversionFunction without input (SourceRead can be invalid)
+		// then we copy the result of the ConversionFunction to DestinationWrite
+		//else
+		// then we copy SourceRead to DestiantionWrite.
+		FCompiledBindingLibraryCompiler::FFieldPathHandle SourceRead;
+		FCompiledBindingLibraryCompiler::FFieldPathHandle DestinationWrite;
+		FCompiledBindingLibraryCompiler::FFieldPathHandle ConversionFunction;
+
+		bool bIsConversionFunctionComplex = false;
+	};
 
 	/** 
 	 * Describe a binding for the compiler.
@@ -129,28 +153,21 @@ private:
 	enum class ECompilerBindingType
 	{
 		PropertyBinding, // normal binding
+		ViewModelDynamic, // added by a long binding chain
 	};
 	struct FCompilerBinding
 	{
+		ECompilerBindingType Type = ECompilerBindingType::PropertyBinding;
 		int32 BindingIndex = INDEX_NONE;
 		int32 UserWidgetPropertyContextIndex = INDEX_NONE;
 		int32 SourceCreatorContextIndex = INDEX_NONE;
 		bool bSourceIsUserWidget = false;
 		bool bFieldIdNeeded = false;
-		bool bIsConversionFunctionComplex = false;
-
-		FCompiledBindingLibraryCompiler::FBindingHandle BindingHandle;
+		bool bIsForwardBinding = false;
 
-		// At runtime, when TriggerFields changes
-		//if the ConversionFunction exist
-		// then we need to read all the inputs of the ConversionFunction
-		// then we copy the result of the ConversionFunction to DestinationWrite
-		//else we read the same property that TriggerField changed
-		// then we copy TriggerField to DestiantionWrite.
 		FCompiledBindingLibraryCompiler::FFieldIdHandle FieldIdHandle;
-		FCompiledBindingLibraryCompiler::FFieldPathHandle SourceRead;
-		FCompiledBindingLibraryCompiler::FFieldPathHandle DestinationWrite;
-		FCompiledBindingLibraryCompiler::FFieldPathHandle ConversionFunction;
+		FCompiledBinding CompiledBinding;
+		FName DynamicViewModelName;
 	};
 	TArray<FCompilerBinding> CompilerBindings;
 
@@ -161,8 +178,10 @@ private:
 	{
 		int32 BindingIndex = INDEX_NONE;
 		bool bIsForwardBinding = false;
+		// Complex binding do not need to add the PropertyPath to the binding info.
+		bool bIsComplexBinding = false;
 
-		UClass* SourceClass = nullptr;
+		const UClass* SourceClass = nullptr;
 		// The property that are registering to.
 		FFieldNotificationId FieldId;
 		// The path always start with the UserWidget as self.
diff --git a/Engine/Plugins/Runtime/ModelViewViewModel/Source/ModelViewViewModelBlueprint/Public/Bindings/MVVMCompiledBindingLibraryCompiler.h b/Engine/Plugins/Runtime/ModelViewViewModel/Source/ModelViewViewModelBlueprint/Public/Bindings/MVVMCompiledBindingLibraryCompiler.h
--- a/Engine/Plugins/Runtime/ModelViewViewModel/Source/ModelViewViewModelBlueprint/Public/Bindings/MVVMCompiledBindingLibraryCompiler.h
+++ b/Engine/Plugins/Runtime/ModelViewViewModel/Source/ModelViewViewModelBlueprint/Public/Bindings/MVVMCompiledBindingLibraryCompiler.h
@@ -153,7 +153,7 @@ public:
 
 public:
 	/** */
-	TValueOrError<FFieldIdHandle, FText> AddFieldId(TSubclassOf<UObject> SourceClass, FName FieldName);
+	TValueOrError<FFieldIdHandle, FText> AddFieldId(const UClass* SourceClass, FName FieldName);
 	
 	/** */
 	TValueOrError<FFieldPathHandle, FText> AddFieldPath(TArrayView<const UE::MVVM::FMVVMConstFieldVariant> FieldPath, bool bRead);
@@ -162,7 +162,7 @@ public:
 	TValueOrError<FFieldPathHandle, FText> AddObjectFieldPath(TArrayView<const UE::MVVM::FMVVMConstFieldVariant> FieldPath, UClass* ExpectedType, bool bRead);
 
 	/** */
-	TValueOrError<FFieldPathHandle, FText> AddConversionFunctionFieldPath(TSubclassOf<UObject> SourceClass, const UFunction* Function);
+	TValueOrError<FFieldPathHandle, FText> AddConversionFunctionFieldPath(const UClass* SourceClass, const UFunction* Function);
 
 	/** */
 	TValueOrError<FBindingHandle, FText> AddBinding(FFieldPathHandle Source, FFieldPathHandle Destination);
@@ -171,7 +171,7 @@ public:
 	TValueOrError<FBindingHandle, FText> AddBinding(FFieldPathHandle Source, FFieldPathHandle Destination, FFieldPathHandle ConversionFunction);
 
 	/** */
-	TValueOrError<FBindingHandle, FText> AddBinding(TArrayView<const FFieldPathHandle> Sources, FFieldPathHandle Destination, FFieldPathHandle ConversionFunction);
+	TValueOrError<FBindingHandle, FText> AddComplexBinding(FFieldPathHandle Destination, FFieldPathHandle ConversionFunction);
 
 	struct FCompileResult
 	{
@@ -187,6 +187,7 @@ public:
 private:
 	/** */
 	TValueOrError<FFieldPathHandle, FText> AddFieldPathImpl(TArrayView<const UE::MVVM::FMVVMConstFieldVariant> FieldPath, bool bRead);
+	TValueOrError<FBindingHandle, FText> AddBindingImpl(FFieldPathHandle Source, FFieldPathHandle Destination, FFieldPathHandle ConversionFunction, bool bIsComplexBinding);
 
 	TPimplPtr<Private::FCompiledBindingLibraryCompilerImpl> Impl;
 };
diff --git a/Engine/Plugins/Runtime/ModelViewViewModel/Source/ModelViewViewModelBlueprint/Public/MVVMDeveloperProjectSettings.h b/Engine/Plugins/Runtime/ModelViewViewModel/Source/ModelViewViewModelBlueprint/Public/MVVMDeveloperProjectSettings.h
--- a/Engine/Plugins/Runtime/ModelViewViewModel/Source/ModelViewViewModelBlueprint/Public/MVVMDeveloperProjectSettings.h
+++ b/Engine/Plugins/Runtime/ModelViewViewModel/Source/ModelViewViewModelBlueprint/Public/MVVMDeveloperProjectSettings.h
@@ -90,9 +90,13 @@ public:
 	UPROPERTY(EditAnywhere, config, Category = "Viewmodel")
 	bool bAllowBindingFromDetailView = true;
 
-	/** When generating a source from the Viewmodel editor, allow the compiler to generate a setter function. */
+	/** When generating a source in the viewmodel editor, allow the compiler to generate a setter function. */
 	UPROPERTY(EditAnywhere, config, Category = "Viewmodel")
 	bool bAllowGeneratedViewModelSetter = true;
+
+	/** When generating a binding with a long source path, allow the compiler to generate a new viewmodel source. */
+	UPROPERTY(EditAnywhere, config, Category = "Viewmodel")
+	bool bAllowLongSourcePath = true;
 	
 	/** For the binding list widget, allow the user to edit the binding in the detail view. */
 	UPROPERTY(EditAnywhere, config, Category = "Viewmodel")
diff --git a/Engine/Plugins/Runtime/ModelViewViewModel/Source/ModelViewViewModelEditor/Private/Tests/MVVMFieldPathHelperTest.cpp b/Engine/Plugins/Runtime/ModelViewViewModel/Source/ModelViewViewModelEditor/Private/Tests/MVVMFieldPathHelperTest.cpp
new file mode 100644
--- /dev/null
+++ b/Engine/Plugins/Runtime/ModelViewViewModel/Source/ModelViewViewModelEditor/Private/Tests/MVVMFieldPathHelperTest.cpp
@@ -0,0 +1,226 @@
+// Copyright Epic Games, Inc. All Rights Reserved.
+
+#include "Tests/MVVMFieldPathHelperTest.h"
+
+#include "Bindings/MVVMFieldPathHelper.h"
+#include "Containers/ArrayView.h"
+#include "Containers/StringFwd.h"
+#include "Misc/AutomationTest.h"
+#include "Templates/ValueOrError.h"
+#include "Types/MVVMFieldVariant.h"
+
+#include UE_INLINE_GENERATED_CPP_BY_NAME(MVVMFieldPathHelperTest)
+
+#if WITH_AUTOMATION_WORKER
+
+#define LOCTEXT_NAMESPACE "MVVMBindingHelperTest"
+
+IMPLEMENT_SIMPLE_AUTOMATION_TEST(FMVVMFieldPathHelperTestGenerateFieldPathList, "System.Plugins.MVVM.GetNotifyBindingInfoFromFieldPath", EAutomationTestFlags::ApplicationContextMask | EAutomationTestFlags::EngineFilter)
+
+bool FMVVMFieldPathHelperTestGenerateFieldPathList::RunTest(const FString& Parameters)
+{
+	using namespace UE::MVVM;
+
+	auto CompareResult = [this](const FieldPathHelper::FParsedNotifyBindingInfo& Result, const FieldPathHelper::FParsedNotifyBindingInfo& Expected, FStringView FieldPath)
+	{
+		if (Result.NotifyFieldClass != Expected.NotifyFieldClass)
+		{
+			AddError(FString::Printf(TEXT("Wrong NotifyFieldClass for '%s'"), FieldPath.GetData()));
+		}
+		if (Result.NotifyFieldId != Expected.NotifyFieldId)
+		{
+			AddError(FString::Printf(TEXT("Wrong NotifyFieldId for '%s'"), FieldPath.GetData()));
+		}
+		if (Result.ViewModelIndex != Expected.ViewModelIndex)
+		{
+			AddError(FString::Printf(TEXT("Wrong NotifyFieldId for '%s'"), FieldPath.GetData()));
+		}
+	};
+
+	auto Test = [this, &CompareResult](UClass* Accessor, FStringView FieldPath, TOptional<FieldPathHelper::FParsedNotifyBindingInfo> ExpectedResult)
+	{
+		TValueOrError<TArray<FMVVMConstFieldVariant>, FText> Fields = FieldPathHelper::GenerateFieldPathList(UMVVMWidgetFieldPathHelperTest::StaticClass(), FieldPath, true);
+		if (Fields.HasError())
+		{
+			AddError(FString::Printf(TEXT("The path for '%s' could not be generated."), FieldPath.GetData()));
+		}
+		else
+		{
+			TValueOrError<FieldPathHelper::FParsedNotifyBindingInfo, FText> BindingInfo = FieldPathHelper::GetNotifyBindingInfoFromFieldPath(UMVVMWidgetFieldPathHelperTest::StaticClass(), Fields.GetValue());
+			if (ExpectedResult.IsSet() && BindingInfo.HasError())
+			{
+				AddError(FString::Printf(TEXT("GetNotifyBindingInfoFromFieldPath failed for '%s'. %s."), FieldPath.GetData(), *BindingInfo.GetError().ToString()));
+			}
+			else if (!ExpectedResult.IsSet() && BindingInfo.HasValue())
+			{
+				AddError(FString::Printf(TEXT("GetNotifyBindingInfoFromFieldPath succeeed when it should had failed for '%s'"), FieldPath.GetData()));
+			}
+			else if (ExpectedResult.IsSet() && BindingInfo.HasValue())
+			{
+				CompareResult(ExpectedResult.GetValue(), BindingInfo.GetValue(), FieldPath);
+			}
+		}
+	};
+
+
+	//Always from the UserWidget point of view
+	{
+		FieldPathHelper::FParsedNotifyBindingInfo ExpectedInvalidInfo;
+		TOptional<FieldPathHelper::FParsedNotifyBindingInfo> ExpectedErrorInfo;
+
+		ExpectedInvalidInfo.NotifyFieldClass = nullptr;
+		ExpectedInvalidInfo.NotifyFieldId = FFieldNotificationId();
+		ExpectedInvalidInfo.ViewModelIndex = INDEX_NONE;
+
+		auto BuildInfo = [](UClass* Class, FName Notification, int32 Index)
+		{
+			FieldPathHelper::FParsedNotifyBindingInfo Result;
+			Result.NotifyFieldClass = Class;
+			Result.NotifyFieldId = FFieldNotificationId(Notification);
+			Result.ViewModelIndex = Index;
+			return Result;
+		};
+
+		Test(UMVVMWidgetFieldPathHelperTest::StaticClass(), TEXT("PropertyInt"), ExpectedInvalidInfo);
+		Test(UMVVMWidgetFieldPathHelperTest::StaticClass(), TEXT("PropertyIntNotify"), BuildInfo(UMVVMWidgetFieldPathHelperTest::StaticClass(), TEXT("PropertyIntNotify"), -1));
+		Test(UMVVMWidgetFieldPathHelperTest::StaticClass(), TEXT("PropertyVector"), ExpectedInvalidInfo);
+		Test(UMVVMWidgetFieldPathHelperTest::StaticClass(), TEXT("PropertyVector.X"), ExpectedInvalidInfo);
+		Test(UMVVMWidgetFieldPathHelperTest::StaticClass(), TEXT("PropertyVectorNotify"), BuildInfo(UMVVMWidgetFieldPathHelperTest::StaticClass(), TEXT("PropertyVectorNotify"), -1));
+		Test(UMVVMWidgetFieldPathHelperTest::StaticClass(), TEXT("PropertyVectorNotify.X"), BuildInfo(UMVVMWidgetFieldPathHelperTest::StaticClass(), TEXT("PropertyVectorNotify"), -1));
+		Test(UMVVMWidgetFieldPathHelperTest::StaticClass(), TEXT("PropertyStruct"), ExpectedInvalidInfo);
+		Test(UMVVMWidgetFieldPathHelperTest::StaticClass(), TEXT("PropertyStruct.PropertyInt"), ExpectedInvalidInfo);
+		Test(UMVVMWidgetFieldPathHelperTest::StaticClass(), TEXT("PropertyStruct.PropertyVector"), ExpectedInvalidInfo);
+		Test(UMVVMWidgetFieldPathHelperTest::StaticClass(), TEXT("PropertyStruct.PropertyVector.X"), ExpectedInvalidInfo);
+		Test(UMVVMWidgetFieldPathHelperTest::StaticClass(), TEXT("PropertyStruct.PropertyObject"), ExpectedInvalidInfo);
+		Test(UMVVMWidgetFieldPathHelperTest::StaticClass(), TEXT("PropertyStruct.PropertyObject.PropertyInt"), ExpectedInvalidInfo);
+		Test(UMVVMWidgetFieldPathHelperTest::StaticClass(), TEXT("PropertyStruct.PropertyObject.PropertyVector"), ExpectedInvalidInfo);
+		Test(UMVVMWidgetFieldPathHelperTest::StaticClass(), TEXT("PropertyStruct.PropertyObject.PropertyVector.X"), ExpectedInvalidInfo);
+		Test(UMVVMWidgetFieldPathHelperTest::StaticClass(), TEXT("PropertyStruct.PropertyObject.PropertyObject"), ExpectedInvalidInfo);
+		Test(UMVVMWidgetFieldPathHelperTest::StaticClass(), TEXT("PropertyStruct.PropertyObject.PropertyObject.PropertyInt"), ExpectedInvalidInfo);
+		Test(UMVVMWidgetFieldPathHelperTest::StaticClass(), TEXT("PropertyStruct.PropertyObject.PropertyObject.PropertyVector"), ExpectedInvalidInfo);
+		Test(UMVVMWidgetFieldPathHelperTest::StaticClass(), TEXT("PropertyStruct.PropertyObject.PropertyObject.PropertyVector.X"), ExpectedInvalidInfo);
+		Test(UMVVMWidgetFieldPathHelperTest::StaticClass(), TEXT("PropertyStruct.PropertyObject.PropertyViewModel"), ExpectedInvalidInfo);
+		Test(UMVVMWidgetFieldPathHelperTest::StaticClass(), TEXT("PropertyStruct.PropertyViewModel"), ExpectedInvalidInfo);
+		Test(UMVVMWidgetFieldPathHelperTest::StaticClass(), TEXT("PropertyStruct.PropertyViewModel.PropertyInt"), ExpectedInvalidInfo);
+		Test(UMVVMWidgetFieldPathHelperTest::StaticClass(), TEXT("PropertyStruct.PropertyViewModel.PropertyIntNotify"), ExpectedInvalidInfo);
+		Test(UMVVMWidgetFieldPathHelperTest::StaticClass(), TEXT("PropertyStruct.PropertyViewModel.PropertyVector"), ExpectedInvalidInfo);
+		Test(UMVVMWidgetFieldPathHelperTest::StaticClass(), TEXT("PropertyStruct.PropertyViewModel.PropertyVector.X"), ExpectedInvalidInfo);
+		Test(UMVVMWidgetFieldPathHelperTest::StaticClass(), TEXT("PropertyStruct.PropertyViewModel.PropertyVectorNotify"), ExpectedInvalidInfo);
+		Test(UMVVMWidgetFieldPathHelperTest::StaticClass(), TEXT("PropertyStruct.PropertyViewModel.PropertyVectorNotify.X"), ExpectedInvalidInfo);
+		Test(UMVVMWidgetFieldPathHelperTest::StaticClass(), TEXT("PropertyStruct.PropertyViewModel.PropertyObject"), ExpectedInvalidInfo);
+		Test(UMVVMWidgetFieldPathHelperTest::StaticClass(), TEXT("PropertyStruct.PropertyViewModel.PropertyViewModel.PropertyInt"), ExpectedInvalidInfo);
+		Test(UMVVMWidgetFieldPathHelperTest::StaticClass(), TEXT("PropertyStruct.PropertyViewModel.PropertyViewModel.PropertyIntNotify"), ExpectedInvalidInfo);
+		Test(UMVVMWidgetFieldPathHelperTest::StaticClass(), TEXT("PropertyStruct.PropertyViewModel.PropertyViewModel.PropertyVector"), ExpectedInvalidInfo);
+		Test(UMVVMWidgetFieldPathHelperTest::StaticClass(), TEXT("PropertyStruct.PropertyViewModel.PropertyViewModel.PropertyVector.X"), ExpectedInvalidInfo);
+		Test(UMVVMWidgetFieldPathHelperTest::StaticClass(), TEXT("PropertyStruct.PropertyViewModel.PropertyViewModel.PropertyVectorNotify"), ExpectedInvalidInfo);
+		Test(UMVVMWidgetFieldPathHelperTest::StaticClass(), TEXT("PropertyStruct.PropertyViewModel.PropertyViewModel.PropertyVectorNotify.X"), ExpectedInvalidInfo);
+		Test(UMVVMWidgetFieldPathHelperTest::StaticClass(), TEXT("PropertyObject"), ExpectedInvalidInfo);
+		Test(UMVVMWidgetFieldPathHelperTest::StaticClass(), TEXT("PropertyObject.PropertyInt"), ExpectedInvalidInfo);
+		Test(UMVVMWidgetFieldPathHelperTest::StaticClass(), TEXT("PropertyObject.PropertyVector"), ExpectedInvalidInfo);
+		Test(UMVVMWidgetFieldPathHelperTest::StaticClass(), TEXT("PropertyObject.PropertyVector.X"), ExpectedInvalidInfo);
+		Test(UMVVMWidgetFieldPathHelperTest::StaticClass(), TEXT("PropertyObject.PropertyObject"), ExpectedInvalidInfo);
+		Test(UMVVMWidgetFieldPathHelperTest::StaticClass(), TEXT("PropertyObject.PropertyObject.PropertyInt"), ExpectedInvalidInfo);
+		Test(UMVVMWidgetFieldPathHelperTest::StaticClass(), TEXT("PropertyObject.PropertyObject.PropertyVector"), ExpectedInvalidInfo);
+		Test(UMVVMWidgetFieldPathHelperTest::StaticClass(), TEXT("PropertyObject.PropertyObject.PropertyVector.X"), ExpectedInvalidInfo);
+		Test(UMVVMWidgetFieldPathHelperTest::StaticClass(), TEXT("PropertyObject.PropertyObject.PropertyObject"), ExpectedInvalidInfo);
+		Test(UMVVMWidgetFieldPathHelperTest::StaticClass(), TEXT("PropertyObject.PropertyObject.PropertyViewModel"), ExpectedInvalidInfo);
+		Test(UMVVMWidgetFieldPathHelperTest::StaticClass(), TEXT("PropertyObject.PropertyViewModel"), ExpectedInvalidInfo);
+		Test(UMVVMWidgetFieldPathHelperTest::StaticClass(), TEXT("PropertyObject.PropertyViewModel.PropertyInt"), ExpectedInvalidInfo);
+		Test(UMVVMWidgetFieldPathHelperTest::StaticClass(), TEXT("PropertyObject.PropertyViewModel.PropertyIntNotify"), ExpectedInvalidInfo);
+		Test(UMVVMWidgetFieldPathHelperTest::StaticClass(), TEXT("PropertyObject.PropertyViewModel.PropertyVector"), ExpectedInvalidInfo);
+		Test(UMVVMWidgetFieldPathHelperTest::StaticClass(), TEXT("PropertyObject.PropertyViewModel.PropertyVector.X"), ExpectedInvalidInfo);
+		Test(UMVVMWidgetFieldPathHelperTest::StaticClass(), TEXT("PropertyObject.PropertyViewModel.PropertyVectorNotify"), ExpectedInvalidInfo);
+		Test(UMVVMWidgetFieldPathHelperTest::StaticClass(), TEXT("PropertyObject.PropertyViewModel.PropertyVectorNotify.X"), ExpectedInvalidInfo);
+		Test(UMVVMWidgetFieldPathHelperTest::StaticClass(), TEXT("PropertyObject.PropertyViewModel.PropertyObject"), ExpectedInvalidInfo);
+		Test(UMVVMWidgetFieldPathHelperTest::StaticClass(), TEXT("PropertyObject.PropertyViewModel.PropertyViewModel.PropertyInt"), ExpectedInvalidInfo);
+		Test(UMVVMWidgetFieldPathHelperTest::StaticClass(), TEXT("PropertyObject.PropertyViewModel.PropertyViewModel.PropertyIntNotify"), ExpectedInvalidInfo);
+		Test(UMVVMWidgetFieldPathHelperTest::StaticClass(), TEXT("PropertyObject.PropertyViewModel.PropertyViewModel.PropertyVector"), ExpectedInvalidInfo);
+		Test(UMVVMWidgetFieldPathHelperTest::StaticClass(), TEXT("PropertyObject.PropertyViewModel.PropertyViewModel.PropertyVector.X"), ExpectedInvalidInfo);
+		Test(UMVVMWidgetFieldPathHelperTest::StaticClass(), TEXT("PropertyObject.PropertyViewModel.PropertyViewModel.PropertyVectorNotify"), ExpectedInvalidInfo);
+		Test(UMVVMWidgetFieldPathHelperTest::StaticClass(), TEXT("PropertyObject.PropertyViewModel.PropertyViewModel.PropertyVectorNotify.X"), ExpectedInvalidInfo);
+		Test(UMVVMWidgetFieldPathHelperTest::StaticClass(), TEXT("PropertyViewModel"), ExpectedInvalidInfo);
+		Test(UMVVMWidgetFieldPathHelperTest::StaticClass(), TEXT("PropertyViewModel.PropertyInt"), ExpectedInvalidInfo);
+		Test(UMVVMWidgetFieldPathHelperTest::StaticClass(), TEXT("PropertyViewModel.PropertyIntNotify"), BuildInfo(UMVVMViewModelFieldPathHelperTest::StaticClass(), TEXT("PropertyIntNotify"), 0));
+		Test(UMVVMWidgetFieldPathHelperTest::StaticClass(), TEXT("PropertyViewModel.PropertyVector"), ExpectedInvalidInfo);
+		Test(UMVVMWidgetFieldPathHelperTest::StaticClass(), TEXT("PropertyViewModel.PropertyVector.X"), ExpectedInvalidInfo);
+		Test(UMVVMWidgetFieldPathHelperTest::StaticClass(), TEXT("PropertyViewModel.PropertyVectorNotify"), BuildInfo(UMVVMViewModelFieldPathHelperTest::StaticClass(), TEXT("PropertyVectorNotify"), 0));
+		Test(UMVVMWidgetFieldPathHelperTest::StaticClass(), TEXT("PropertyViewModel.PropertyVectorNotify.X"), BuildInfo(UMVVMViewModelFieldPathHelperTest::StaticClass(), TEXT("PropertyVectorNotify"), 0));
+		Test(UMVVMWidgetFieldPathHelperTest::StaticClass(), TEXT("PropertyViewModel.PropertyStructNotify.PropertyInt"), BuildInfo(UMVVMViewModelFieldPathHelperTest::StaticClass(), TEXT("PropertyStructNotify"), 0));
+		Test(UMVVMWidgetFieldPathHelperTest::StaticClass(), TEXT("PropertyViewModel.PropertyStructNotify.PropertyVector"), BuildInfo(UMVVMViewModelFieldPathHelperTest::StaticClass(), TEXT("PropertyStructNotify"), 0));
+		Test(UMVVMWidgetFieldPathHelperTest::StaticClass(), TEXT("PropertyViewModel.PropertyStructNotify.PropertyVector.X"), BuildInfo(UMVVMViewModelFieldPathHelperTest::StaticClass(), TEXT("PropertyStructNotify"), 0));
+		Test(UMVVMWidgetFieldPathHelperTest::StaticClass(), TEXT("PropertyViewModel.PropertyStructNotify.PropertyObject"), BuildInfo(UMVVMViewModelFieldPathHelperTest::StaticClass(), TEXT("PropertyStructNotify"), 0));
+		Test(UMVVMWidgetFieldPathHelperTest::StaticClass(), TEXT("PropertyViewModel.PropertyStructNotify.PropertyObject.PropertyVector.X"), BuildInfo(UMVVMViewModelFieldPathHelperTest::StaticClass(), TEXT("PropertyStructNotify"), 0));
+		Test(UMVVMWidgetFieldPathHelperTest::StaticClass(), TEXT("PropertyViewModel.PropertyStructNotify.PropertyViewModel"), BuildInfo(UMVVMViewModelFieldPathHelperTest::StaticClass(), TEXT("PropertyStructNotify"), 0));
+		Test(UMVVMWidgetFieldPathHelperTest::StaticClass(), TEXT("PropertyViewModel.PropertyObject"), ExpectedInvalidInfo);
+		Test(UMVVMWidgetFieldPathHelperTest::StaticClass(), TEXT("PropertyViewModel.PropertyObject.PropertyInt"), ExpectedInvalidInfo);
+		Test(UMVVMWidgetFieldPathHelperTest::StaticClass(), TEXT("PropertyViewModel.PropertyObject.PropertyVector"), ExpectedInvalidInfo);
+		Test(UMVVMWidgetFieldPathHelperTest::StaticClass(), TEXT("PropertyViewModel.PropertyObject.PropertyVector.X"), ExpectedInvalidInfo);
+		Test(UMVVMWidgetFieldPathHelperTest::StaticClass(), TEXT("PropertyViewModel.PropertyObject.PropertyObject"), ExpectedInvalidInfo);
+		Test(UMVVMWidgetFieldPathHelperTest::StaticClass(), TEXT("PropertyViewModel.PropertyObject.PropertyViewModel"), ExpectedInvalidInfo);
+		Test(UMVVMWidgetFieldPathHelperTest::StaticClass(), TEXT("PropertyViewModel.PropertyObject.PropertyViewModel.PropertyInt"), ExpectedInvalidInfo);
+		Test(UMVVMWidgetFieldPathHelperTest::StaticClass(), TEXT("PropertyViewModel.PropertyObject.PropertyViewModel.PropertyIntNotify"), ExpectedInvalidInfo);
+		Test(UMVVMWidgetFieldPathHelperTest::StaticClass(), TEXT("PropertyViewModel.PropertyObject.PropertyViewModel.PropertyVector"), ExpectedInvalidInfo);
+		Test(UMVVMWidgetFieldPathHelperTest::StaticClass(), TEXT("PropertyViewModel.PropertyObject.PropertyViewModel.PropertyVector.X"), ExpectedInvalidInfo);
+		Test(UMVVMWidgetFieldPathHelperTest::StaticClass(), TEXT("PropertyViewModel.PropertyObject.PropertyViewModel.PropertyVectorNotify"), ExpectedInvalidInfo);
+		Test(UMVVMWidgetFieldPathHelperTest::StaticClass(), TEXT("PropertyViewModel.PropertyObject.PropertyViewModel.PropertyVectorNotify.X"), ExpectedInvalidInfo);
+		Test(UMVVMWidgetFieldPathHelperTest::StaticClass(), TEXT("PropertyViewModel.PropertyViewModel"), ExpectedInvalidInfo);
+		Test(UMVVMWidgetFieldPathHelperTest::StaticClass(), TEXT("PropertyViewModel.PropertyViewModel.PropertyInt"), ExpectedInvalidInfo);
+		Test(UMVVMWidgetFieldPathHelperTest::StaticClass(), TEXT("PropertyViewModel.PropertyViewModel.PropertyIntNotify"), ExpectedInvalidInfo);
+		Test(UMVVMWidgetFieldPathHelperTest::StaticClass(), TEXT("PropertyViewModel.PropertyViewModel.PropertyVector"), ExpectedInvalidInfo);
+		Test(UMVVMWidgetFieldPathHelperTest::StaticClass(), TEXT("PropertyViewModel.PropertyViewModel.PropertyVector.X"), ExpectedInvalidInfo);
+		Test(UMVVMWidgetFieldPathHelperTest::StaticClass(), TEXT("PropertyViewModel.PropertyViewModel.PropertyVectorNotify"), ExpectedInvalidInfo);
+		Test(UMVVMWidgetFieldPathHelperTest::StaticClass(), TEXT("PropertyViewModel.PropertyViewModel.PropertyVectorNotify.X"), ExpectedInvalidInfo);
+		Test(UMVVMWidgetFieldPathHelperTest::StaticClass(), TEXT("PropertyViewModel.PropertyViewModel.PropertyObject"), ExpectedInvalidInfo);
+		Test(UMVVMWidgetFieldPathHelperTest::StaticClass(), TEXT("PropertyViewModel.PropertyViewModel.PropertyViewModel.PropertyInt"), ExpectedInvalidInfo);
+		Test(UMVVMWidgetFieldPathHelperTest::StaticClass(), TEXT("PropertyViewModel.PropertyViewModel.PropertyViewModel.PropertyIntNotify"), ExpectedInvalidInfo);
+		Test(UMVVMWidgetFieldPathHelperTest::StaticClass(), TEXT("PropertyViewModel.PropertyViewModel.PropertyViewModel.PropertyVector"), ExpectedInvalidInfo);
+		Test(UMVVMWidgetFieldPathHelperTest::StaticClass(), TEXT("PropertyViewModel.PropertyViewModel.PropertyViewModel.PropertyVector.X"), ExpectedInvalidInfo);
+		Test(UMVVMWidgetFieldPathHelperTest::StaticClass(), TEXT("PropertyViewModel.PropertyViewModel.PropertyViewModel.PropertyVectorNotify"), ExpectedInvalidInfo);
+		Test(UMVVMWidgetFieldPathHelperTest::StaticClass(), TEXT("PropertyViewModel.PropertyViewModel.PropertyViewModel.PropertyVectorNotify.X"), ExpectedInvalidInfo);
+		Test(UMVVMWidgetFieldPathHelperTest::StaticClass(), TEXT("PropertyViewModel.PropertyViewModelNotify"), BuildInfo(UMVVMViewModelFieldPathHelperTest::StaticClass(), TEXT("PropertyViewModelNotify"), 0));
+		Test(UMVVMWidgetFieldPathHelperTest::StaticClass(), TEXT("PropertyViewModel.PropertyViewModelNotify.PropertyInt"), BuildInfo(UMVVMViewModelFieldPathHelperTest::StaticClass(), TEXT("PropertyViewModelNotify"), 0));
+		Test(UMVVMWidgetFieldPathHelperTest::StaticClass(), TEXT("PropertyViewModel.PropertyViewModelNotify.PropertyIntNotify"), BuildInfo(UMVVMViewModelFieldPathHelperTest::StaticClass(), TEXT("PropertyIntNotify"), 1));
+		Test(UMVVMWidgetFieldPathHelperTest::StaticClass(), TEXT("PropertyViewModel.PropertyViewModelNotify.PropertyVector"), BuildInfo(UMVVMViewModelFieldPathHelperTest::StaticClass(), TEXT("PropertyViewModelNotify"), 0));
+		Test(UMVVMWidgetFieldPathHelperTest::StaticClass(), TEXT("PropertyViewModel.PropertyViewModelNotify.PropertyVector.X"), BuildInfo(UMVVMViewModelFieldPathHelperTest::StaticClass(), TEXT("PropertyViewModelNotify"), 0));
+		Test(UMVVMWidgetFieldPathHelperTest::StaticClass(), TEXT("PropertyViewModel.PropertyViewModelNotify.PropertyVectorNotify"), BuildInfo(UMVVMViewModelFieldPathHelperTest::StaticClass(), TEXT("PropertyVectorNotify"), 1));
+		Test(UMVVMWidgetFieldPathHelperTest::StaticClass(), TEXT("PropertyViewModel.PropertyViewModelNotify.PropertyVectorNotify.X"), BuildInfo(UMVVMViewModelFieldPathHelperTest::StaticClass(), TEXT("PropertyVectorNotify"), 1));
+		Test(UMVVMWidgetFieldPathHelperTest::StaticClass(), TEXT("PropertyViewModel.PropertyViewModelNotify.PropertyObject"), BuildInfo(UMVVMViewModelFieldPathHelperTest::StaticClass(), TEXT("PropertyViewModelNotify"), 0));
+		Test(UMVVMWidgetFieldPathHelperTest::StaticClass(), TEXT("PropertyViewModel.PropertyViewModelNotify.PropertyViewModel.PropertyInt"), BuildInfo(UMVVMViewModelFieldPathHelperTest::StaticClass(), TEXT("PropertyViewModelNotify"), 0));
+		Test(UMVVMWidgetFieldPathHelperTest::StaticClass(), TEXT("PropertyViewModel.PropertyViewModelNotify.PropertyViewModel.PropertyIntNotify"), ExpectedInvalidInfo);
+		Test(UMVVMWidgetFieldPathHelperTest::StaticClass(), TEXT("PropertyViewModel.PropertyViewModelNotify.PropertyViewModel.PropertyVector"), BuildInfo(UMVVMViewModelFieldPathHelperTest::StaticClass(), TEXT("PropertyViewModelNotify"), 0));
+		Test(UMVVMWidgetFieldPathHelperTest::StaticClass(), TEXT("PropertyViewModel.PropertyViewModelNotify.PropertyViewModel.PropertyVector.X"), BuildInfo(UMVVMViewModelFieldPathHelperTest::StaticClass(), TEXT("PropertyViewModelNotify"), 0));
+		Test(UMVVMWidgetFieldPathHelperTest::StaticClass(), TEXT("PropertyViewModel.PropertyViewModelNotify.PropertyViewModel.PropertyVectorNotify"), ExpectedInvalidInfo);
+		Test(UMVVMWidgetFieldPathHelperTest::StaticClass(), TEXT("PropertyViewModel.PropertyViewModelNotify.PropertyViewModel.PropertyVectorNotify.X"), ExpectedInvalidInfo);
+		Test(UMVVMWidgetFieldPathHelperTest::StaticClass(), TEXT("PropertyViewModel.PropertyViewModelNotify.PropertyViewModelNotify.PropertyInt"), BuildInfo(UMVVMViewModelFieldPathHelperTest::StaticClass(), TEXT("PropertyViewModelNotify"), 1));
+		Test(UMVVMWidgetFieldPathHelperTest::StaticClass(), TEXT("PropertyViewModel.PropertyViewModelNotify.PropertyViewModelNotify.PropertyIntNotify"), BuildInfo(UMVVMViewModelFieldPathHelperTest::StaticClass(), TEXT("PropertyIntNotify"), 2));
+		Test(UMVVMWidgetFieldPathHelperTest::StaticClass(), TEXT("PropertyViewModel.PropertyViewModelNotify.PropertyViewModelNotify.PropertyVector"), BuildInfo(UMVVMViewModelFieldPathHelperTest::StaticClass(), TEXT("PropertyViewModelNotify"), 1));
+		Test(UMVVMWidgetFieldPathHelperTest::StaticClass(), TEXT("PropertyViewModel.PropertyViewModelNotify.PropertyViewModelNotify.PropertyVector.X"), BuildInfo(UMVVMViewModelFieldPathHelperTest::StaticClass(), TEXT("PropertyViewModelNotify"), 1));
+		Test(UMVVMWidgetFieldPathHelperTest::StaticClass(), TEXT("PropertyViewModel.PropertyViewModelNotify.PropertyViewModelNotify.PropertyVectorNotify"), BuildInfo(UMVVMViewModelFieldPathHelperTest::StaticClass(), TEXT("PropertyVectorNotify"), 2));
+		Test(UMVVMWidgetFieldPathHelperTest::StaticClass(), TEXT("PropertyViewModel.PropertyViewModelNotify.PropertyViewModelNotify.PropertyVectorNotify.X"), BuildInfo(UMVVMViewModelFieldPathHelperTest::StaticClass(), TEXT("PropertyVectorNotify"), 2));
+		Test(UMVVMWidgetFieldPathHelperTest::StaticClass(), TEXT("PropertyViewModelNotify.PropertyViewModelNotify"), BuildInfo(UMVVMViewModelFieldPathHelperTest::StaticClass(), TEXT("PropertyViewModelNotify"), 0));
+		Test(UMVVMWidgetFieldPathHelperTest::StaticClass(), TEXT("PropertyViewModelNotify.PropertyViewModelNotify.PropertyInt"), BuildInfo(UMVVMViewModelFieldPathHelperTest::StaticClass(), TEXT("PropertyViewModelNotify"), 0));
+		Test(UMVVMWidgetFieldPathHelperTest::StaticClass(), TEXT("PropertyViewModelNotify.PropertyViewModelNotify.PropertyIntNotify"), BuildInfo(UMVVMViewModelFieldPathHelperTest::StaticClass(), TEXT("PropertyIntNotify"), 1));
+		Test(UMVVMWidgetFieldPathHelperTest::StaticClass(), TEXT("PropertyViewModelNotify.PropertyViewModelNotify.PropertyVector"), BuildInfo(UMVVMViewModelFieldPathHelperTest::StaticClass(), TEXT("PropertyViewModelNotify"), 0));
+		Test(UMVVMWidgetFieldPathHelperTest::StaticClass(), TEXT("PropertyViewModelNotify.PropertyViewModelNotify.PropertyVector.X"), BuildInfo(UMVVMViewModelFieldPathHelperTest::StaticClass(), TEXT("PropertyViewModelNotify"), 0));
+		Test(UMVVMWidgetFieldPathHelperTest::StaticClass(), TEXT("PropertyViewModelNotify.PropertyViewModelNotify.PropertyVectorNotify"), BuildInfo(UMVVMViewModelFieldPathHelperTest::StaticClass(), TEXT("PropertyVectorNotify"), 1));
+		Test(UMVVMWidgetFieldPathHelperTest::StaticClass(), TEXT("PropertyViewModelNotify.PropertyViewModelNotify.PropertyVectorNotify.X"), BuildInfo(UMVVMViewModelFieldPathHelperTest::StaticClass(), TEXT("PropertyVectorNotify"), 1));
+		Test(UMVVMWidgetFieldPathHelperTest::StaticClass(), TEXT("PropertyViewModelNotify.PropertyViewModelNotify.PropertyObject"), BuildInfo(UMVVMViewModelFieldPathHelperTest::StaticClass(), TEXT("PropertyViewModelNotify"), 0));
+		Test(UMVVMWidgetFieldPathHelperTest::StaticClass(), TEXT("PropertyViewModelNotify.PropertyViewModelNotify.PropertyViewModel.PropertyInt"), BuildInfo(UMVVMViewModelFieldPathHelperTest::StaticClass(), TEXT("PropertyViewModelNotify"), 0));
+		Test(UMVVMWidgetFieldPathHelperTest::StaticClass(), TEXT("PropertyViewModelNotify.PropertyViewModelNotify.PropertyViewModel.PropertyIntNotify"), ExpectedInvalidInfo);
+		Test(UMVVMWidgetFieldPathHelperTest::StaticClass(), TEXT("PropertyViewModelNotify.PropertyViewModelNotify.PropertyViewModel.PropertyVector"), BuildInfo(UMVVMViewModelFieldPathHelperTest::StaticClass(), TEXT("PropertyViewModelNotify"), 0));
+		Test(UMVVMWidgetFieldPathHelperTest::StaticClass(), TEXT("PropertyViewModelNotify.PropertyViewModelNotify.PropertyViewModel.PropertyVector.X"), BuildInfo(UMVVMViewModelFieldPathHelperTest::StaticClass(), TEXT("PropertyViewModelNotify"), 0));
+		Test(UMVVMWidgetFieldPathHelperTest::StaticClass(), TEXT("PropertyViewModelNotify.PropertyViewModelNotify.PropertyViewModel.PropertyVectorNotify"), ExpectedInvalidInfo);
+		Test(UMVVMWidgetFieldPathHelperTest::StaticClass(), TEXT("PropertyViewModelNotify.PropertyViewModelNotify.PropertyViewModel.PropertyVectorNotify.X"), ExpectedInvalidInfo);
+		Test(UMVVMWidgetFieldPathHelperTest::StaticClass(), TEXT("PropertyViewModelNotify.PropertyViewModelNotify.PropertyViewModelNotify.PropertyInt"), BuildInfo(UMVVMViewModelFieldPathHelperTest::StaticClass(), TEXT("PropertyViewModelNotify"), 1));
+		Test(UMVVMWidgetFieldPathHelperTest::StaticClass(), TEXT("PropertyViewModelNotify.PropertyViewModelNotify.PropertyViewModelNotify.PropertyIntNotify"), BuildInfo(UMVVMViewModelFieldPathHelperTest::StaticClass(), TEXT("PropertyIntNotify"), 2));
+		Test(UMVVMWidgetFieldPathHelperTest::StaticClass(), TEXT("PropertyViewModelNotify.PropertyViewModelNotify.PropertyViewModelNotify.PropertyVector"), BuildInfo(UMVVMViewModelFieldPathHelperTest::StaticClass(), TEXT("PropertyViewModelNotify"), 1));
+		Test(UMVVMWidgetFieldPathHelperTest::StaticClass(), TEXT("PropertyViewModelNotify.PropertyViewModelNotify.PropertyViewModelNotify.PropertyVector.X"), BuildInfo(UMVVMViewModelFieldPathHelperTest::StaticClass(), TEXT("PropertyViewModelNotify"), 1));
+		Test(UMVVMWidgetFieldPathHelperTest::StaticClass(), TEXT("PropertyViewModelNotify.PropertyViewModelNotify.PropertyViewModelNotify.PropertyVectorNotify"), BuildInfo(UMVVMViewModelFieldPathHelperTest::StaticClass(), TEXT("PropertyVectorNotify"), 2));
+		Test(UMVVMWidgetFieldPathHelperTest::StaticClass(), TEXT("PropertyViewModelNotify.PropertyViewModelNotify.PropertyViewModelNotify.PropertyVectorNotify.X"), BuildInfo(UMVVMViewModelFieldPathHelperTest::StaticClass(), TEXT("PropertyVectorNotify"), 2));
+	}
+
+	return true;
+}
+
+#undef LOCTEXT_NAMESPACE 
+#endif //WITH_AUTOMATION_WORKER
diff --git a/Engine/Plugins/Runtime/ModelViewViewModel/Source/ModelViewViewModelEditor/Private/Tests/MVVMFieldPathHelperTest.h b/Engine/Plugins/Runtime/ModelViewViewModel/Source/ModelViewViewModelEditor/Private/Tests/MVVMFieldPathHelperTest.h
new file mode 100644
--- /dev/null
+++ b/Engine/Plugins/Runtime/ModelViewViewModel/Source/ModelViewViewModelEditor/Private/Tests/MVVMFieldPathHelperTest.h
@@ -0,0 +1,439 @@
+// Copyright Epic Games, Inc. All Rights Reserved.
+
+#pragma once
+
+#include "CoreMinimal.h"
+#include "Blueprint/UserWidget.h"
+#include "MVVMViewModelBase.h"
+
+#include "MVVMFieldPathHelperTest.generated.h"
+
+class UMVVMObjectFieldPathHelperTest;
+class UMVVMViewModelFieldPathHelperTest;
+class UMVVMWidgetFieldPathHelperTest;
+
+/** */
+USTRUCT(BlueprintInternalUseOnly)
+struct FMVVMStructFieldPathHelperTest
+{
+	GENERATED_BODY()
+
+public:
+	UPROPERTY()
+	int32 PropertyInt;
+
+	UPROPERTY()
+	FVector PropertyVector;
+
+	UPROPERTY()
+	TObjectPtr<UMVVMObjectFieldPathHelperTest> PropertyObject;
+
+	UPROPERTY()
+	TObjectPtr<UMVVMViewModelFieldPathHelperTest> PropertyViewModel;
+};
+
+/** */
+UCLASS(MinimalAPI, Transient, hidedropdown, Abstract, Hidden)
+class UMVVMObjectFieldPathHelperTest : public UObject
+{
+	GENERATED_BODY()
+
+public:
+	UPROPERTY()
+	int32 PropertyInt;
+	UFUNCTION()
+	int32 GetPropertyInt() const { return 0; }
+	UFUNCTION()
+	void SetPropertyInt(int32 Value) { }
+
+	UPROPERTY(Getter, Setter)
+	int32 PropertyIntWithGetterSetter;
+	UFUNCTION()
+	int32 GetPropertyIntWithGetterSetter() const { return 0; }
+	UFUNCTION()
+	void SetPropertyIntWithGetterSetter(int32 Value) {}
+
+	UPROPERTY(Getter, Setter, BlueprintGetter="GetPropertyIntWithGetterSetterAndBP", BlueprintSetter="SetPropertyIntWithGetterSetterAndBP", Category="Test")
+	int32 PropertyIntWithGetterSetterAndBP;
+	UFUNCTION(BlueprintPure, Category = "Test")
+	int32 GetPropertyIntWithGetterSetterAndBP() const { return 0;}
+	UFUNCTION(BlueprintCallable, Category = "Test")
+	void SetPropertyIntWithGetterSetterAndBP(int32 Value){}
+
+	UPROPERTY(BlueprintGetter="GetPropertyIntWithBPGetterSetter", BlueprintSetter="SetPropertyIntWithBPGetterSetter", Category = "Test")
+	int32 PropertyIntWithBPGetterSetter;
+	UFUNCTION(BlueprintPure, Category = "Test")
+	int32 GetPropertyIntWithBPGetterSetter() const { return 0; }
+	UFUNCTION(BlueprintCallable, Category = "Test")
+	void SetPropertyIntWithBPGetterSetter(int32 Value) { }
+
+public:
+	UPROPERTY()
+	FVector PropertyVector;
+	UFUNCTION()
+	FVector GetPropertyVector() const { return FVector(); }
+	UFUNCTION()
+	void SetPropertyVector(FVector Value) { }
+
+	UPROPERTY(Getter, Setter)
+	FVector PropertyVectorWithGetterSetter;
+	UFUNCTION()
+	FVector GetPropertyVectorWithGetterSetter() const { return FVector(); }
+	UFUNCTION()
+	void SetPropertyVectorWithGetterSetter(FVector Value) {}
+
+	UPROPERTY(Getter, Setter, BlueprintGetter="GetPropertyVectorWithGetterSetterAndBP", BlueprintSetter="SetPropertyVectorWithGetterSetterAndBP", Category = "Test")
+	FVector PropertyVectorWithGetterSetterAndBP;
+	UFUNCTION(BlueprintPure, Category = "Test")
+	FVector GetPropertyVectorWithGetterSetterAndBP() const { return FVector();}
+	UFUNCTION(BlueprintCallable, Category = "Test")
+	void SetPropertyVectorWithGetterSetterAndBP(FVector Value){}
+
+	UPROPERTY(BlueprintGetter="GetPropertyVectorWithBPGetterSetter", BlueprintSetter="SetPropertyVectorWithBPGetterSetter", Category = "Test")
+	FVector PropertyVectorWithBPGetterSetter;
+	UFUNCTION(BlueprintPure, Category = "Test")
+	FVector GetPropertyVectorWithBPGetterSetter() const { return FVector(); }
+	UFUNCTION(BlueprintCallable, Category = "Test")
+	void SetPropertyVectorWithBPGetterSetter(FVector Value) { }
+
+public:
+	UPROPERTY()
+	FMVVMStructFieldPathHelperTest PropertyStruct;
+	UFUNCTION()
+	FMVVMStructFieldPathHelperTest GetPropertyStruct() const { return FMVVMStructFieldPathHelperTest(); }
+	UFUNCTION()
+	void SetPropertyStruct(FMVVMStructFieldPathHelperTest Value) { }
+
+	UPROPERTY(Getter, Setter)
+	FMVVMStructFieldPathHelperTest PropertyStructWithGetterSetter;
+	UFUNCTION()
+	FMVVMStructFieldPathHelperTest GetPropertyStructWithGetterSetter() const { return FMVVMStructFieldPathHelperTest(); }
+	UFUNCTION()
+	void SetPropertyStructWithGetterSetter(FMVVMStructFieldPathHelperTest Value) {}
+
+	UPROPERTY(BlueprintReadWrite, Getter, Setter, BlueprintGetter="GetPropertyStructWithGetterSetterAndBP", BlueprintSetter="SetPropertyStructWithGetterSetterAndBP", Category = "Test")
+	FMVVMStructFieldPathHelperTest PropertyStructWithGetterSetterAndBP;
+	UFUNCTION(BlueprintPure, Category = "Test")
+	FMVVMStructFieldPathHelperTest GetPropertyStructWithGetterSetterAndBP() const { return FMVVMStructFieldPathHelperTest();}
+	UFUNCTION(BlueprintCallable, Category = "Test")
+	void SetPropertyStructWithGetterSetterAndBP(FMVVMStructFieldPathHelperTest Value){}
+
+	UPROPERTY(BlueprintReadWrite, BlueprintGetter="GetPropertyStructWithBPGetterSetter", BlueprintSetter="SetPropertyStructWithBPGetterSetter", Category = "Test")
+	FMVVMStructFieldPathHelperTest PropertyStructWithBPGetterSetter;
+	UFUNCTION(BlueprintPure, Category = "Test")
+	FMVVMStructFieldPathHelperTest GetPropertyStructWithBPGetterSetter() const { return FMVVMStructFieldPathHelperTest(); }
+	UFUNCTION(BlueprintCallable, Category = "Test")
+	void SetPropertyStructWithBPGetterSetter(FMVVMStructFieldPathHelperTest Value) { }
+
+public:
+	UPROPERTY()
+	TObjectPtr<UMVVMObjectFieldPathHelperTest> PropertyObject;
+	UFUNCTION()
+	UMVVMObjectFieldPathHelperTest* GetPropertyObject() const { return nullptr; }
+	UFUNCTION()
+	void SetPropertyObject(UMVVMObjectFieldPathHelperTest* Value) { }
+
+	UPROPERTY(Getter, Setter)
+	TObjectPtr<UMVVMObjectFieldPathHelperTest> PropertyObjectWithGetterSetter;
+	UFUNCTION()
+	UMVVMObjectFieldPathHelperTest* GetPropertyObjectWithGetterSetter() const { return nullptr; }
+	UFUNCTION()
+	void SetPropertyObjectWithGetterSetter(UMVVMObjectFieldPathHelperTest* Value) {}
+
+	UPROPERTY(BlueprintReadWrite, Getter, Setter, BlueprintGetter="GetPropertyObjectWithGetterSetterAndBP", BlueprintSetter="SetPropertyObjectWithGetterSetterAndBP", Category = "Test")
+	TObjectPtr<UMVVMObjectFieldPathHelperTest> PropertyObjectWithGetterSetterAndBP;
+	UFUNCTION(BlueprintPure, Category = "Test")
+	UMVVMObjectFieldPathHelperTest* GetPropertyObjectWithGetterSetterAndBP() const { return nullptr;}
+	UFUNCTION(BlueprintCallable, Category = "Test")
+	void SetPropertyObjectWithGetterSetterAndBP(UMVVMObjectFieldPathHelperTest* Value){}
+
+	UPROPERTY(BlueprintReadWrite, BlueprintGetter="GetPropertyObjectWithBPGetterSetter", BlueprintSetter="SetPropertyObjectWithBPGetterSetter", Category = "Test")
+	TObjectPtr<UMVVMObjectFieldPathHelperTest> PropertyObjectWithBPGetterSetter;
+	UFUNCTION(BlueprintPure, Category = "Test")
+	UMVVMObjectFieldPathHelperTest* GetPropertyObjectWithBPGetterSetter() const { return nullptr; }
+	UFUNCTION(BlueprintCallable, Category = "Test")
+	void SetPropertyObjectWithBPGetterSetter(UMVVMObjectFieldPathHelperTest* Value) { }
+
+public:
+	UPROPERTY()
+	TObjectPtr<UMVVMViewModelFieldPathHelperTest> PropertyViewModel;
+	UFUNCTION()
+	UMVVMViewModelFieldPathHelperTest* GetPropertyViewModel() const { return nullptr; }
+	UFUNCTION()
+	void SetPropertyViewModel(UMVVMViewModelFieldPathHelperTest* Value) { }
+
+	UPROPERTY(Getter, Setter)
+	TObjectPtr<UMVVMViewModelFieldPathHelperTest> PropertyViewModelWithGetterSetter;
+	UFUNCTION()
+	UMVVMViewModelFieldPathHelperTest* GetPropertyViewModelWithGetterSetter() const { return nullptr; }
+	UFUNCTION()
+	void SetPropertyViewModelWithGetterSetter(UMVVMViewModelFieldPathHelperTest* Value) {}
+
+	UPROPERTY(BlueprintReadWrite, Getter, Setter, BlueprintGetter="GetPropertyViewModelWithGetterSetterAndBP", BlueprintSetter="SetPropertyViewModelWithGetterSetterAndBP", Category = "Test")
+	TObjectPtr<UMVVMViewModelFieldPathHelperTest> PropertyViewModelWithGetterSetterAndBP;
+	UFUNCTION(BlueprintPure, Category = "Test")
+	UMVVMViewModelFieldPathHelperTest* GetPropertyViewModelWithGetterSetterAndBP() const { return nullptr;}
+	UFUNCTION(BlueprintCallable, Category = "Test")
+	void SetPropertyViewModelWithGetterSetterAndBP(UMVVMViewModelFieldPathHelperTest* Value){}
+
+	UPROPERTY(BlueprintReadWrite, BlueprintGetter="GetPropertyViewModelWithBPGetterSetter", BlueprintSetter="SetPropertyViewModelWithBPGetterSetter", Category = "Test")
+	TObjectPtr<UMVVMViewModelFieldPathHelperTest> PropertyViewModelWithBPGetterSetter;
+	UFUNCTION(BlueprintPure, Category = "Test")
+	UMVVMViewModelFieldPathHelperTest* GetPropertyViewModelWithBPGetterSetter() const { return nullptr; }
+	UFUNCTION(BlueprintCallable, Category = "Test")
+	void SetPropertyViewModelWithBPGetterSetter(UMVVMViewModelFieldPathHelperTest* Value) { }
+};
+
+/** */
+UCLASS(MinimalAPI, Transient, hidedropdown, Abstract, Hidden)
+class UMVVMViewModelFieldPathHelperTestBase : public UMVVMObjectFieldPathHelperTest, public INotifyFieldValueChanged
+{
+	GENERATED_BODY()
+
+public:
+	struct FFieldNotificationClassDescriptor : public ::UE::FieldNotification::IClassDescriptor
+	{
+	};
+
+	//~ Begin INotifyFieldValueChanged Interface
+	virtual FDelegateHandle AddFieldValueChangedDelegate(UE::FieldNotification::FFieldId InFieldId, FFieldValueChangedDelegate InNewDelegate) override final { return FDelegateHandle(); }
+	virtual bool RemoveFieldValueChangedDelegate(UE::FieldNotification::FFieldId InFieldId, FDelegateHandle InHandle) override final { return false; }
+	virtual int32 RemoveAllFieldValueChangedDelegates(const void* InUserObject) override final { return 0; }
+	virtual int32 RemoveAllFieldValueChangedDelegates(UE::FieldNotification::FFieldId InFieldId, const void* InUserObject) override final { return 0; }
+	virtual const UE::FieldNotification::IClassDescriptor& GetFieldNotificationDescriptor() const override { static FFieldNotificationClassDescriptor Tmp; return Tmp; }
+	//~ End INotifyFieldValueChanged Interface
+};
+
+/** */
+UCLASS(MinimalAPI, Transient, hidedropdown, Abstract, Hidden)
+class UMVVMViewModelFieldPathHelperTest : public UMVVMViewModelFieldPathHelperTestBase
+{
+	GENERATED_BODY()
+
+public:
+	UPROPERTY(BlueprintReadWrite, FieldNotify, Category = "Test")
+	int32 PropertyIntNotify;
+	UFUNCTION()
+	int32 GetPropertyIntNotify() const { return 0; }
+	UFUNCTION()
+	void SetPropertyIntNotify(int32 Value) { }
+
+public:
+	UPROPERTY(BlueprintReadWrite, FieldNotify, Category = "Test")
+	FVector PropertyVectorNotify;
+	UFUNCTION()
+	FVector GetPropertyVectorNotify() const { return FVector(); }
+	UFUNCTION()
+	void SetPropertyVectorNotify(FVector Value) { }
+
+public:
+	UPROPERTY(BlueprintReadWrite, FieldNotify, Category = "Test")
+	FMVVMStructFieldPathHelperTest PropertyStructNotify;
+	UFUNCTION()
+	FMVVMStructFieldPathHelperTest GetPropertyStructNotify() const { return FMVVMStructFieldPathHelperTest(); }
+	UFUNCTION()
+	void SetPropertyStructNotify(FMVVMStructFieldPathHelperTest Value) { }
+
+public:
+	UPROPERTY(BlueprintReadWrite, FieldNotify, Category = "Test")
+	TObjectPtr<UMVVMObjectFieldPathHelperTest> PropertyObjectNotify;
+	UFUNCTION()
+	UMVVMObjectFieldPathHelperTest* GetPropertyObjectNotify() const { return nullptr; }
+	UFUNCTION()
+	void SetPropertyObjectNotify(UMVVMObjectFieldPathHelperTest* Value) { }
+
+public:
+	UPROPERTY(BlueprintReadWrite, FieldNotify, Category = "Test")
+	TObjectPtr<UMVVMViewModelFieldPathHelperTest> PropertyViewModelNotify;
+	UFUNCTION()
+	UMVVMViewModelFieldPathHelperTest* GetPropertyViewModelNotify() const { return nullptr; }
+	UFUNCTION()
+	void SetPropertyViewModelNotify(UMVVMViewModelFieldPathHelperTest* Value) { }
+};
+
+/** */
+UCLASS(MinimalAPI, Transient, hidedropdown, Abstract, Hidden)
+class UMVVMWidgetFieldPathHelperTest : public UUserWidget
+{
+	GENERATED_BODY()
+
+public:
+	UPROPERTY()
+	int32 PropertyInt;
+	UFUNCTION()
+	int32 GetPropertyInt() const { return 0; }
+	UFUNCTION()
+	void SetPropertyInt(int32 Value) { }
+
+	UPROPERTY(Getter, Setter)
+	int32 PropertyIntWithGetterSetter;
+	UFUNCTION()
+	int32 GetPropertyIntWithGetterSetter() const { return 0; }
+	UFUNCTION()
+	void SetPropertyIntWithGetterSetter(int32 Value) {}
+
+	UPROPERTY(Getter, Setter, BlueprintGetter="GetPropertyIntWithGetterSetterAndBP", BlueprintSetter="SetPropertyIntWithGetterSetterAndBP", Category="Test")
+	int32 PropertyIntWithGetterSetterAndBP;
+	UFUNCTION(BlueprintPure, Category = "Test")
+	int32 GetPropertyIntWithGetterSetterAndBP() const { return 0;}
+	UFUNCTION(BlueprintCallable, Category = "Test")
+	void SetPropertyIntWithGetterSetterAndBP(int32 Value){}
+
+	UPROPERTY(BlueprintGetter="GetPropertyIntWithBPGetterSetter", BlueprintSetter="SetPropertyIntWithBPGetterSetter", Category = "Test")
+	int32 PropertyIntWithBPGetterSetter;
+	UFUNCTION(BlueprintPure, Category = "Test")
+	int32 GetPropertyIntWithBPGetterSetter() const { return 0; }
+	UFUNCTION(BlueprintCallable, Category = "Test")
+	void SetPropertyIntWithBPGetterSetter(int32 Value) { }
+
+	UPROPERTY(BlueprintReadWrite, FieldNotify, Category = "Test")
+	int32 PropertyIntNotify;
+	UFUNCTION()
+	int32 GetPropertyIntNotify() const { return 0; }
+	UFUNCTION()
+	void SetPropertyIntNotify(int32 Value) { }
+
+public:
+	UPROPERTY()
+	FVector PropertyVector;
+	UFUNCTION()
+	FVector GetPropertyVector() const { return FVector(); }
+	UFUNCTION()
+	void SetPropertyVector(FVector Value) { }
+
+	UPROPERTY(Getter, Setter)
+	FVector PropertyVectorWithGetterSetter;
+	UFUNCTION()
+	FVector GetPropertyVectorWithGetterSetter() const { return FVector(); }
+	UFUNCTION()
+	void SetPropertyVectorWithGetterSetter(FVector Value) {}
+
+	UPROPERTY(Getter, Setter, BlueprintGetter="GetPropertyVectorWithGetterSetterAndBP", BlueprintSetter="SetPropertyVectorWithGetterSetterAndBP", Category = "Test")
+	FVector PropertyVectorWithGetterSetterAndBP;
+	UFUNCTION(BlueprintPure, Category = "Test")
+	FVector GetPropertyVectorWithGetterSetterAndBP() const { return FVector();}
+	UFUNCTION(BlueprintCallable, Category = "Test")
+	void SetPropertyVectorWithGetterSetterAndBP(FVector Value){}
+
+	UPROPERTY(BlueprintGetter="GetPropertyVectorWithBPGetterSetter", BlueprintSetter="SetPropertyVectorWithBPGetterSetter", Category = "Test")
+	FVector PropertyVectorWithBPGetterSetter;
+	UFUNCTION(BlueprintPure, Category = "Test")
+	FVector GetPropertyVectorWithBPGetterSetter() const { return FVector(); }
+	UFUNCTION(BlueprintCallable, Category = "Test")
+	void SetPropertyVectorWithBPGetterSetter(FVector Value) { }
+
+	UPROPERTY(BlueprintReadWrite, FieldNotify, Category = "Test")
+	FVector PropertyVectorNotify;
+	UFUNCTION()
+	FVector GetPropertyVectorNotify() const { return FVector(); }
+	UFUNCTION()
+	void SetPropertyVectorNotify(FVector Value) { }
+
+public:
+	UPROPERTY()
+	FMVVMStructFieldPathHelperTest PropertyStruct;
+	UFUNCTION()
+	FMVVMStructFieldPathHelperTest GetPropertyStruct() const { return FMVVMStructFieldPathHelperTest(); }
+	UFUNCTION()
+	void SetPropertyStruct(FMVVMStructFieldPathHelperTest Value) { }
+
+	UPROPERTY(Getter, Setter)
+	FMVVMStructFieldPathHelperTest PropertyStructWithGetterSetter;
+	UFUNCTION()
+	FMVVMStructFieldPathHelperTest GetPropertyStructWithGetterSetter() const { return FMVVMStructFieldPathHelperTest(); }
+	UFUNCTION()
+	void SetPropertyStructWithGetterSetter(FMVVMStructFieldPathHelperTest Value) {}
+
+	UPROPERTY(BlueprintReadWrite, Getter, Setter, BlueprintGetter="GetPropertyStructWithGetterSetterAndBP", BlueprintSetter="SetPropertyStructWithGetterSetterAndBP", Category = "Test")
+	FMVVMStructFieldPathHelperTest PropertyStructWithGetterSetterAndBP;
+	UFUNCTION(BlueprintPure, Category = "Test")
+	FMVVMStructFieldPathHelperTest GetPropertyStructWithGetterSetterAndBP() const { return FMVVMStructFieldPathHelperTest();}
+	UFUNCTION(BlueprintCallable, Category = "Test")
+	void SetPropertyStructWithGetterSetterAndBP(FMVVMStructFieldPathHelperTest Value){}
+
+	UPROPERTY(BlueprintReadWrite, BlueprintGetter="GetPropertyStructWithBPGetterSetter", BlueprintSetter="SetPropertyStructWithBPGetterSetter", Category = "Test")
+	FMVVMStructFieldPathHelperTest PropertyStructWithBPGetterSetter;
+	UFUNCTION(BlueprintPure, Category = "Test")
+	FMVVMStructFieldPathHelperTest GetPropertyStructWithBPGetterSetter() const { return FMVVMStructFieldPathHelperTest(); }
+	UFUNCTION(BlueprintCallable, Category = "Test")
+	void SetPropertyStructWithBPGetterSetter(FMVVMStructFieldPathHelperTest Value) { }
+
+	UPROPERTY(BlueprintReadWrite, FieldNotify, Category = "Test")
+	FMVVMStructFieldPathHelperTest PropertyStructNotify;
+	UFUNCTION()
+	FMVVMStructFieldPathHelperTest GetPropertyStructNotify() const { return FMVVMStructFieldPathHelperTest(); }
+	UFUNCTION()
+	void SetPropertyStructNotify(FMVVMStructFieldPathHelperTest Value) { }
+
+public:
+	UPROPERTY()
+	TObjectPtr<UMVVMObjectFieldPathHelperTest> PropertyObject;
+	UFUNCTION()
+	UMVVMObjectFieldPathHelperTest* GetPropertyObject() const { return nullptr; }
+	UFUNCTION()
+	void SetPropertyObject(UMVVMObjectFieldPathHelperTest* Value) { }
+
+	UPROPERTY(Getter, Setter)
+	TObjectPtr<UMVVMObjectFieldPathHelperTest> PropertyObjectWithGetterSetter;
+	UFUNCTION()
+	UMVVMObjectFieldPathHelperTest* GetPropertyObjectWithGetterSetter() const { return nullptr; }
+	UFUNCTION()
+	void SetPropertyObjectWithGetterSetter(UMVVMObjectFieldPathHelperTest* Value) {}
+
+	UPROPERTY(BlueprintReadWrite, Getter, Setter, BlueprintGetter="GetPropertyObjectWithGetterSetterAndBP", BlueprintSetter="SetPropertyObjectWithGetterSetterAndBP", Category = "Test")
+	TObjectPtr<UMVVMObjectFieldPathHelperTest> PropertyObjectWithGetterSetterAndBP;
+	UFUNCTION(BlueprintPure, Category = "Test")
+	UMVVMObjectFieldPathHelperTest* GetPropertyObjectWithGetterSetterAndBP() const { return nullptr;}
+	UFUNCTION(BlueprintCallable, Category = "Test")
+	void SetPropertyObjectWithGetterSetterAndBP(UMVVMObjectFieldPathHelperTest* Value){}
+
+	UPROPERTY(BlueprintReadWrite, BlueprintGetter="GetPropertyObjectWithBPGetterSetter", BlueprintSetter="SetPropertyObjectWithBPGetterSetter", Category = "Test")
+	TObjectPtr<UMVVMObjectFieldPathHelperTest> PropertyObjectWithBPGetterSetter;
+	UFUNCTION(BlueprintPure, Category = "Test")
+	UMVVMObjectFieldPathHelperTest* GetPropertyObjectWithBPGetterSetter() const { return nullptr; }
+	UFUNCTION(BlueprintCallable, Category = "Test")
+	void SetPropertyObjectWithBPGetterSetter(UMVVMObjectFieldPathHelperTest* Value) { }
+
+	UPROPERTY(BlueprintReadWrite, FieldNotify, Category = "Test")
+	TObjectPtr<UMVVMObjectFieldPathHelperTest> PropertyObjectNotify;
+	UFUNCTION()
+	UMVVMObjectFieldPathHelperTest* GetPropertyObjectNotify() const { return nullptr; }
+	UFUNCTION()
+	void SetPropertyObjectNotify(UMVVMObjectFieldPathHelperTest* Value) { }
+
+public:
+	UPROPERTY()
+	TObjectPtr<UMVVMViewModelFieldPathHelperTest> PropertyViewModel;
+	UFUNCTION()
+	UMVVMViewModelFieldPathHelperTest* GetPropertyViewModel() const { return nullptr; }
+	UFUNCTION()
+	void SetPropertyViewModel(UMVVMViewModelFieldPathHelperTest* Value) { }
+
+	UPROPERTY(Getter, Setter)
+	TObjectPtr<UMVVMViewModelFieldPathHelperTest> PropertyViewModelWithGetterSetter;
+	UFUNCTION()
+	UMVVMViewModelFieldPathHelperTest* GetPropertyViewModelWithGetterSetter() const { return nullptr; }
+	UFUNCTION()
+	void SetPropertyViewModelWithGetterSetter(UMVVMViewModelFieldPathHelperTest* Value) {}
+
+	UPROPERTY(BlueprintReadWrite, Getter, Setter, BlueprintGetter="GetPropertyViewModelWithGetterSetterAndBP", BlueprintSetter="SetPropertyViewModelWithGetterSetterAndBP", Category = "Test")
+	TObjectPtr<UMVVMViewModelFieldPathHelperTest> PropertyViewModelWithGetterSetterAndBP;
+	UFUNCTION(BlueprintPure, Category = "Test")
+	UMVVMViewModelFieldPathHelperTest* GetPropertyViewModelWithGetterSetterAndBP() const { return nullptr;}
+	UFUNCTION(BlueprintCallable, Category = "Test")
+	void SetPropertyViewModelWithGetterSetterAndBP(UMVVMViewModelFieldPathHelperTest* Value){}
+
+	UPROPERTY(BlueprintReadWrite, BlueprintGetter="GetPropertyViewModelWithBPGetterSetter", BlueprintSetter="SetPropertyViewModelWithBPGetterSetter", Category = "Test")
+	TObjectPtr<UMVVMViewModelFieldPathHelperTest> PropertyViewModelWithBPGetterSetter;
+	UFUNCTION(BlueprintPure, Category = "Test")
+	UMVVMViewModelFieldPathHelperTest* GetPropertyViewModelWithBPGetterSetter() const { return nullptr; }
+	UFUNCTION(BlueprintCallable, Category = "Test")
+	void SetPropertyViewModelWithBPGetterSetter(UMVVMViewModelFieldPathHelperTest* Value) { }
+
+	UPROPERTY(BlueprintReadWrite, FieldNotify, Category = "Test")
+	TObjectPtr<UMVVMViewModelFieldPathHelperTest> PropertyViewModelNotify;
+	UFUNCTION()
+	UMVVMViewModelFieldPathHelperTest* GetPropertyViewModelNotify() const { return nullptr; }
+	UFUNCTION()
+	void SetPropertyViewModelNotify(UMVVMViewModelFieldPathHelperTest* Value) { }
+};
