Commit name change in MVVM viewmodel panel when focus is lost.

diff --git a/Engine/Plugins/Runtime/ModelViewViewModel/Source/ModelViewViewModelEditor/Private/Widgets/SMVVMViewModelPanel.cpp b/Engine/Plugins/Runtime/ModelViewViewModel/Source/ModelViewViewModelEditor/Private/Widgets/SMVVMViewModelPanel.cpp
--- a/Engine/Plugins/Runtime/ModelViewViewModel/Source/ModelViewViewModelEditor/Private/Widgets/SMVVMViewModelPanel.cpp
+++ b/Engine/Plugins/Runtime/ModelViewViewModel/Source/ModelViewViewModelEditor/Private/Widgets/SMVVMViewModelPanel.cpp
@@ -353,7 +353,7 @@ bool SMVVMViewModelPanel::HandleVerifyNameTextChanged(const FText& InText, FText
 
 void SMVVMViewModelPanel::HandleNameTextCommited(const FText& InText, ETextCommit::Type CommitInfo, FGuid ViewModelGuid)
 {
-	if (CommitInfo == ETextCommit::OnEnter)
+	if (CommitInfo == ETextCommit::OnEnter || CommitInfo == ETextCommit::OnUserMovedFocus)
 	{
 		FText OutErrorMessage;
 		RenameViewModelProperty(ViewModelGuid, InText, true, OutErrorMessage);
