Disable/enable context filtering in MVVM binding creation through project settings or a temporary checkbox

diff --git a/Engine/Plugins/Runtime/ModelViewViewModel/Source/ModelViewViewModelBlueprint/Private/MVVMWidgetBlueprintExtension_View.cpp b/Engine/Plugins/Runtime/ModelViewViewModel/Source/ModelViewViewModelBlueprint/Private/MVVMWidgetBlueprintExtension_View.cpp
--- a/Engine/Plugins/Runtime/ModelViewViewModel/Source/ModelViewViewModelBlueprint/Private/MVVMWidgetBlueprintExtension_View.cpp
+++ b/Engine/Plugins/Runtime/ModelViewViewModel/Source/ModelViewViewModelBlueprint/Private/MVVMWidgetBlueprintExtension_View.cpp
@@ -122,3 +122,19 @@ void UMVVMWidgetBlueprintExtension_View::HandleFinishCompilingClass(UWidgetBluep
 	}
 }
 
+void UMVVMWidgetBlueprintExtension_View::SetFilterSettings(FMVVMViewBindingFilterSettings InFilterSettings)
+{
+	FilterSettings = InFilterSettings;
+}
+
+#if WITH_EDITORONLY_DATA
+void UMVVMWidgetBlueprintExtension_View::PostInitProperties()
+{
+	Super::PostInitProperties();
+
+	if (!IsTemplate())
+	{
+		SetFilterSettings(GetDefault<UMVVMDeveloperProjectSettings>()->FilterSettings);
+	}
+}
+#endif
\ No newline at end of file
diff --git a/Engine/Plugins/Runtime/ModelViewViewModel/Source/ModelViewViewModelBlueprint/Public/MVVMDeveloperProjectSettings.h b/Engine/Plugins/Runtime/ModelViewViewModel/Source/ModelViewViewModelBlueprint/Public/MVVMDeveloperProjectSettings.h
--- a/Engine/Plugins/Runtime/ModelViewViewModel/Source/ModelViewViewModelBlueprint/Public/MVVMDeveloperProjectSettings.h
+++ b/Engine/Plugins/Runtime/ModelViewViewModel/Source/ModelViewViewModelBlueprint/Public/MVVMDeveloperProjectSettings.h
@@ -26,6 +26,24 @@ struct FMVVMDeveloperProjectWidgetSettings
 	TSet<FName> AdvancedFieldNames;
 };
 
+UENUM()
+enum class EFilterFlag : uint8
+{
+	None = 0,
+	All = 1 << 0
+};
+ENUM_CLASS_FLAGS(EFilterFlag)
+
+USTRUCT()
+struct FMVVMViewBindingFilterSettings
+{
+	GENERATED_BODY()
+
+	/** Filter out the properties and functions that are not valid in the context of the binding. */
+	UPROPERTY(EditAnywhere, config, Category = "Viewmodel")
+	EFilterFlag FilterFlags = EFilterFlag::None;
+};
+
 /**
  *
  */
@@ -109,6 +127,10 @@ public:
 	/** Individual class that are allowed to be uses as conversion functions. */
 	UPROPERTY(EditAnywhere, config, Category = "Viewmodel", meta = (EditCondition = "ConversionFunctionFilter == EMVVMDeveloperConversionFunctionFilterType::AllowedList"))
 	TSet<FSoftClassPath> AllowedClassForConversionFunctions;
+
+	/** Settings for filtering the list of available properties and functions on binding creation. */
+	UPROPERTY(EditAnywhere, config, Category = "Viewmodel")
+	FMVVMViewBindingFilterSettings FilterSettings;
 };
 
 #if UE_ENABLE_INCLUDE_ORDER_DEPRECATED_IN_5_2
diff --git a/Engine/Plugins/Runtime/ModelViewViewModel/Source/ModelViewViewModelBlueprint/Public/MVVMWidgetBlueprintExtension_View.h b/Engine/Plugins/Runtime/ModelViewViewModel/Source/ModelViewViewModelBlueprint/Public/MVVMWidgetBlueprintExtension_View.h
--- a/Engine/Plugins/Runtime/ModelViewViewModel/Source/ModelViewViewModelBlueprint/Public/MVVMWidgetBlueprintExtension_View.h
+++ b/Engine/Plugins/Runtime/ModelViewViewModel/Source/ModelViewViewModelBlueprint/Public/MVVMWidgetBlueprintExtension_View.h
@@ -3,6 +3,7 @@
 #pragma once
 
 #include "WidgetBlueprintExtension.h"
+#include "MVVMDeveloperProjectSettings.h"
 
 #include "MVVMWidgetBlueprintExtension_View.generated.h"
 
@@ -45,6 +46,9 @@ public:
 public:
 	//~ Begin UObject interface
 	virtual void PostLoad() override;
+#if WITH_EDITORONLY_DATA
+	virtual void PostInitProperties() override;
+#endif
 	//~ End UObject interface
 
 	//~ Begin UWidgetBlueprintExtension interface
@@ -57,10 +61,19 @@ public:
 	virtual void HandleEndCompilation() override;
 	//~ End UWidgetBlueprintExtension interface
 
+	void SetFilterSettings(FMVVMViewBindingFilterSettings InFilterSettings);
+	FMVVMViewBindingFilterSettings GetFilterSettings() const
+	{
+		return FilterSettings;
+	}
+
 private:
 	UPROPERTY(Instanced)
 	TObjectPtr<UMVVMBlueprintView> BlueprintView;
 
+	UPROPERTY(Transient)
+	FMVVMViewBindingFilterSettings FilterSettings;
+
 	FSimpleMulticastDelegate BlueprintViewChangedDelegate;
 	TPimplPtr<UE::MVVM::Private::FMVVMViewBlueprintCompiler> CurrentCompilerContext;
 };
diff --git a/Engine/Plugins/Runtime/ModelViewViewModel/Source/ModelViewViewModelEditor/Private/ModelViewViewModelEditorModule.cpp b/Engine/Plugins/Runtime/ModelViewViewModel/Source/ModelViewViewModelEditor/Private/ModelViewViewModelEditorModule.cpp
--- a/Engine/Plugins/Runtime/ModelViewViewModel/Source/ModelViewViewModelEditor/Private/ModelViewViewModelEditorModule.cpp
+++ b/Engine/Plugins/Runtime/ModelViewViewModel/Source/ModelViewViewModelEditor/Private/ModelViewViewModelEditorModule.cpp
@@ -11,6 +11,7 @@
 #include "Kismet2/BlueprintEditorUtils.h"
 #include "MessageLogModule.h"
 #include "MVVMBlueprintView.h"
+#include "MVVMDeveloperProjectSettings.h"
 #include "MVVMEditorCommands.h"
 #include "MVVMEditorSubsystem.h"
 #include "MVVMWidgetBlueprintExtension_View.h"
@@ -122,6 +123,14 @@ void FModelViewViewModelEditorModule::HandleRegisterBlueprintEditorTab(const FWi
 			ApplicationMode.OnPostActivateMode.AddRaw(this, &FModelViewViewModelEditorModule::HandleActivateMode);
 			ApplicationMode.OnPreDeactivateMode.AddRaw(this, &FModelViewViewModelEditorModule::HandleDeactiveMode);
 		}
+
+		if (TSharedPtr<FWidgetBlueprintEditor> BP = ApplicationMode.GetBlueprintEditor())
+		{
+			if (UMVVMWidgetBlueprintExtension_View* ExtensionView = UMVVMWidgetBlueprintExtension_View::GetExtension<UMVVMWidgetBlueprintExtension_View>(BP->GetWidgetBlueprintObj()))
+			{
+				ExtensionView->SetFilterSettings(GetDefault<UMVVMDeveloperProjectSettings>()->FilterSettings);
+			}
+		}
 	}
 }
 
diff --git a/Engine/Plugins/Runtime/ModelViewViewModel/Source/ModelViewViewModelEditor/Private/Widgets/SMVVMFieldSelector.cpp b/Engine/Plugins/Runtime/ModelViewViewModel/Source/ModelViewViewModelEditor/Private/Widgets/SMVVMFieldSelector.cpp
--- a/Engine/Plugins/Runtime/ModelViewViewModel/Source/ModelViewViewModelEditor/Private/Widgets/SMVVMFieldSelector.cpp
+++ b/Engine/Plugins/Runtime/ModelViewViewModel/Source/ModelViewViewModelEditor/Private/Widgets/SMVVMFieldSelector.cpp
@@ -7,12 +7,15 @@
 #include "Framework/Views/TableViewMetadata.h"
 #include "Hierarchy/SReadOnlyHierarchyView.h"
 #include "MVVMEditorSubsystem.h"
+#include "MVVMDeveloperProjectSettings.h"
+#include "MVVMWidgetBlueprintExtension_View.h"
 #include "SPrimaryButton.h"
 #include "Styling/MVVMEditorStyle.h"
 #include "Styling/SlateIconFinder.h"
 #include "Types/MVVMBindingMode.h"
 #include "Widgets/SMVVMFieldEntry.h"
 #include "Widgets/SMVVMSourceEntry.h"
+#include "Widgets/Input/SCheckBox.h"
 #include "Widgets/Input/SComboButton.h"
 #include "Widgets/Input/SSearchBox.h"
 #include "Widgets/Layout/SScrollBox.h"
@@ -865,8 +868,28 @@ TSharedRef<SWidget> SFieldSelector::OnGetMenuContent()
 		.Padding(0, 4, 0, 4)
 		.AutoHeight()
 		[
-			SAssignNew(SearchBox, SSearchBox)
-			.OnTextChanged(this, &SFieldSelector::OnSearchBoxTextChanged)
+			SNew(SHorizontalBox)
+			+ SHorizontalBox::Slot()
+			.FillWidth(1.f)
+			.Padding(3.f, 0.f, 3.f, 0.f)	
+			[
+				SAssignNew(SearchBox, SSearchBox)
+				.OnTextChanged(this, &SFieldSelector::OnSearchBoxTextChanged)
+			]
+			+ SHorizontalBox::Slot()
+			.HAlign(HAlign_Right)
+			.VAlign(VAlign_Center)
+			.AutoWidth()
+			[
+				SNew(SCheckBox)
+				.OnCheckStateChanged(this, &SFieldSelector::OnContextToggleChanged)
+				.IsChecked(this, &SFieldSelector::ContextToggleIsChecked)
+				.ToolTipText(LOCTEXT("ContextFlagToolTip", "Should the list be filtered to only properties and functions that make sense in the current context?"))
+				[
+					SNew(STextBlock)
+					.Text(LOCTEXT("MVVMContextSensitiveToggle", "Context Sensitive"))
+				]
+			]
 		];
 
 	// If we're showing conversion functions, we don't want to set the AssignableTo property of SSourceBindingList, because then it will only show exact matches, 
@@ -1057,6 +1080,34 @@ EFieldVisibility SFieldSelector::GetFieldVisibilityFlags() const
 	return Flags;
 }
 
+void SFieldSelector::OnContextToggleChanged(ECheckBoxState CheckState)
+{
+	UMVVMWidgetBlueprintExtension_View* ExtensionView = UMVVMWidgetBlueprintExtension_View::GetExtension<UMVVMWidgetBlueprintExtension_View>(WidgetBlueprint.Get());
+	check(ExtensionView);
+
+	FMVVMViewBindingFilterSettings FilterSettings = ExtensionView->GetFilterSettings();
+	FilterSettings.FilterFlags = CheckState == ECheckBoxState::Checked ? EFilterFlag::All : EFilterFlag::None;
+	ExtensionView->SetFilterSettings(FilterSettings);
+
+	if (WidgetList.IsValid())
+	{
+		OnWidgetSelected(WidgetBlueprint->GetFName(), ESelectInfo::Direct);
+	}
+
+	if (ViewModelList.IsValid())
+	{
+		OnViewModelSelected(FBindingSource(), ESelectInfo::Direct);
+	}
+}
+
+ECheckBoxState SFieldSelector::ContextToggleIsChecked() const
+{
+	UMVVMWidgetBlueprintExtension_View* ExtensionView = UMVVMWidgetBlueprintExtension_View::GetExtension<UMVVMWidgetBlueprintExtension_View>(WidgetBlueprint.Get());
+	check(ExtensionView);
+
+	return ExtensionView->GetFilterSettings().FilterFlags == EFilterFlag::All ? ECheckBoxState::Checked : ECheckBoxState::Unchecked;
+}
+
 TSharedPtr<SFieldSelector::FConversionFunctionItem> SFieldSelector::FindOrCreateItemForCategory(TArray<TSharedPtr<FConversionFunctionItem>>& Items, TArrayView<FString> CategoryPath)
 {
 	check(CategoryPath.Num() > 0);
diff --git a/Engine/Plugins/Runtime/ModelViewViewModel/Source/ModelViewViewModelEditor/Private/Widgets/SMVVMFieldSelector.h b/Engine/Plugins/Runtime/ModelViewViewModel/Source/ModelViewViewModelEditor/Private/Widgets/SMVVMFieldSelector.h
--- a/Engine/Plugins/Runtime/ModelViewViewModel/Source/ModelViewViewModelEditor/Private/Widgets/SMVVMFieldSelector.h
+++ b/Engine/Plugins/Runtime/ModelViewViewModel/Source/ModelViewViewModelEditor/Private/Widgets/SMVVMFieldSelector.h
@@ -101,6 +101,9 @@ private:
 
 	void OnSearchBoxTextChanged(const FText& NewText);
 
+	void OnContextToggleChanged(ECheckBoxState CheckState);
+	ECheckBoxState ContextToggleIsChecked() const;
+
 	struct FConversionFunctionItem
 	{
 		FString GetCategoryName() { return CategoryPath.Num() > 0 ? CategoryPath.Last() : FString(); }
diff --git a/Engine/Plugins/Runtime/ModelViewViewModel/Source/ModelViewViewModelEditor/Private/Widgets/SMVVMViewModelBindingListWidget.cpp b/Engine/Plugins/Runtime/ModelViewViewModel/Source/ModelViewViewModelEditor/Private/Widgets/SMVVMViewModelBindingListWidget.cpp
--- a/Engine/Plugins/Runtime/ModelViewViewModel/Source/ModelViewViewModelEditor/Private/Widgets/SMVVMViewModelBindingListWidget.cpp
+++ b/Engine/Plugins/Runtime/ModelViewViewModel/Source/ModelViewViewModelEditor/Private/Widgets/SMVVMViewModelBindingListWidget.cpp
@@ -8,6 +8,7 @@
 #include "MVVMBlueprintViewModelContext.h"
 #include "MVVMDeveloperProjectSettings.h"
 #include "MVVMSubsystem.h"
+#include "MVVMWidgetBlueprintExtension_View.h"
 #include "PropertyPermissionList.h"
 #include "Styling/MVVMEditorStyle.h"
 #include "Types/MVVMAvailableBinding.h"
@@ -170,9 +171,28 @@ TArray<FFieldVariant> FFieldIterator_Bindable::GetFields(const UStruct* Struct)
 	auto AddResult = [this, &Result, Struct](const TArray<FMVVMAvailableBinding>& AvailableBindingsList)
 	{
 		Result.Reserve(AvailableBindingsList.Num());
+
+		EFilterFlag FilterFlags;
+		if (UMVVMWidgetBlueprintExtension_View* ExtensionView = UMVVMWidgetBlueprintExtension_View::GetExtension<UMVVMWidgetBlueprintExtension_View>(WidgetBlueprint.Get()))
+		{
+			FilterFlags = ExtensionView->GetFilterSettings().FilterFlags;
+		}
+		else
+		{
+			FilterFlags = GetDefault<UMVVMDeveloperProjectSettings>()->FilterSettings.FilterFlags;
+		}
+
 		for (const FMVVMAvailableBinding& Value : AvailableBindingsList)
 		{
-			TOptional<FFieldVariant> PassResult = Private::PassFilter(Value, Struct, FieldVisibilityFlags, AssignableTo, false);
+			TOptional<FFieldVariant> PassResult;
+			if (FilterFlags == EFilterFlag::All)
+			{
+				PassResult = Private::PassFilter(Value, Struct, FieldVisibilityFlags, AssignableTo, false);
+			}
+			else
+			{
+				PassResult = Private::PassFilter(Value, Struct, EFieldVisibility::None, nullptr, false);
+			}
 			if (PassResult.IsSet())
 			{
 				Result.Add(MoveTemp(PassResult.GetValue()));
@@ -507,7 +527,24 @@ FMVVMBlueprintPropertyPath SSourceBindingList::CreateBlueprintPropertyPath(SProp
 				FMVVMAvailableBinding Binding = UMVVMSubsystem::GetAvailableBinding(OwnerClass, FMVVMBindingName(FieldName), AccessorClass);
 				if (Binding.IsValid())
 				{
-					bPassFilter = Private::PassFilter(Binding, OwnerClass, FieldIterator->GetFieldVisibilityFlags(), FieldIterator->GetAssignableTo(), true).IsSet();
+					EFilterFlag FilterFlags;
+					if (UMVVMWidgetBlueprintExtension_View* ExtensionView = UMVVMWidgetBlueprintExtension_View::GetExtension<UMVVMWidgetBlueprintExtension_View>(WidgetBlueprint.Get()))
+					{
+						FilterFlags = ExtensionView->GetFilterSettings().FilterFlags;
+					}
+					else
+					{
+						FilterFlags = GetDefault<UMVVMDeveloperProjectSettings>()->FilterSettings.FilterFlags;
+					}
+
+					if (FilterFlags == EFilterFlag::All)
+					{
+						bPassFilter = Private::PassFilter(Binding, OwnerClass, FieldIterator->GetFieldVisibilityFlags(), FieldIterator->GetAssignableTo(), true).IsSet();
+					}
+					else
+					{
+						bPassFilter = Private::PassFilter(Binding, OwnerClass, EFieldVisibility::None, nullptr, true).IsSet();
+					}
 				}
 				break;
 			}
@@ -538,7 +575,25 @@ FMVVMBlueprintPropertyPath SSourceBindingList::CreateBlueprintPropertyPath(SProp
 		FMVVMAvailableBinding Binding = UMVVMSubsystem::GetAvailableBinding(AccessorClass, BindingName, AccessorClass);
 		if (Binding.IsValid())
 		{
-			bool bPassFilter = Private::PassFilter(Binding, AccessorClass, FieldIterator->GetFieldVisibilityFlags(), FieldIterator->GetAssignableTo(), true).IsSet();
+			EFilterFlag FilterFlags;
+			if (UMVVMWidgetBlueprintExtension_View* ExtensionView = UMVVMWidgetBlueprintExtension_View::GetExtension<UMVVMWidgetBlueprintExtension_View>(WidgetBlueprint.Get()))
+			{
+				FilterFlags = ExtensionView->GetFilterSettings().FilterFlags;
+			}
+			else
+			{
+				FilterFlags = GetDefault<UMVVMDeveloperProjectSettings>()->FilterSettings.FilterFlags;
+			}
+			bool bPassFilter;
+
+			if (FilterFlags == EFilterFlag::All)
+			{
+				bPassFilter = Private::PassFilter(Binding, AccessorClass, FieldIterator->GetFieldVisibilityFlags(), FieldIterator->GetAssignableTo(), true).IsSet();
+			}
+			else
+			{
+				bPassFilter = Private::PassFilter(Binding, AccessorClass, EFieldVisibility::None, nullptr, true).IsSet();
+			}
 			if (bPassFilter)
 			{
 				if (Source->Key.ViewModelId.IsValid())
