MVVM: Fix false positive V758 error in MVVM code base. Fix bad variable initialization.
#rnx

diff --git a/Engine/Plugins/Runtime/ModelViewViewModel/Source/ModelViewViewModel/Private/View/MVVMView.cpp b/Engine/Plugins/Runtime/ModelViewViewModel/Source/ModelViewViewModel/Private/View/MVVMView.cpp
--- a/Engine/Plugins/Runtime/ModelViewViewModel/Source/ModelViewViewModel/Private/View/MVVMView.cpp
+++ b/Engine/Plugins/Runtime/ModelViewViewModel/Source/ModelViewViewModel/Private/View/MVVMView.cpp
@@ -193,7 +193,8 @@ bool UMVVMView::EvaluateSourceCreator(int32 SourceCreatorIndex)
 	}
 
 	UUserWidget* UserWidget = GetUserWidget();
-	const FMVVMViewClass_SourceCreator& SourceCreator = ClassExtension->GetViewModelCreators()[SourceCreatorIndex];
+	const TArrayView<const FMVVMViewClass_SourceCreator> ViewModelCreators = ClassExtension->GetViewModelCreators();
+	const FMVVMViewClass_SourceCreator& SourceCreator = ViewModelCreators[SourceCreatorIndex];
 	UObject* NewSource = SourceCreator.CreateInstance(ClassExtension, this, UserWidget);
 	ensureMsgf((NewSource == nullptr || NewSource->GetClass()->ImplementsInterface(UNotifyFieldValueChanged::StaticClass())), TEXT("The source has implement the interface. It should be check at compile time."));
 	return SetSourceInternal(SourceCreator.GetSourceName(), NewSource, true);
diff --git a/Engine/Plugins/Runtime/ModelViewViewModel/Source/ModelViewViewModel/Public/Bindings/MVVMCompiledBindingLibrary.h b/Engine/Plugins/Runtime/ModelViewViewModel/Source/ModelViewViewModel/Public/Bindings/MVVMCompiledBindingLibrary.h
--- a/Engine/Plugins/Runtime/ModelViewViewModel/Source/ModelViewViewModel/Public/Bindings/MVVMCompiledBindingLibrary.h
+++ b/Engine/Plugins/Runtime/ModelViewViewModel/Source/ModelViewViewModel/Public/Bindings/MVVMCompiledBindingLibrary.h
@@ -249,7 +249,7 @@ private:
 	};
 
 	UPROPERTY()
-	uint8 Flags;
+	uint8 Flags = 0;
 
 #if WITH_EDITORONLY_DATA
 	UPROPERTY(meta = (IgnoreForMemberInitializationTest))
diff --git a/Engine/Plugins/Runtime/ModelViewViewModel/Source/ModelViewViewModelEditor/Private/Tests/MVVMFieldPathHelperTest.h b/Engine/Plugins/Runtime/ModelViewViewModel/Source/ModelViewViewModelEditor/Private/Tests/MVVMFieldPathHelperTest.h
--- a/Engine/Plugins/Runtime/ModelViewViewModel/Source/ModelViewViewModelEditor/Private/Tests/MVVMFieldPathHelperTest.h
+++ b/Engine/Plugins/Runtime/ModelViewViewModel/Source/ModelViewViewModelEditor/Private/Tests/MVVMFieldPathHelperTest.h
@@ -20,10 +20,10 @@ struct FMVVMStructFieldPathHelperTest
 
 public:
 	UPROPERTY()
-	int32 PropertyInt;
+	int32 PropertyInt = 0;
 
 	UPROPERTY()
-	FVector PropertyVector;
+	FVector PropertyVector = FVector::ZeroVector;
 
 	UPROPERTY()
 	TObjectPtr<UMVVMObjectFieldPathHelperTest> PropertyObject;
