MVVM: Fix false positive V758 error in MVVM code base. Fix bad variable initialization.
#rnx

diff --git a/Engine/Plugins/Runtime/ModelViewViewModel/Source/ModelViewViewModel/Private/View/MVVMView.cpp b/Engine/Plugins/Runtime/ModelViewViewModel/Source/ModelViewViewModel/Private/View/MVVMView.cpp
--- a/Engine/Plugins/Runtime/ModelViewViewModel/Source/ModelViewViewModel/Private/View/MVVMView.cpp
+++ b/Engine/Plugins/Runtime/ModelViewViewModel/Source/ModelViewViewModel/Private/View/MVVMView.cpp
@@ -229,7 +229,8 @@ bool UMVVMView::SetSourceInternal(FName ViewModelName, TScriptInterface<INotifyF
 		return false;
 	}
 
-	const FMVVMViewClass_SourceCreator& SourceCreator = ClassExtension->GetViewModelCreators()[AllCreatedSourcesIndex];
+	const TArrayView<const FMVVMViewClass_SourceCreator> ViewModelCreators = ClassExtension->GetViewModelCreators();
+	const FMVVMViewClass_SourceCreator& SourceCreator = ViewModelCreators[AllCreatedSourcesIndex];
 
 	if (bForDynamicSource && !SourceCreator.CanBeEvaluated())
 	{
