MVVM:
- Fix wrong equality test in UMVVMViewModelBase::SetPropertyValue

diff --git a/Engine/Plugins/Runtime/ModelViewViewModel/Source/ModelViewViewModel/Public/MVVMViewModelBase.h b/Engine/Plugins/Runtime/ModelViewViewModel/Source/ModelViewViewModel/Public/MVVMViewModelBase.h
--- a/Engine/Plugins/Runtime/ModelViewViewModel/Source/ModelViewViewModel/Public/MVVMViewModelBase.h
+++ b/Engine/Plugins/Runtime/ModelViewViewModel/Source/ModelViewViewModel/Public/MVVMViewModelBase.h
@@ -84,7 +84,7 @@ protected:
 
 	bool SetPropertyValue(FText& Value, const FText& NewValue, UE::FieldNotification::FFieldId FieldId)
 	{
-		if (Value.IdenticalTo(Value))
+		if (Value.IdenticalTo(NewValue))
 		{
 			return false;
 		}
