Overload SetViewModel function to use the type instead of FName to search for a source viewmodel (if there are multiple, fails to set)

diff --git a/Engine/Plugins/Runtime/ModelViewViewModel/Source/ModelViewViewModel/Private/View/MVVMView.cpp b/Engine/Plugins/Runtime/ModelViewViewModel/Source/ModelViewViewModel/Private/View/MVVMView.cpp
--- a/Engine/Plugins/Runtime/ModelViewViewModel/Source/ModelViewViewModel/Private/View/MVVMView.cpp
+++ b/Engine/Plugins/Runtime/ModelViewViewModel/Source/ModelViewViewModel/Private/View/MVVMView.cpp
@@ -183,6 +183,51 @@ bool UMVVMView::SetViewModel(FName ViewModelName, TScriptInterface<INotifyFieldV
 }
 
 
+bool UMVVMView::SetViewModelByClass(TScriptInterface<INotifyFieldValueChanged> NewValue)
+{
+	if (!NewValue.GetObject())
+	{
+		UE::MVVM::FMessageLog Log(GetUserWidget());
+		Log.Error(LOCTEXT("SetViewModelInvalidObject", "The new viewmodel is null."));
+		return false;
+	}
+
+	if (ClassExtension == nullptr)
+	{
+		UE::MVVM::FMessageLog Log(GetUserWidget());
+		Log.Error(LOCTEXT("SetViewModelInvalidClass", "The view is not constructed."));
+		return false;
+	}
+
+	FName ViewModelName;
+	for (const FMVVMViewClass_SourceCreator& SourceCreator : ClassExtension->GetViewModelCreators())
+	{
+		if (NewValue.GetObject()->GetClass()->IsChildOf(SourceCreator.GetSourceClass()))
+		{
+			if (ViewModelName.IsNone())
+			{
+				ViewModelName = SourceCreator.GetSourceName();
+			}
+			else
+			{
+				UE::MVVM::FMessageLog Log(GetUserWidget());
+				Log.Error(FText::Format(LOCTEXT("SetViewModelViewModelNotUnique", "More than one viewmodel match the type of the passed instance. Make sure there exists only one viewmodel of type {0} in the widget blueprint, or try SetViewModel by name."), FText::FromName(NewValue.GetObject()->GetClass()->GetFName())));
+				return false;
+			}
+		}
+	}
+
+	if (ViewModelName.IsNone())
+	{
+		UE::MVVM::FMessageLog Log(GetUserWidget());
+		Log.Error(LOCTEXT("SetViewModelViewModelNotFound", "A created viewmodel matching the class of the passed instance could not be found."));
+		return false;
+	}
+
+	return SetViewModel(ViewModelName, NewValue);
+}
+
+
 bool UMVVMView::EvaluateSourceCreator(int32 SourceCreatorIndex)
 {
 	if (!ClassExtension->GetViewModelCreators().IsValidIndex(SourceCreatorIndex))
diff --git a/Engine/Plugins/Runtime/ModelViewViewModel/Source/ModelViewViewModel/Public/View/MVVMView.h b/Engine/Plugins/Runtime/ModelViewViewModel/Source/ModelViewViewModel/Public/View/MVVMView.h
--- a/Engine/Plugins/Runtime/ModelViewViewModel/Source/ModelViewViewModel/Public/View/MVVMView.h
+++ b/Engine/Plugins/Runtime/ModelViewViewModel/Source/ModelViewViewModel/Public/View/MVVMView.h
@@ -91,6 +91,14 @@ public:
 	UFUNCTION(BlueprintCallable, Category = "Viewmodel")
 	bool SetViewModel(FName ViewModelName, TScriptInterface<INotifyFieldValueChanged> ViewModel);
 
+	/**
+	 * Set the first viewmodel matching the exact specified type. If none is found, set the first viewmodel matching a child of the specified type
+	 * The viewmodel needs to be settable and it should have a valid name.
+	 * If the view is initialized, all bindings that uses that viewmodel will be re-executed with the new viewmodel instance.
+	 */
+	UFUNCTION(BlueprintCallable, Category = "Viewmodel")
+	bool SetViewModelByClass(TScriptInterface<INotifyFieldValueChanged> NewValue);
+
 private:
 	bool EvaluateSourceCreator(int32 SourceIndex);
 	bool SetSourceInternal(FName ViewModelName, TScriptInterface<INotifyFieldValueChanged> ViewModel, bool bForDynamicSource);
