MVVM: Remove warning when building the viewmodel setter graph. The cast not is not needed anymore.

diff --git a/Engine/Plugins/Runtime/ModelViewViewModel/Source/ModelViewViewModelBlueprint/Private/MVVMFunctionGraphHelper.cpp b/Engine/Plugins/Runtime/ModelViewViewModel/Source/ModelViewViewModelBlueprint/Private/MVVMFunctionGraphHelper.cpp
--- a/Engine/Plugins/Runtime/ModelViewViewModel/Source/ModelViewViewModelBlueprint/Private/MVVMFunctionGraphHelper.cpp
+++ b/Engine/Plugins/Runtime/ModelViewViewModel/Source/ModelViewViewModelBlueprint/Private/MVVMFunctionGraphHelper.cpp
@@ -154,12 +154,6 @@ bool GenerateViewModelSetter(FKismetCompilerContext& InContext, UEdGraph* InFunc
 			Schema->TrySetDefaultObject(*ExtensionTypePin, UMVVMView::StaticClass(), bMarkAsModified);
 		}
 	}
-	UK2Node_DynamicCast* DynCastNode = InContext.SpawnIntermediateNode<UK2Node_DynamicCast>(FunctionEntry, InFunctionGraph);
-	{
-		DynCastNode->TargetType = UMVVMView::StaticClass();
-		DynCastNode->SetPurity(true);
-		DynCastNode->AllocateDefaultPins();
-	}
 	// Entry -> SetViewModel
 	{
 		UEdGraphPin* ThenPin = FunctionEntry->FindPin(UEdGraphSchema_K2::PN_Then, EGPD_Output);
@@ -176,19 +170,9 @@ bool GenerateViewModelSetter(FKismetCompilerContext& InContext, UEdGraph* InFunc
 		}
 		bResult = bResult && ThenPin && ExecPin && EntryViewModelPin && SetViewModelPin;
 	}
-	// GetExtension -> Cast
+	// GetExtension -> SetViewModel
 	{
 		UEdGraphPin* ResultPin = CallGetExtensionlNode->FindPin(UEdGraphSchema_K2::PN_ReturnValue, EGPD_Output);
-		UEdGraphPin* ObjectPin = DynCastNode->GetCastSourcePin();
-		if (ensure(ResultPin && ObjectPin))
-		{
-			ensure(Schema->TryCreateConnection(ResultPin, ObjectPin));
-		}
-		bResult = bResult && ResultPin && ObjectPin;
-	}
-	// Cast -> SetViewModel
-	{
-		UEdGraphPin* ResultPin = DynCastNode->GetCastResultPin();
 		UEdGraphPin* ObjectPin = CallSetViewModelNode->FindPin(UEdGraphSchema_K2::PN_Self, EGPD_Input);
 		if (ensure(ResultPin && ObjectPin))
 		{
