MVVM: Fix bad logic in SetViewModel.

diff --git a/Engine/Plugins/Runtime/ModelViewViewModel/Source/ModelViewViewModel/Private/View/MVVMView.cpp b/Engine/Plugins/Runtime/ModelViewViewModel/Source/ModelViewViewModel/Private/View/MVVMView.cpp
--- a/Engine/Plugins/Runtime/ModelViewViewModel/Source/ModelViewViewModel/Private/View/MVVMView.cpp
+++ b/Engine/Plugins/Runtime/ModelViewViewModel/Source/ModelViewViewModel/Private/View/MVVMView.cpp
@@ -324,7 +324,7 @@ bool UMVVMView::SetSourceInternal(FName ViewModelName, TScriptInterface<INotifyF
 		}
 
 		ViewSource.Source = NewValue.GetObject();
-		ViewSource.bSetManually = bForDynamicSource;
+		ViewSource.bSetManually = !bForDynamicSource;
 		if (SourceCreator.IsSourceAUserWidgetProperty())
 		{
 			FObjectPropertyBase* FoundObjectProperty = FindFProperty<FObjectPropertyBase>(GetUserWidget()->GetClass(), ViewModelName);
