MVVM: Add option to not add the compiled result to the final GeneratedClass. Used for property validation.

diff --git a/Engine/Plugins/Runtime/ModelViewViewModel/Source/ModelViewViewModelBlueprint/Private/MVVMWidgetBlueprintExtension_View.cpp b/Engine/Plugins/Runtime/ModelViewViewModel/Source/ModelViewViewModelBlueprint/Private/MVVMWidgetBlueprintExtension_View.cpp
--- a/Engine/Plugins/Runtime/ModelViewViewModel/Source/ModelViewViewModelBlueprint/Private/MVVMWidgetBlueprintExtension_View.cpp
+++ b/Engine/Plugins/Runtime/ModelViewViewModel/Source/ModelViewViewModelBlueprint/Private/MVVMWidgetBlueprintExtension_View.cpp
@@ -8,6 +8,17 @@
 #include UE_INLINE_GENERATED_CPP_BY_NAME(MVVMWidgetBlueprintExtension_View)
 
 
+namespace UE::MVVM::Private
+{
+	bool GAllowViewClass = true;
+	static FAutoConsoleVariableRef CVarAllowViewClass(
+		TEXT("MVVM.AllowViewClass"),
+		GAllowViewClass,
+		TEXT("Is the model view viewmodel view is allowed to be added to the generated Widget GeneratedClass."),
+		ECVF_ReadOnly
+	);
+}
+
 void UMVVMWidgetBlueprintExtension_View::CreateBlueprintViewInstance()
 {
 	BlueprintView = NewObject<UMVVMBlueprintView>(this, FName(), RF_Transactional);
@@ -110,7 +121,7 @@ void UMVVMWidgetBlueprintExtension_View::HandleFinishCompilingClass(UWidgetBluep
 
 		CurrentCompilerContext->CleanTemporaries(Class);
 
-		if (bCompiled)
+		if (bCompiled && UE::MVVM::Private::GAllowViewClass)
 		{
 			check(ViewExtension);
 			CurrentCompilerContext->AddExtension(Class, ViewExtension);
