MVVM: Test if the property handle has a valid value before fetching the data.

diff --git a/Engine/Plugins/Runtime/ModelViewViewModel/Source/ModelViewViewModelEditor/Private/Customizations/MVVMPropertyPathCustomization.cpp b/Engine/Plugins/Runtime/ModelViewViewModel/Source/ModelViewViewModelEditor/Private/Customizations/MVVMPropertyPathCustomization.cpp
--- a/Engine/Plugins/Runtime/ModelViewViewModel/Source/ModelViewViewModelEditor/Private/Customizations/MVVMPropertyPathCustomization.cpp
+++ b/Engine/Plugins/Runtime/ModelViewViewModel/Source/ModelViewViewModelEditor/Private/Customizations/MVVMPropertyPathCustomization.cpp
@@ -205,35 +205,13 @@ FBindingSource FPropertyPathCustomization::OnGetSelectedSource() const
 
 FMVVMBlueprintPropertyPath FPropertyPathCustomization::OnGetSelectedField() const
 {
-	// TODO: This crashes when removing the viewmodel that the selected binding is pointed at
-	TArray<void*> RawData;
-	PropertyHandle->AccessRawData(RawData);
-
-	FMVVMBlueprintPropertyPath SelectedField;
-
-	bool bFirst = true;
-
-	for (void* RawPtr : RawData)
+	FMVVMBlueprintPropertyPath PropertyPath;
+	void* Buffer = reinterpret_cast<void*>(&PropertyPath);
+	if (PropertyHandle->GetValueData(Buffer) == FPropertyAccess::Success)
 	{
-		if (RawPtr == nullptr)
-		{
-			continue;
-		}
-
-		FMVVMBlueprintPropertyPath* PropertyPath = (FMVVMBlueprintPropertyPath*) RawPtr;
-		if (bFirst)
-		{
-			SelectedField = *PropertyPath;
-		}
-		else if (SelectedField != *PropertyPath)
-		{
-			SelectedField = FMVVMBlueprintPropertyPath();
-			break;
-		}
-		bFirst = false;
+		return PropertyPath;
 	}
-
-	return SelectedField;
+	return FMVVMBlueprintPropertyPath();
 }
 
 } // namespace UE::MVVM
