MVVM: Fix SetViewModel that didn't trigger the child viewmodel binding when changed. Use QAGame UMG_MVVM_LongPath to test. The order the sources are added is important to make sure the bindings are executed in the correct order.

diff --git a/Engine/Plugins/Runtime/ModelViewViewModel/Source/ModelViewViewModel/Private/View/MVVMView.cpp b/Engine/Plugins/Runtime/ModelViewViewModel/Source/ModelViewViewModel/Private/View/MVVMView.cpp
--- a/Engine/Plugins/Runtime/ModelViewViewModel/Source/ModelViewViewModel/Private/View/MVVMView.cpp
+++ b/Engine/Plugins/Runtime/ModelViewViewModel/Source/ModelViewViewModel/Private/View/MVVMView.cpp
@@ -337,18 +337,32 @@ bool UMVVMView::SetSourceInternal(FName ViewModelName, TScriptInterface<INotifyF
 
 		bool bPreviousEveryTickBinding = bHasEveryTickBinding;
 		bHasEveryTickBinding = false;
-		if (NewValue.GetObject())
+		// Register back any binding that was previously enabled
+		if (NewValue.GetObject() && bInitialized)
 		{
-			// Register back any binding that was previously enabled
-			if (bInitialized)
+			// Enabled the default bindings
+			for (int32 Index = 0; Index < CompiledBindings.Num(); ++Index)
 			{
-				// Enabled the default bindings
-				for (int32 Index = 0; Index < CompiledBindings.Num(); ++Index)
+				const FMVVMViewClass_CompiledBinding& Binding = CompiledBindings[Index];
+				if (Binding.IsEnabledByDefault())
 				{
-					const FMVVMViewClass_CompiledBinding& Binding = CompiledBindings[Index];
-					if (Binding.IsEnabledByDefault() && Binding.GetSourceName() == ViewModelName)
+					// Binding on this viewmodel
+					if (Binding.GetSourceName() == ViewModelName)
 					{
 						EnableLibraryBinding(Binding, Index);
+						// Bindings that depends on this binding
+						if (Binding.IsEvaluateSourceCreatorBinding())
+						{
+							int32 ParentSourceIndex = ClassExtension->GetViewModelCreators().IndexOfByPredicate([ViewModelName](const FMVVMViewClass_SourceCreator& Item)
+								{
+									return Item.GetParentSourceName() == ViewModelName;
+								});
+							if (ParentSourceIndex != INDEX_NONE)
+							{
+								check(ParentSourceIndex != AllCreatedSourcesIndex)
+								EvaluateSourceCreator(ParentSourceIndex); // it will deactivate previous binding 
+							}
+						}
 					}
 				}
 			}
diff --git a/Engine/Plugins/Runtime/ModelViewViewModel/Source/ModelViewViewModel/Private/View/MVVMViewClass.cpp b/Engine/Plugins/Runtime/ModelViewViewModel/Source/ModelViewViewModel/Private/View/MVVMViewClass.cpp
--- a/Engine/Plugins/Runtime/ModelViewViewModel/Source/ModelViewViewModel/Private/View/MVVMViewClass.cpp
+++ b/Engine/Plugins/Runtime/ModelViewViewModel/Source/ModelViewViewModel/Private/View/MVVMViewClass.cpp
@@ -148,6 +148,12 @@ FString FMVVMViewClass_SourceCreator::ToString(const FMVVMCompiledBindingLibrary
 	StringBuilder << TEXT(", SourceName: ");
 	StringBuilder << GetSourceName();
 
+	if (!ParentSourceName.IsNone())
+	{
+		StringBuilder << TEXT(", ParentSourceName: ");
+		StringBuilder << GetParentSourceName();
+	}
+
 	if (Args.bAddFlags)
 	{
 		bool bHasFlag = false;
diff --git a/Engine/Plugins/Runtime/ModelViewViewModel/Source/ModelViewViewModel/Public/View/MVVMViewClass.h b/Engine/Plugins/Runtime/ModelViewViewModel/Source/ModelViewViewModel/Public/View/MVVMViewClass.h
--- a/Engine/Plugins/Runtime/ModelViewViewModel/Source/ModelViewViewModel/Public/View/MVVMViewClass.h
+++ b/Engine/Plugins/Runtime/ModelViewViewModel/Source/ModelViewViewModel/Public/View/MVVMViewClass.h
@@ -61,6 +61,11 @@ public:
 	{
 		return PropertyName;
 	}
+	
+	FName GetParentSourceName() const
+	{
+		return ParentSourceName;
+	}
 
 #if UE_WITH_MVVM_DEBUGGING
 	struct FToStringArgs
@@ -102,9 +107,14 @@ private:
 	UPROPERTY()
 	FMVVMVCompiledFieldPath FieldPath;
 
+	/** The source name and the property's name of the view (if the flag IsUserWidgetProperty is set). */
 	UPROPERTY()
 	FName PropertyName;
 
+	/** The name of the parent source if it's the dynamic source. */
+	UPROPERTY()
+	FName ParentSourceName;
+
 	enum class ESourceFlags : uint8
 	{
 		None = 0,
diff --git a/Engine/Plugins/Runtime/ModelViewViewModel/Source/ModelViewViewModelBlueprint/Private/MVVMViewBlueprintCompiler.cpp b/Engine/Plugins/Runtime/ModelViewViewModel/Source/ModelViewViewModelBlueprint/Private/MVVMViewBlueprintCompiler.cpp
--- a/Engine/Plugins/Runtime/ModelViewViewModel/Source/ModelViewViewModelBlueprint/Private/MVVMViewBlueprintCompiler.cpp
+++ b/Engine/Plugins/Runtime/ModelViewViewModel/Source/ModelViewViewModelBlueprint/Private/MVVMViewBlueprintCompiler.cpp
@@ -1128,6 +1128,7 @@ bool FMVVMViewBlueprintCompiler::CompileSourceCreators(const FCompiledBindingLib
 
 			CanBeEvaluated = true;
 			CompiledSourceCreator.FieldPath = *CompiledFieldPath;
+			CompiledSourceCreator.ParentSourceName = SourceCreatorContext.DynamicParentSourceName;
 		}
 
 		CompiledSourceCreator.Flags = 0;
@@ -1572,6 +1573,7 @@ TValueOrError<FMVVMViewBlueprintCompiler::FBindingSourceContext, FText> FMVVMVie
 			return MakeError(LOCTEXT("DynamicSourceEntryNotSupport", "Long source entry is not supported. Add the viewmodel manually."));
 		}
 
+		//NB. The for loop order is important. The order is used in UMVVMView::SetViewModelInternal when enabling the bindings
 		int32 SourceCreatorContextIndex = INDEX_NONE;
 		for (int32 DynamicIndex = 1; DynamicIndex <= BindingInfo.ViewModelIndex; ++DynamicIndex)
 		{
@@ -1581,6 +1583,7 @@ TValueOrError<FMVVMViewBlueprintCompiler::FBindingSourceContext, FText> FMVVMVie
 			}
 
 			FName NewSourceName;
+			FName ParentSourceName;
 			FString NewSourcePropertyPath;
 			{
 				TStringBuilder<512> PropertyPathBuilder;
@@ -1594,6 +1597,11 @@ TValueOrError<FMVVMViewBlueprintCompiler::FBindingSourceContext, FText> FMVVMVie
 					}
 					PropertyPathBuilder << Result.PropertyPath[Index].GetName();
 					DynamicNameBuilder << Result.PropertyPath[Index].GetName();
+
+					if (Index == DynamicIndex - 1)
+					{
+						ParentSourceName = FName(DynamicNameBuilder.ToString());
+					}
 				}
 
 				NewSourceName = FName(DynamicNameBuilder.ToString());
@@ -1619,6 +1627,7 @@ TValueOrError<FMVVMViewBlueprintCompiler::FBindingSourceContext, FText> FMVVMVie
 
 					FCompilerSourceCreatorContext SourceCreatorContext;
 					SourceCreatorContext.Type = ECompilerSourceCreatorType::ViewModelDynamic;
+					SourceCreatorContext.DynamicParentSourceName = ParentSourceName;
 					SourceCreatorContext.ViewModelContext = FMVVMBlueprintViewModelContext(ViewModelClass, NewSourceName);
 					SourceCreatorContext.ViewModelContext.bCreateSetterFunction = false;
 					SourceCreatorContext.ViewModelContext.bOptional = false;
diff --git a/Engine/Plugins/Runtime/ModelViewViewModel/Source/ModelViewViewModelBlueprint/Private/MVVMViewBlueprintCompiler.h b/Engine/Plugins/Runtime/ModelViewViewModel/Source/ModelViewViewModelBlueprint/Private/MVVMViewBlueprintCompiler.h
--- a/Engine/Plugins/Runtime/ModelViewViewModel/Source/ModelViewViewModelBlueprint/Private/MVVMViewBlueprintCompiler.h
+++ b/Engine/Plugins/Runtime/ModelViewViewModel/Source/ModelViewViewModelBlueprint/Private/MVVMViewBlueprintCompiler.h
@@ -121,6 +121,7 @@ private:
 		ECompilerSourceCreatorType Type = ECompilerSourceCreatorType::ViewModel;
 		FString SetterFunctionName;
 		UEdGraph* SetterGraph = nullptr;
+		FName DynamicParentSourceName;
 	};
 	TArray<FCompilerSourceCreatorContext> CompilerSourceCreatorContexts;
 
