MVVM: Add the Broadcast field to the INotifyFieldValudChanged. Previously the viewmodel or widget implemented the functionality. The functionality was protected just because. Since we move it to a more engine wide feature and that it will be integrated with the Blueprint compiler it needs to be visible by Blueprint.

diff --git a/Engine/Plugins/Runtime/ModelViewViewModel/Source/ModelViewViewModel/Public/MVVMViewModelBase.h b/Engine/Plugins/Runtime/ModelViewViewModel/Source/ModelViewViewModel/Public/MVVMViewModelBase.h
--- a/Engine/Plugins/Runtime/ModelViewViewModel/Source/ModelViewViewModel/Public/MVVMViewModelBase.h
+++ b/Engine/Plugins/Runtime/ModelViewViewModel/Source/ModelViewViewModel/Public/MVVMViewModelBase.h
@@ -49,6 +49,7 @@ public:
 	virtual int32 RemoveAllFieldValueChangedDelegates(const void* InUserObject) override final;
 	virtual int32 RemoveAllFieldValueChangedDelegates(UE::FieldNotification::FFieldId InFieldId, const void* InUserObject) override final;
 	virtual const UE::FieldNotification::IClassDescriptor& GetFieldNotificationDescriptor() const override;
+	virtual void BroadcastFieldValueChanged(UE::FieldNotification::FFieldId InFieldId) override;
 	//~ End INotifyFieldValueChanged Interface
 
 	UFUNCTION(BlueprintCallable, Category = "FieldNotify", meta = (DisplayName = "Add Field Value Changed Delegate", ScriptName = "AddFieldValueChangedDelegate"))
@@ -58,14 +59,12 @@ public:
 	void K2_RemoveFieldValueChangedDelegate(FFieldNotificationId FieldId, FFieldValueChangedDynamicDelegate Delegate);
 
 protected:
-	UFUNCTION(BlueprintCallable, Category="FieldNotify", meta=(DisplayName="Broadcast Field Value Changed", ScriptName="BroadcastFieldValueChanged"))
+	UE_DEPRECATED(5.3, "BroadcastFieldValueChanged has been deprecated, please use the regular blueprint setter.")
+	UFUNCTION(BlueprintCallable, Category="FieldNotify", meta=(DisplayName="Broadcast Field Value Changed", ScriptName="BroadcastFieldValueChanged", DeprecatedFunction, DeprecatedMessage = "Use the regular setter node."))
 	void K2_BroadcastFieldValueChanged(FFieldNotificationId FieldId);
 
-	/** Broadcast the event (doesn't execute the replication code). */
-	void BroadcastFieldValueChanged(UE::FieldNotification::FFieldId InFieldId);
-
-protected:
-	UFUNCTION(BlueprintCallable, CustomThunk, Category="Viewmodel", meta=(CustomStructureParam="OldValue,NewValue", ScriptName="SetPropertyValue", BlueprintInternalUseOnly="true"))
+	UE_DEPRECATED(5.3, "SetPropertyValue has been deprecated, please use the regular blueprint setter.")
+	UFUNCTION(BlueprintCallable, CustomThunk, Category="Viewmodel", meta=(CustomStructureParam="OldValue,NewValue", ScriptName="SetPropertyValue", DeprecatedFunction, DeprecatedMessage = "Use the regular setter node.", BlueprintInternalUseOnly="true"))
 	bool K2_SetPropertyValue(UPARAM(ref) const int32& OldValue, UPARAM(ref) const int32& NewValue);
 
 	/** Set the new value and notify if the property value changed. */
diff --git a/Engine/Plugins/Runtime/ModelViewViewModel/Source/ModelViewViewModelEditor/Private/Tests/MVVMFieldPathHelperTest.h b/Engine/Plugins/Runtime/ModelViewViewModel/Source/ModelViewViewModelEditor/Private/Tests/MVVMFieldPathHelperTest.h
--- a/Engine/Plugins/Runtime/ModelViewViewModel/Source/ModelViewViewModelEditor/Private/Tests/MVVMFieldPathHelperTest.h
+++ b/Engine/Plugins/Runtime/ModelViewViewModel/Source/ModelViewViewModelEditor/Private/Tests/MVVMFieldPathHelperTest.h
@@ -201,6 +201,7 @@ public:
 	virtual int32 RemoveAllFieldValueChangedDelegates(const void* InUserObject) override final { return 0; }
 	virtual int32 RemoveAllFieldValueChangedDelegates(UE::FieldNotification::FFieldId InFieldId, const void* InUserObject) override final { return 0; }
 	virtual const UE::FieldNotification::IClassDescriptor& GetFieldNotificationDescriptor() const override { static FFieldNotificationClassDescriptor Tmp; return Tmp; }
+	virtual void BroadcastFieldValueChanged(UE::FieldNotification::FFieldId InFieldId) override {}
 	//~ End INotifyFieldValueChanged Interface
 };
 
diff --git a/Engine/Source/Runtime/FieldNotification/Public/INotifyFieldValueChanged.h b/Engine/Source/Runtime/FieldNotification/Public/INotifyFieldValueChanged.h
--- a/Engine/Source/Runtime/FieldNotification/Public/INotifyFieldValueChanged.h
+++ b/Engine/Source/Runtime/FieldNotification/Public/INotifyFieldValueChanged.h
@@ -42,4 +42,7 @@ public:
 
 	/** @returns the list of all the field that can notify when their value changes. */
 	virtual const UE::FieldNotification::IClassDescriptor& GetFieldNotificationDescriptor() const = 0;
+
+	/** Broadcast to the registered delegate that the FieldId value changed. */
+	virtual void BroadcastFieldValueChanged(UE::FieldNotification::FFieldId InFieldId) = 0;
 };
diff --git a/Engine/Source/Runtime/UMG/Public/Components/Widget.h b/Engine/Source/Runtime/UMG/Public/Components/Widget.h
--- a/Engine/Source/Runtime/UMG/Public/Components/Widget.h
+++ b/Engine/Source/Runtime/UMG/Public/Components/Widget.h
@@ -771,6 +771,7 @@ public:
 	virtual bool RemoveFieldValueChangedDelegate(UE::FieldNotification::FFieldId InFieldId, FDelegateHandle InHandle) override final;
 	virtual int32 RemoveAllFieldValueChangedDelegates(const void* InUserObject) override final;
 	virtual int32 RemoveAllFieldValueChangedDelegates(UE::FieldNotification::FFieldId InFieldId, const void* InUserObject) override final;
+	virtual void BroadcastFieldValueChanged(UE::FieldNotification::FFieldId InFieldId) override final;
 	//~ End INotifyFieldValueChanged Interface
 
 	UFUNCTION(BlueprintCallable, Category = "FieldNotify", meta = (DisplayName = "Add Field Value Changed Delegate", ScriptName = "AddFieldValueChangedDelegate"))
@@ -780,8 +781,6 @@ public:
 	void K2_RemoveFieldValueChangedDelegate(FFieldNotificationId FieldId, FFieldValueChangedDynamicDelegate Delegate);
 
 protected:
-	void BroadcastFieldValueChanged(UE::FieldNotification::FFieldId InFieldId);
-
 	UFUNCTION(BlueprintCallable, Category="FieldNotify", meta = (DisplayName="Broadcast Field Value Changed", ScriptName="BroadcastFieldValueChanged"))
 	void K2_BroadcastFieldValueChanged(FFieldNotificationId FieldId);
 
