MVVM: Fix false positive when detecting duplicated binding. Complex functions do not have a source and it's normal that it reuse the same compiled binding.

diff --git a/Engine/Plugins/Runtime/ModelViewViewModel/Source/ModelViewViewModelBlueprint/Private/Bindings/MVVMCompiledBindingLibraryCompiler.cpp b/Engine/Plugins/Runtime/ModelViewViewModel/Source/ModelViewViewModelBlueprint/Private/Bindings/MVVMCompiledBindingLibraryCompiler.cpp
--- a/Engine/Plugins/Runtime/ModelViewViewModel/Source/ModelViewViewModelBlueprint/Private/Bindings/MVVMCompiledBindingLibraryCompiler.cpp
+++ b/Engine/Plugins/Runtime/ModelViewViewModel/Source/ModelViewViewModelBlueprint/Private/Bindings/MVVMCompiledBindingLibraryCompiler.cpp
@@ -544,7 +544,14 @@ TValueOrError<FCompiledBindingLibraryCompiler::FBindingHandle, FText> FCompiledB
 
 	if (FoundSameBindingIndex != INDEX_NONE)
 	{
-		return MakeError(LOCTEXT("BindingAlreadyAdded", "The binding already exist."));
+		if (!bInIsComplexBinding)
+		{
+			return MakeError(LOCTEXT("BindingAlreadyAdded", "The binding already exist."));
+		}
+		else
+		{
+			return MakeValue(Impl->Bindings[FoundSameBindingIndex].BindingHandle);
+		}
 	}
 	else
 	{
