MVVM: Remove the specialized viewmodel editor and use the Blueprint base editor instead.

diff --git a/Engine/Plugins/Runtime/ModelViewViewModel/Source/ModelViewViewModel/Private/MVVMViewModelBase.cpp b/Engine/Plugins/Runtime/ModelViewViewModel/Source/ModelViewViewModel/Private/MVVMViewModelBase.cpp
--- a/Engine/Plugins/Runtime/ModelViewViewModel/Source/ModelViewViewModel/Private/MVVMViewModelBase.cpp
+++ b/Engine/Plugins/Runtime/ModelViewViewModel/Source/ModelViewViewModel/Private/MVVMViewModelBase.cpp
@@ -63,9 +63,9 @@ const UE::FieldNotification::IClassDescriptor& UMVVMViewModelBase::GetFieldNotif
 
 void UMVVMViewModelBase::FFieldNotificationClassDescriptor::ForEachField(const UClass* Class, TFunctionRef<bool(::UE::FieldNotification::FFieldId FieldId)> Callback) const
 {
-	if (const UMVVMViewModelBlueprintGeneratedClass* ViewModelBPClass = Cast<const UMVVMViewModelBlueprintGeneratedClass>(Class))
+	if (const UBlueprintGeneratedClass* ViewModelBPClass = Cast<const UBlueprintGeneratedClass>(Class))
 	{
-		ViewModelBPClass->ForEachField(Callback);
+		ViewModelBPClass->ForEachFieldNotify(Callback, true);
 	}
 }
 
diff --git a/Engine/Plugins/Runtime/ModelViewViewModel/Source/ModelViewViewModelBlueprint/ModelViewViewModelBlueprint.Build.cs b/Engine/Plugins/Runtime/ModelViewViewModel/Source/ModelViewViewModelBlueprint/ModelViewViewModelBlueprint.Build.cs
--- a/Engine/Plugins/Runtime/ModelViewViewModel/Source/ModelViewViewModelBlueprint/ModelViewViewModelBlueprint.Build.cs
+++ b/Engine/Plugins/Runtime/ModelViewViewModel/Source/ModelViewViewModelBlueprint/ModelViewViewModelBlueprint.Build.cs
@@ -32,7 +32,5 @@ public class ModelViewViewModelBlueprint : ModuleRules
 				"UMGEditor",
 				"UnrealEd",
 			});
-
-		PublicDefinitions.Add("UE_MVVM_WITH_VIEWMODEL_EDITOR=0");
 	}
 }
diff --git a/Engine/Plugins/Runtime/ModelViewViewModel/Source/ModelViewViewModelBlueprint/Private/ModelViewViewModelBlueprintModule.cpp b/Engine/Plugins/Runtime/ModelViewViewModel/Source/ModelViewViewModelBlueprint/Private/ModelViewViewModelBlueprintModule.cpp
--- a/Engine/Plugins/Runtime/ModelViewViewModel/Source/ModelViewViewModelBlueprint/Private/ModelViewViewModelBlueprintModule.cpp
+++ b/Engine/Plugins/Runtime/ModelViewViewModel/Source/ModelViewViewModelBlueprint/Private/ModelViewViewModelBlueprintModule.cpp
@@ -3,36 +3,5 @@
 #include "ModelViewViewModelBlueprintModule.h"
 #include "Modules/ModuleManager.h"
 
-#include "ViewModel/MVVMViewModelBlueprint.h"
-#include "ViewModel/MVVMViewModelBlueprintCompiler.h"
 
-
-class FModelViewViewModelBlueprintModule : public IModuleInterface
-{
-public:
-	virtual void StartupModule() override
-	{
-#if UE_MVVM_WITH_VIEWMODEL_EDITOR
-		IKismetCompilerInterface& KismetCompilerModule = FModuleManager::GetModuleChecked<IKismetCompilerInterface>("KismetCompiler");
-		KismetCompilerModule.GetCompilers().Add(&ViewModelBlueprintCompiler);
-
-		FKismetCompilerContext::RegisterCompilerForBP(UMVVMViewModelBlueprint::StaticClass(), &UMVVMViewModelBlueprint::GetCompilerForViewModelBlueprint);
-#endif
-	}
-
-	virtual void ShutdownModule() override
-	{
-#if UE_MVVM_WITH_VIEWMODEL_EDITOR
-		if (IKismetCompilerInterface* KismetCompilerModule = FModuleManager::GetModulePtr<IKismetCompilerInterface>("KismetCompiler"))
-		{
-			KismetCompilerModule->GetCompilers().Remove(&ViewModelBlueprintCompiler);
-		}
-#endif
-	}
-
-private:
-	UE::MVVM::FViewModelBlueprintCompiler ViewModelBlueprintCompiler;
-};
-
-
-IMPLEMENT_MODULE(FModelViewViewModelBlueprintModule, ModelViewViewModelBlueprint);
+IMPLEMENT_MODULE(FDefaultModuleImpl, ModelViewViewModelBlueprint);
diff --git a/Engine/Plugins/Runtime/ModelViewViewModel/Source/ModelViewViewModelBlueprint/Private/ViewModel/MVVMViewModelBlueprint.cpp b/Engine/Plugins/Runtime/ModelViewViewModel/Source/ModelViewViewModelBlueprint/Private/ViewModel/MVVMViewModelBlueprint.cpp
deleted file mode 100644
diff --git a/Engine/Plugins/Runtime/ModelViewViewModel/Source/ModelViewViewModelBlueprint/Private/ViewModel/MVVMViewModelBlueprintCompiler.cpp b/Engine/Plugins/Runtime/ModelViewViewModel/Source/ModelViewViewModelBlueprint/Private/ViewModel/MVVMViewModelBlueprintCompiler.cpp
deleted file mode 100644
diff --git a/Engine/Plugins/Runtime/ModelViewViewModel/Source/ModelViewViewModelBlueprint/Private/ViewModel/MVVMViewModelBlueprintCompiler.h b/Engine/Plugins/Runtime/ModelViewViewModel/Source/ModelViewViewModelBlueprint/Private/ViewModel/MVVMViewModelBlueprintCompiler.h
deleted file mode 100644
diff --git a/Engine/Plugins/Runtime/ModelViewViewModel/Source/ModelViewViewModelBlueprint/Public/ViewModel/MVVMViewModelBlueprint.h b/Engine/Plugins/Runtime/ModelViewViewModel/Source/ModelViewViewModelBlueprint/Public/ViewModel/MVVMViewModelBlueprint.h
deleted file mode 100644
diff --git a/Engine/Plugins/Runtime/ModelViewViewModel/Source/ModelViewViewModelEditor/Private/ModelViewViewModelEditorModule.cpp b/Engine/Plugins/Runtime/ModelViewViewModel/Source/ModelViewViewModelEditor/Private/ModelViewViewModelEditorModule.cpp
--- a/Engine/Plugins/Runtime/ModelViewViewModel/Source/ModelViewViewModelEditor/Private/ModelViewViewModelEditorModule.cpp
+++ b/Engine/Plugins/Runtime/ModelViewViewModel/Source/ModelViewViewModelEditor/Private/ModelViewViewModelEditorModule.cpp
@@ -20,7 +20,6 @@
 #include "Tabs/MVVMBindingSummoner.h"
 #include "Tabs/MVVMViewModelSummoner.h"
 #include "UMGEditorModule.h"
-#include "ViewModel/AssetTypeActions_ViewModelBlueprint.h"
 #include "WidgetBlueprintEditor.h"
 #include "WidgetDrawerConfig.h"
 
@@ -36,16 +35,6 @@ void FModelViewViewModelEditorModule::StartupModule()
 	IUMGEditorModule& UMGEditorModule = FModuleManager::LoadModuleChecked<IUMGEditorModule>("UMGEditor");
 	UMGEditorModule.OnRegisterTabsForEditor().AddRaw(this, &FModelViewViewModelEditorModule::HandleRegisterBlueprintEditorTab);
 
-	// Register asset types
-	{
-#if UE_MVVM_WITH_VIEWMODEL_EDITOR
-		// Only remove what is related to the viewmodel editor, not the UMG extention for view.
-		IAssetTools& AssetTools = FModuleManager::LoadModuleChecked<FAssetToolsModule>("AssetTools").Get();
-		ViewModelBlueprintActions = MakeShared<UE::MVVM::FAssetTypeActions_ViewModelBlueprint>();
-		AssetTools.RegisterAssetTypeActions(ViewModelBlueprintActions.ToSharedRef());
-#endif
-	}
-
 	PropertyBindingExtension = MakeShared<FMVVMPropertyBindingExtension>();
 	UMGEditorModule.GetPropertyBindingExtensibilityManager()->AddExtension(PropertyBindingExtension.ToSharedRef());
 
@@ -91,17 +80,6 @@ void FModelViewViewModelEditorModule::ShutdownModule()
 	}
 	PropertyBindingExtension.Reset();
 
-#if UE_MVVM_WITH_VIEWMODEL_EDITOR
-	// Unregister all the asset types that we registered
-	if (FAssetToolsModule* AssetTools = FModuleManager::GetModulePtr<FAssetToolsModule>("AssetTools"))
-	{
-		if (ViewModelBlueprintActions)
-		{
-			AssetTools->Get().UnregisterAssetTypeActions(ViewModelBlueprintActions.ToSharedRef());
-		}
-	}
-#endif
-
 	FMVVMEditorStyle::DestroyInstance();
 
 	FMVVMEditorCommands::Unregister();
diff --git a/Engine/Plugins/Runtime/ModelViewViewModel/Source/ModelViewViewModelEditor/Private/ViewModel/AssetTypeActions_ViewModelBlueprint.cpp b/Engine/Plugins/Runtime/ModelViewViewModel/Source/ModelViewViewModelEditor/Private/ViewModel/AssetTypeActions_ViewModelBlueprint.cpp
deleted file mode 100644
diff --git a/Engine/Plugins/Runtime/ModelViewViewModel/Source/ModelViewViewModelEditor/Private/ViewModel/AssetTypeActions_ViewModelBlueprint.h b/Engine/Plugins/Runtime/ModelViewViewModel/Source/ModelViewViewModelEditor/Private/ViewModel/AssetTypeActions_ViewModelBlueprint.h
deleted file mode 100644
diff --git a/Engine/Plugins/Runtime/ModelViewViewModel/Source/ModelViewViewModelEditor/Private/ViewModel/MVVMViewModelBlueprintEditor.cpp b/Engine/Plugins/Runtime/ModelViewViewModel/Source/ModelViewViewModelEditor/Private/ViewModel/MVVMViewModelBlueprintEditor.cpp
deleted file mode 100644
diff --git a/Engine/Plugins/Runtime/ModelViewViewModel/Source/ModelViewViewModelEditor/Private/ViewModel/MVVMViewModelBlueprintFactory.cpp b/Engine/Plugins/Runtime/ModelViewViewModel/Source/ModelViewViewModelEditor/Private/ViewModel/MVVMViewModelBlueprintFactory.cpp
deleted file mode 100644
diff --git a/Engine/Plugins/Runtime/ModelViewViewModel/Source/ModelViewViewModelEditor/Private/ViewModel/MVVMViewModelBlueprintFactory.h b/Engine/Plugins/Runtime/ModelViewViewModel/Source/ModelViewViewModelEditor/Private/ViewModel/MVVMViewModelBlueprintFactory.h
deleted file mode 100644
diff --git a/Engine/Plugins/Runtime/ModelViewViewModel/Source/ModelViewViewModelEditor/Public/ViewModel/MVVMViewModelBlueprintEditor.h b/Engine/Plugins/Runtime/ModelViewViewModel/Source/ModelViewViewModelEditor/Public/ViewModel/MVVMViewModelBlueprintEditor.h
deleted file mode 100644
diff --git a/Engine/Plugins/Runtime/ModelViewViewModel/Source/ModelViewViewModelEditor/Public/ViewModel/MVVMViewModelBlueprintToolMenuContext.h b/Engine/Plugins/Runtime/ModelViewViewModel/Source/ModelViewViewModelEditor/Public/ViewModel/MVVMViewModelBlueprintToolMenuContext.h
deleted file mode 100644
