[Backout] - CL25027663
Original CL Desc
-----------------------------------------------------------------
MVVM: Remove the specialized viewmodel editor and use the Blueprint base editor instead.

diff --git a/Engine/Plugins/Runtime/ModelViewViewModel/Source/ModelViewViewModel/Private/MVVMViewModelBase.cpp b/Engine/Plugins/Runtime/ModelViewViewModel/Source/ModelViewViewModel/Private/MVVMViewModelBase.cpp
--- a/Engine/Plugins/Runtime/ModelViewViewModel/Source/ModelViewViewModel/Private/MVVMViewModelBase.cpp
+++ b/Engine/Plugins/Runtime/ModelViewViewModel/Source/ModelViewViewModel/Private/MVVMViewModelBase.cpp
@@ -63,9 +63,9 @@ const UE::FieldNotification::IClassDescriptor& UMVVMViewModelBase::GetFieldNotif
 
 void UMVVMViewModelBase::FFieldNotificationClassDescriptor::ForEachField(const UClass* Class, TFunctionRef<bool(::UE::FieldNotification::FFieldId FieldId)> Callback) const
 {
-	if (const UBlueprintGeneratedClass* ViewModelBPClass = Cast<const UBlueprintGeneratedClass>(Class))
+	if (const UMVVMViewModelBlueprintGeneratedClass* ViewModelBPClass = Cast<const UMVVMViewModelBlueprintGeneratedClass>(Class))
 	{
-		ViewModelBPClass->ForEachFieldNotify(Callback, true);
+		ViewModelBPClass->ForEachField(Callback);
 	}
 }
 
diff --git a/Engine/Plugins/Runtime/ModelViewViewModel/Source/ModelViewViewModelBlueprint/ModelViewViewModelBlueprint.Build.cs b/Engine/Plugins/Runtime/ModelViewViewModel/Source/ModelViewViewModelBlueprint/ModelViewViewModelBlueprint.Build.cs
--- a/Engine/Plugins/Runtime/ModelViewViewModel/Source/ModelViewViewModelBlueprint/ModelViewViewModelBlueprint.Build.cs
+++ b/Engine/Plugins/Runtime/ModelViewViewModel/Source/ModelViewViewModelBlueprint/ModelViewViewModelBlueprint.Build.cs
@@ -32,5 +32,7 @@ public class ModelViewViewModelBlueprint : ModuleRules
 				"UMGEditor",
 				"UnrealEd",
 			});
+
+		PublicDefinitions.Add("UE_MVVM_WITH_VIEWMODEL_EDITOR=0");
 	}
 }
diff --git a/Engine/Plugins/Runtime/ModelViewViewModel/Source/ModelViewViewModelBlueprint/Private/ModelViewViewModelBlueprintModule.cpp b/Engine/Plugins/Runtime/ModelViewViewModel/Source/ModelViewViewModelBlueprint/Private/ModelViewViewModelBlueprintModule.cpp
--- a/Engine/Plugins/Runtime/ModelViewViewModel/Source/ModelViewViewModelBlueprint/Private/ModelViewViewModelBlueprintModule.cpp
+++ b/Engine/Plugins/Runtime/ModelViewViewModel/Source/ModelViewViewModelBlueprint/Private/ModelViewViewModelBlueprintModule.cpp
@@ -3,5 +3,36 @@
 #include "ModelViewViewModelBlueprintModule.h"
 #include "Modules/ModuleManager.h"
 
+#include "ViewModel/MVVMViewModelBlueprint.h"
+#include "ViewModel/MVVMViewModelBlueprintCompiler.h"
 
-IMPLEMENT_MODULE(FDefaultModuleImpl, ModelViewViewModelBlueprint);
+
+class FModelViewViewModelBlueprintModule : public IModuleInterface
+{
+public:
+	virtual void StartupModule() override
+	{
+#if UE_MVVM_WITH_VIEWMODEL_EDITOR
+		IKismetCompilerInterface& KismetCompilerModule = FModuleManager::GetModuleChecked<IKismetCompilerInterface>("KismetCompiler");
+		KismetCompilerModule.GetCompilers().Add(&ViewModelBlueprintCompiler);
+
+		FKismetCompilerContext::RegisterCompilerForBP(UMVVMViewModelBlueprint::StaticClass(), &UMVVMViewModelBlueprint::GetCompilerForViewModelBlueprint);
+#endif
+	}
+
+	virtual void ShutdownModule() override
+	{
+#if UE_MVVM_WITH_VIEWMODEL_EDITOR
+		if (IKismetCompilerInterface* KismetCompilerModule = FModuleManager::GetModulePtr<IKismetCompilerInterface>("KismetCompiler"))
+		{
+			KismetCompilerModule->GetCompilers().Remove(&ViewModelBlueprintCompiler);
+		}
+#endif
+	}
+
+private:
+	UE::MVVM::FViewModelBlueprintCompiler ViewModelBlueprintCompiler;
+};
+
+
+IMPLEMENT_MODULE(FModelViewViewModelBlueprintModule, ModelViewViewModelBlueprint);
diff --git a/Engine/Plugins/Runtime/ModelViewViewModel/Source/ModelViewViewModelBlueprint/Private/ViewModel/MVVMViewModelBlueprint.cpp b/Engine/Plugins/Runtime/ModelViewViewModel/Source/ModelViewViewModelBlueprint/Private/ViewModel/MVVMViewModelBlueprint.cpp
new file mode 100644
--- /dev/null
+++ b/Engine/Plugins/Runtime/ModelViewViewModel/Source/ModelViewViewModelBlueprint/Private/ViewModel/MVVMViewModelBlueprint.cpp
@@ -0,0 +1,51 @@
+// Copyright Epic Games, Inc. All Rights Reserved.
+
+#include "ViewModel/MVVMViewModelBlueprint.h"
+#include "ViewModel/MVVMViewModelBlueprintCompiler.h"
+#include "ViewModel/MVVMViewModelBlueprintGeneratedClass.h"
+
+
+#include UE_INLINE_GENERATED_CPP_BY_NAME(MVVMViewModelBlueprint)
+
+
+TSharedPtr<FKismetCompilerContext> UMVVMViewModelBlueprint::GetCompilerForViewModelBlueprint(UBlueprint* BP, FCompilerResultsLog& InMessageLog, const FKismetCompilerOptions& InCompileOptions)
+{
+	return TSharedPtr<FKismetCompilerContext>(new UE::MVVM::FViewModelBlueprintCompilerContext(CastChecked<UMVVMViewModelBlueprint>(BP), InMessageLog, InCompileOptions));
+}
+
+
+UMVVMViewModelBlueprintGeneratedClass* UMVVMViewModelBlueprint::GetViewModelBlueprintGeneratedClass() const
+{
+	return Cast<UMVVMViewModelBlueprintGeneratedClass>(*GeneratedClass);
+}
+
+
+UMVVMViewModelBlueprintGeneratedClass* UMVVMViewModelBlueprint::GetViewModelBlueprintSkeletonClass() const
+{
+	return Cast<UMVVMViewModelBlueprintGeneratedClass>(*SkeletonGeneratedClass);
+}
+
+
+bool UMVVMViewModelBlueprint::SupportsMacros() const
+{
+	return false;
+}
+
+
+bool UMVVMViewModelBlueprint::SupportsEventGraphs() const
+{
+	return false;
+}
+
+
+bool UMVVMViewModelBlueprint::SupportsDelegates() const
+{
+	return false;
+}
+
+
+bool UMVVMViewModelBlueprint::SupportsAnimLayers() const
+{
+	return false;
+}
+
diff --git a/Engine/Plugins/Runtime/ModelViewViewModel/Source/ModelViewViewModelBlueprint/Private/ViewModel/MVVMViewModelBlueprintCompiler.cpp b/Engine/Plugins/Runtime/ModelViewViewModel/Source/ModelViewViewModelBlueprint/Private/ViewModel/MVVMViewModelBlueprintCompiler.cpp
new file mode 100644
--- /dev/null
+++ b/Engine/Plugins/Runtime/ModelViewViewModel/Source/ModelViewViewModelBlueprint/Private/ViewModel/MVVMViewModelBlueprintCompiler.cpp
@@ -0,0 +1,321 @@
+// Copyright Epic Games, Inc. All Rights Reserved.
+
+#include "ViewModel/MVVMViewModelBlueprintCompiler.h"
+#include "HAL/IConsoleManager.h"
+#include "ViewModel/MVVMViewModelBlueprint.h"
+#include "KismetCompiledFunctionContext.h"
+#include "ViewModel/MVVMViewModelBlueprintGeneratedClass.h"
+#include "MVVMViewModelBase.h"
+
+#include "FieldNotification/CustomizationHelper.h"
+#include "FieldNotificationHelpers.h"
+#include "K2Node_FunctionEntry.h"
+#include "Kismet2/KismetReinstanceUtilities.h"
+#include "MVVMFunctionGraphHelper.h"
+
+#define LOCTEXT_NAMESPACE "ViewModelBlueprintCompiler"
+
+namespace UE::MVVM
+{
+
+static TAutoConsoleVariable<bool> CVarAutogenerateSetter(
+	TEXT("MVVM.Viewmodel.GenerateSetter"),
+	true,
+	TEXT("If true, the setter function will be always be autogenerated"),
+	ECVF_Default);
+
+static TAutoConsoleVariable<bool> CVarAutogenerateOnRep(
+	TEXT("MVVM.Viewmodel.GenerateOnRep"),
+	true,
+	TEXT("If true, the OnRep function will be always be autogenerated"),
+	ECVF_Default);
+
+//////////////////////////////////////////////////////////////////////////
+// FViewModelBlueprintCompiler 
+
+bool FViewModelBlueprintCompiler::CanCompile(const UBlueprint* Blueprint)
+{
+	return Cast<UMVVMViewModelBlueprint>(Blueprint) != nullptr;
+}
+
+
+void FViewModelBlueprintCompiler::Compile(UBlueprint* Blueprint, const FKismetCompilerOptions& CompileOptions, FCompilerResultsLog& Results)
+{
+	if (UMVVMViewModelBlueprint* WidgetBlueprint = CastChecked<UMVVMViewModelBlueprint>(Blueprint))
+	{
+		FViewModelBlueprintCompilerContext Compiler(WidgetBlueprint, Results, CompileOptions);
+		Compiler.Compile();
+		check(Compiler.NewClass);
+	}
+}
+
+
+bool FViewModelBlueprintCompiler::GetBlueprintTypesForClass(UClass* ParentClass, UClass*& OutBlueprintClass, UClass*& OutBlueprintGeneratedClass) const
+{
+	if (ParentClass == UMVVMViewModelBase::StaticClass() || ParentClass->IsChildOf(UMVVMViewModelBase::StaticClass()))
+	{
+		OutBlueprintClass = UMVVMViewModelBlueprint::StaticClass();
+		OutBlueprintGeneratedClass = UMVVMViewModelBlueprintGeneratedClass::StaticClass();
+		return true;
+	}
+
+	return false;
+}
+
+
+//////////////////////////////////////////////////////////////////////////
+// FViewModelBlueprintCompilerContext
+
+
+UMVVMViewModelBlueprint* FViewModelBlueprintCompilerContext::GetViewModelBlueprint() const
+{
+	return Cast<UMVVMViewModelBlueprint>(Blueprint);
+}
+
+
+FProperty* FViewModelBlueprintCompilerContext::FindPropertyByNameOnNewClass(FName PropertyName) const
+{
+	for (FField* CurrentField = NewClass->ChildProperties; CurrentField != nullptr; CurrentField = CurrentField->Next)
+	{
+		if (CurrentField->GetFName() == PropertyName)
+		{
+			return CastField<FProperty>(CurrentField);
+		}
+	}
+	return nullptr;
+}
+
+
+void FViewModelBlueprintCompilerContext::SaveSubObjectsFromCleanAndSanitizeClass(FSubobjectCollection& SubObjectsToSave, UBlueprintGeneratedClass* ClassToClean)
+{
+	Super::SaveSubObjectsFromCleanAndSanitizeClass(SubObjectsToSave, ClassToClean);
+
+	// Make sure our typed pointer is set
+	check(ClassToClean == NewClass);
+	NewViewModelBlueprintClass = CastChecked<UMVVMViewModelBlueprintGeneratedClass>((UObject*)NewClass);
+}
+
+
+void FViewModelBlueprintCompilerContext::CleanAndSanitizeClass(UBlueprintGeneratedClass* ClassToClean, UObject*& InOutOldCDO)
+{
+	Super::CleanAndSanitizeClass(ClassToClean, InOutOldCDO);
+
+	NewViewModelBlueprintClass->FieldNotifyNames.Empty();
+}
+
+
+void FViewModelBlueprintCompilerContext::CreateFunctionList()
+{
+	for (const FGeneratedFunction& Function : GeneratedFunctions)
+	{
+		if (Function.Property == nullptr && Function.SkelProperty == nullptr)
+		{
+			MessageLog.Error(*FString::Printf(TEXT("Internal error. The property '%s' is invalid. Setter and net rep was not generated."), *Function.PropertyName.ToString()));
+			continue;
+		}
+
+		if (Function.SetterFunction)
+		{
+			if (!FunctionGraphHelper::GenerateViewModelFieldNotifySetter(*this, Function.SetterFunction, Function.SkelProperty, TEXT("NewValue")))
+			{
+				MessageLog.Error(*FString::Printf(TEXT("The setter function for '%s' could not be generated."), *Function.PropertyName.ToString()));
+				continue;
+			}
+		}
+		if (Function.NetRepFunction)
+		{
+			if (!FunctionGraphHelper::GenerateViewModelFielNotifyBroadcast(*this, Function.NetRepFunction, Function.SkelProperty))
+			{
+				MessageLog.Error(*FString::Printf(TEXT("The net rep function for '%s' could not be generated."), *Function.PropertyName.ToString()));
+				continue;
+			}
+		}
+	}
+
+	Super::CreateFunctionList();
+}
+
+
+void FViewModelBlueprintCompilerContext::CreateClassVariablesFromBlueprint()
+{
+	Super::CreateClassVariablesFromBlueprint();
+
+	if (CompileOptions.CompileType == EKismetCompileType::SkeletonOnly)
+	{
+		auto RenameObjectToTransientPackage = [](UObject* ObjectToRename)
+		{
+			const ERenameFlags RenFlags = REN_DontCreateRedirectors | REN_NonTransactional | REN_DoNotDirty;
+			ObjectToRename->Rename(nullptr, GetTransientPackage(), RenFlags);
+			ObjectToRename->SetFlags(RF_Transient);
+			ObjectToRename->ClearFlags(RF_Public | RF_Standalone | RF_ArchetypeObject);
+			FLinkerLoad::InvalidateExport(ObjectToRename);
+		};
+
+		for (UEdGraph* OldGraph : GetViewModelBlueprint()->TemporaryGraph)
+		{
+			if (OldGraph)
+			{
+				RenameObjectToTransientPackage(OldGraph);
+			}
+		}
+		GetViewModelBlueprint()->TemporaryGraph.Reset();
+	}
+
+	const FName NAME_MetaDataBlueprintSetter = "BlueprintSetter";
+	const FName NAME_Category = "Category";
+	for (TFieldIterator<FProperty> PropertyIt(NewClass, EFieldIterationFlags::None); PropertyIt; ++PropertyIt)
+	{
+		FProperty* Property = *PropertyIt;
+		FName PropertyName = Property->GetFName();
+
+		bool bGenerateSetterFunction = CVarAutogenerateSetter.GetValueOnGameThread() || Property->HasMetaData(UE::FieldNotification::FCustomizationHelper::MetaData_FieldNotify);
+		bool bGenerateReplicatedFunction = CVarAutogenerateOnRep.GetValueOnGameThread(); // Property->RepNotifyFunc.IsNone();
+		if (bGenerateSetterFunction || bGenerateReplicatedFunction)
+		{
+			Property->SetPropertyFlags(CPF_DisableEditOnInstance);
+
+			const FString SetterName = FString::Printf(TEXT("Set%s"), *Property->GetName());
+			const FString NetRepName = FString::Printf(TEXT("__RepNotify_%s"), *Property->GetName());
+
+#if WITH_EDITOR
+			if (bGenerateSetterFunction)
+			{
+				Property->SetMetaData(FBlueprintMetadata::MD_PropertySetFunction, *SetterName);
+			}
+#endif
+
+			if (CompileOptions.CompileType == EKismetCompileType::SkeletonOnly)
+			{
+				FGeneratedFunction GeneratedFunction;
+				GeneratedFunction.PropertyName = PropertyName;
+				GeneratedFunction.SkelProperty = Property;
+
+				if (GeneratedFunctions.ContainsByPredicate([PropertyName](const FGeneratedFunction& Other)
+					{
+						return Other.PropertyName == PropertyName;
+					}))
+				{
+					MessageLog.Error(*FString::Printf(TEXT("Internal error. The functions are already generated for properpty."), *PropertyName.ToString()));
+					continue;
+				}
+
+				if (bGenerateSetterFunction)
+				{
+					GeneratedFunction.SetterFunction = FunctionGraphHelper::CreateIntermediateFunctionGraph(
+						*this
+						, *SetterName
+						, (FUNC_BlueprintCallable | FUNC_Public)
+						, Property->GetMetaData(NAME_Category)
+						, false);
+					GetViewModelBlueprint()->TemporaryGraph.Add(GeneratedFunction.SetterFunction);
+
+					if (GeneratedFunction.SetterFunction == nullptr)
+					{
+						MessageLog.Error(*FString::Printf(TEXT("The setter name '%s' already exists and could not be autogenerated."), *SetterName));
+						continue;
+					}
+
+					if (GeneratedFunction.SetterFunction->GetFName() != *SetterName)
+					{
+						MessageLog.Error(*FString::Printf(TEXT("The setter name '%s' was generated with a different name."), *SetterName));
+						continue;
+					}
+
+					FunctionGraphHelper::AddFunctionArgument(GeneratedFunction.SetterFunction, Property, TEXT("NewValue"));
+				}
+
+				if (bGenerateReplicatedFunction)
+				{
+					GeneratedFunction.NetRepFunction = FunctionGraphHelper::CreateIntermediateFunctionGraph(
+						*this
+						, *NetRepName
+						, (FUNC_Private)
+						, Property->GetMetaData(NAME_Category)
+						, false);
+					GetViewModelBlueprint()->TemporaryGraph.Add(GeneratedFunction.NetRepFunction);
+
+					if (GeneratedFunction.NetRepFunction == nullptr)
+					{
+						MessageLog.Error(*FString::Printf(TEXT("The netrep function name '%s' already exists and could not be autogenerated."), *NetRepName));
+						continue;
+					}
+
+					if (GeneratedFunction.NetRepFunction->GetFName() != *NetRepName)
+					{
+						MessageLog.Error(*FString::Printf(TEXT("The netrep function name '%s' was generated with a different name."), *NetRepName));
+						continue;
+					}
+
+					Property->SetPropertyFlags(CPF_Net);
+					Property->RepNotifyFunc = GeneratedFunction.NetRepFunction->GetFName();
+
+					FunctionGraphHelper::AddFunctionArgument(GeneratedFunction.SetterFunction, Property, TEXT("NewValue"));
+				}
+
+				GeneratedFunctions.Add(GeneratedFunction);
+			}
+
+			NewViewModelBlueprintClass->FieldNotifyNames.Emplace(Property->GetFName());
+		}
+	}
+}
+
+
+void FViewModelBlueprintCompilerContext::SpawnNewClass(const FString& NewClassName)
+{
+	NewViewModelBlueprintClass = FindObject<UMVVMViewModelBlueprintGeneratedClass>(Blueprint->GetOutermost(), *NewClassName);
+
+	if (NewViewModelBlueprintClass == nullptr)
+	{
+		NewViewModelBlueprintClass = NewObject<UMVVMViewModelBlueprintGeneratedClass>(Blueprint->GetOutermost(), FName(*NewClassName), RF_Public | RF_Transactional);
+	}
+	else
+	{
+		// Already existed, but wasn't linked in the Blueprint yet due to load ordering issues
+		FBlueprintCompileReinstancer::Create(NewViewModelBlueprintClass);
+	}
+	NewClass = NewViewModelBlueprintClass;
+}
+
+
+void FViewModelBlueprintCompilerContext::OnNewClassSet(UBlueprintGeneratedClass* ClassToUse)
+{
+	NewViewModelBlueprintClass = CastChecked<UMVVMViewModelBlueprintGeneratedClass>(ClassToUse);
+}
+
+
+void FViewModelBlueprintCompilerContext::PostcompileFunction(FKismetFunctionContext& Context)
+{
+	Super::PostcompileFunction(Context);
+
+	if (Context.Function && Context.Function->HasMetaData(UE::FieldNotification::FCustomizationHelper::MetaData_FieldNotify))
+	{
+		if (!UE::FieldNotification::Helpers::IsValidAsField(Context.Function))
+		{
+			MessageLog.Error(*LOCTEXT("FieldNotify_IsEventGraph", "Function @@ cannot be a FieldNotify. A function needs to be const, returns a single properties, has no inputs, not be an event or a net function.").ToString(), Context.EntryPoint);
+		}
+		NewViewModelBlueprintClass->FieldNotifyNames.Emplace(Context.Function->GetFName());
+	}
+}
+
+
+void FViewModelBlueprintCompilerContext::FinishCompilingClass(UClass* Class)
+{
+	Super::FinishCompilingClass(Class);
+
+	const bool bIsSkeletonOnly = CompileOptions.CompileType == EKismetCompileType::SkeletonOnly;
+	if (!bIsSkeletonOnly)
+	{
+		if (UMVVMViewModelBlueprintGeneratedClass* BPGClass = CastChecked<UMVVMViewModelBlueprintGeneratedClass>(Class))
+		{
+			if (UMVVMViewModelBase* ViewModelBase = Cast<UMVVMViewModelBase>(BPGClass->GetDefaultObject()))
+			{
+				BPGClass->InitializeFieldNotification(ViewModelBase);
+			}
+		}
+	}
+}
+
+} //namespace
+
+#undef LOCTEXT_NAMESPACE
diff --git a/Engine/Plugins/Runtime/ModelViewViewModel/Source/ModelViewViewModelBlueprint/Private/ViewModel/MVVMViewModelBlueprintCompiler.h b/Engine/Plugins/Runtime/ModelViewViewModel/Source/ModelViewViewModelBlueprint/Private/ViewModel/MVVMViewModelBlueprintCompiler.h
new file mode 100644
--- /dev/null
+++ b/Engine/Plugins/Runtime/ModelViewViewModel/Source/ModelViewViewModelBlueprint/Private/ViewModel/MVVMViewModelBlueprintCompiler.h
@@ -0,0 +1,72 @@
+// Copyright Epic Games, Inc. All Rights Reserved.
+
+#pragma once
+
+#include "KismetCompiler.h"
+#include "KismetCompilerModule.h"
+
+class UEdGraph;
+class UMVVMViewModelBlueprint;
+class UMVVMViewModelBlueprintGeneratedClass;
+
+//////////////////////////////////////////////////////////////////////////
+// FViewModelBlueprintCompiler 
+
+namespace UE::MVVM
+{
+
+class FViewModelBlueprintCompiler : public IBlueprintCompiler
+{
+public:
+	bool CanCompile(const UBlueprint* Blueprint) override;
+	void Compile(UBlueprint* Blueprint, const FKismetCompilerOptions& CompileOptions, FCompilerResultsLog& Results) override;
+	bool GetBlueprintTypesForClass(UClass* ParentClass, UClass*& OutBlueprintClass, UClass*& OutBlueprintGeneratedClass) const override;
+};
+
+} //namespace
+
+
+//////////////////////////////////////////////////////////////////////////
+// FViewModelBlueprintCompilerContext
+
+namespace UE::MVVM
+{
+
+class FViewModelBlueprintCompilerContext : public FKismetCompilerContext
+{
+	typedef FKismetCompilerContext Super;
+
+public:
+	using FKismetCompilerContext::FKismetCompilerContext;
+
+private:
+	UMVVMViewModelBlueprint* GetViewModelBlueprint() const;
+	FProperty* FindPropertyByNameOnNewClass(FName PropertyName) const;
+
+	//~ Begin FKismetCompilerContext Interface
+	virtual void SaveSubObjectsFromCleanAndSanitizeClass(FSubobjectCollection& SubObjectsToSave, UBlueprintGeneratedClass* ClassToClean) override;
+	virtual void CleanAndSanitizeClass(UBlueprintGeneratedClass* ClassToClean, UObject*& InOutOldCDO) override;
+	virtual void CreateFunctionList() override;
+	virtual void CreateClassVariablesFromBlueprint() override;
+	virtual void SpawnNewClass(const FString& NewClassName) override;
+	virtual void OnNewClassSet(UBlueprintGeneratedClass* ClassToUse) override;
+	virtual void PostcompileFunction(FKismetFunctionContext& Context) override;
+	virtual void FinishCompilingClass(UClass* Class) override;
+	//~ End FKismetCompilerContext Interface
+
+private:
+	struct FGeneratedFunction
+	{
+		FName PropertyName;
+		UEdGraph* SetterFunction = nullptr;
+		UEdGraph* NetRepFunction = nullptr;
+		FProperty* Property = nullptr;
+		FProperty* SkelProperty = nullptr;
+	};
+
+	UMVVMViewModelBlueprintGeneratedClass* NewViewModelBlueprintClass;
+	TArray<FGeneratedFunction> GeneratedFunctions;
+};
+
+} //namespace
+
diff --git a/Engine/Plugins/Runtime/ModelViewViewModel/Source/ModelViewViewModelBlueprint/Public/ViewModel/MVVMViewModelBlueprint.h b/Engine/Plugins/Runtime/ModelViewViewModel/Source/ModelViewViewModelBlueprint/Public/ViewModel/MVVMViewModelBlueprint.h
new file mode 100644
--- /dev/null
+++ b/Engine/Plugins/Runtime/ModelViewViewModel/Source/ModelViewViewModelBlueprint/Public/ViewModel/MVVMViewModelBlueprint.h
@@ -0,0 +1,48 @@
+// Copyright Epic Games, Inc. All Rights Reserved.
+
+#pragma once
+
+#include "Engine/Blueprint.h"
+
+#include "MVVMViewModelBlueprint.generated.h"
+
+class FKismetCompilerContext;
+
+class FCompilerResultsLog;
+class UEdGraph;
+class UMVVMViewModelBlueprintGeneratedClass;
+
+
+/**
+ * 
+ */
+UCLASS(BlueprintType, meta=(IgnoreClassThumbnail))
+class MODELVIEWVIEWMODELBLUEPRINT_API UMVVMViewModelBlueprint : public UBlueprint
+{
+	GENERATED_BODY()
+
+public:
+	/** */
+	static TSharedPtr<FKismetCompilerContext> GetCompilerForViewModelBlueprint(UBlueprint* BP, FCompilerResultsLog& InMessageLog, const FKismetCompilerOptions& InCompileOptions);
+
+	/** Get the (full) generated class for this viewmodel blueprint */
+	UMVVMViewModelBlueprintGeneratedClass* GetViewModelBlueprintGeneratedClass() const;
+
+	/** Get the (skeleton) generated class for this viewmodel blueprint */
+	UMVVMViewModelBlueprintGeneratedClass* GetViewModelBlueprintSkeletonClass() const;
+
+protected:
+	virtual bool SupportsMacros() const override;
+	virtual bool SupportsEventGraphs() const override;
+	virtual bool SupportsDelegates() const override;
+	virtual bool SupportsAnimLayers() const override;
+
+public:
+	// Use during compilation to clean the automatically generated graph.
+	UPROPERTY(Transient)
+	TArray<TObjectPtr<UEdGraph>> TemporaryGraph;
+};
+
+#if UE_ENABLE_INCLUDE_ORDER_DEPRECATED_IN_5_2
+#include "CoreMinimal.h"
+#endif
diff --git a/Engine/Plugins/Runtime/ModelViewViewModel/Source/ModelViewViewModelEditor/Private/ModelViewViewModelEditorModule.cpp b/Engine/Plugins/Runtime/ModelViewViewModel/Source/ModelViewViewModelEditor/Private/ModelViewViewModelEditorModule.cpp
--- a/Engine/Plugins/Runtime/ModelViewViewModel/Source/ModelViewViewModelEditor/Private/ModelViewViewModelEditorModule.cpp
+++ b/Engine/Plugins/Runtime/ModelViewViewModel/Source/ModelViewViewModelEditor/Private/ModelViewViewModelEditorModule.cpp
@@ -20,6 +20,7 @@
 #include "Tabs/MVVMBindingSummoner.h"
 #include "Tabs/MVVMViewModelSummoner.h"
 #include "UMGEditorModule.h"
+#include "ViewModel/AssetTypeActions_ViewModelBlueprint.h"
 #include "WidgetBlueprintEditor.h"
 #include "WidgetDrawerConfig.h"
 
@@ -35,6 +36,16 @@ void FModelViewViewModelEditorModule::StartupModule()
 	IUMGEditorModule& UMGEditorModule = FModuleManager::LoadModuleChecked<IUMGEditorModule>("UMGEditor");
 	UMGEditorModule.OnRegisterTabsForEditor().AddRaw(this, &FModelViewViewModelEditorModule::HandleRegisterBlueprintEditorTab);
 
+	// Register asset types
+	{
+#if UE_MVVM_WITH_VIEWMODEL_EDITOR
+		// Only remove what is related to the viewmodel editor, not the UMG extention for view.
+		IAssetTools& AssetTools = FModuleManager::LoadModuleChecked<FAssetToolsModule>("AssetTools").Get();
+		ViewModelBlueprintActions = MakeShared<UE::MVVM::FAssetTypeActions_ViewModelBlueprint>();
+		AssetTools.RegisterAssetTypeActions(ViewModelBlueprintActions.ToSharedRef());
+#endif
+	}
+
 	PropertyBindingExtension = MakeShared<FMVVMPropertyBindingExtension>();
 	UMGEditorModule.GetPropertyBindingExtensibilityManager()->AddExtension(PropertyBindingExtension.ToSharedRef());
 
@@ -80,6 +91,17 @@ void FModelViewViewModelEditorModule::ShutdownModule()
 	}
 	PropertyBindingExtension.Reset();
 
+#if UE_MVVM_WITH_VIEWMODEL_EDITOR
+	// Unregister all the asset types that we registered
+	if (FAssetToolsModule* AssetTools = FModuleManager::GetModulePtr<FAssetToolsModule>("AssetTools"))
+	{
+		if (ViewModelBlueprintActions)
+		{
+			AssetTools->Get().UnregisterAssetTypeActions(ViewModelBlueprintActions.ToSharedRef());
+		}
+	}
+#endif
+
 	FMVVMEditorStyle::DestroyInstance();
 
 	FMVVMEditorCommands::Unregister();
diff --git a/Engine/Plugins/Runtime/ModelViewViewModel/Source/ModelViewViewModelEditor/Private/ViewModel/AssetTypeActions_ViewModelBlueprint.cpp b/Engine/Plugins/Runtime/ModelViewViewModel/Source/ModelViewViewModelEditor/Private/ViewModel/AssetTypeActions_ViewModelBlueprint.cpp
new file mode 100644
--- /dev/null
+++ b/Engine/Plugins/Runtime/ModelViewViewModel/Source/ModelViewViewModelEditor/Private/ViewModel/AssetTypeActions_ViewModelBlueprint.cpp
@@ -0,0 +1,72 @@
+// Copyright Epic Games, Inc. All Rights Reserved.
+
+#include "AssetTypeActions_ViewModelBlueprint.h"
+#include "Misc/MessageDialog.h"
+#include "Kismet2/BlueprintEditorUtils.h"
+#include "SBlueprintDiff.h"
+
+#include "ViewModel/MVVMViewModelBlueprint.h"
+#include "ViewModel/MVVMViewModelBlueprintEditor.h"
+
+#define LOCTEXT_NAMESPACE "AssetTypeActions"
+
+#if UE_MVVM_WITH_VIEWMODEL_EDITOR
+namespace UE::MVVM
+{
+
+void FAssetTypeActions_ViewModelBlueprint::OpenAssetEditor( const TArray<UObject*>& InObjects, TSharedPtr<IToolkitHost> EditWithinLevelEditor )
+{
+	EToolkitMode::Type Mode = EditWithinLevelEditor.IsValid() ? EToolkitMode::WorldCentric : EToolkitMode::Standalone;
+
+	for (auto ObjIt = InObjects.CreateConstIterator(); ObjIt; ++ObjIt)
+	{
+		UBlueprint* Blueprint = Cast<UBlueprint>(*ObjIt);
+		if (Blueprint && Blueprint->SkeletonGeneratedClass && Blueprint->GeneratedClass )
+		{
+			TSharedRef< FMVVMViewModelBlueprintEditor > NewBlueprintEditor(new FMVVMViewModelBlueprintEditor());
+
+			bool IsDataOnlyBlueprint = FBlueprintEditorUtils::IsDataOnlyBlueprint(Blueprint)
+				&& !Blueprint->bForceFullEditor
+				&& !Blueprint->bIsNewlyCreated;
+			TArray<UBlueprint*> Blueprints;
+			Blueprints.Add(Blueprint);
+			NewBlueprintEditor->InitViewModelBlueprintEditor(Mode, EditWithinLevelEditor, Blueprints, IsDataOnlyBlueprint);
+		}
+		else
+		{
+			FMessageDialog::Open( EAppMsgType::Ok, LOCTEXT("FailedToLoadViewModelBlueprint", "Viewmodel Blueprint could not be loaded because it derives from an invalid class.\nCheck to make sure the parent class for this blueprint hasn't been removed!"));
+		}
+	}
+}
+
+
+UClass* FAssetTypeActions_ViewModelBlueprint::GetSupportedClass() const
+{
+	return UMVVMViewModelBlueprint::StaticClass();
+}
+
+
+FText FAssetTypeActions_ViewModelBlueprint::GetAssetDescription( const FAssetData& AssetData ) const
+{
+	FString Description = AssetData.GetTagValueRef<FString>( GET_MEMBER_NAME_CHECKED( UBlueprint, BlueprintDescription ) );
+	if ( !Description.IsEmpty() )
+	{
+		Description.ReplaceInline( TEXT( "\\n" ), TEXT( "\n" ) );
+		return FText::FromString( MoveTemp(Description) );
+	}
+
+	return FText::GetEmpty();
+}
+
+
+void FAssetTypeActions_ViewModelBlueprint::PerformAssetDiff(UObject* Asset1, UObject* Asset2, const struct FRevisionInfo& OldRevision, const struct FRevisionInfo& NewRevision) const
+{
+	UBlueprint* OldBlueprint = Cast<UBlueprint>(Asset1);
+	UBlueprint* NewBlueprint = Cast<UBlueprint>(Asset2);
+	SBlueprintDiff::CreateDiffWindow(OldBlueprint, NewBlueprint, OldRevision, NewRevision, GetSupportedClass());
+}
+
+} //namespace
+#endif
+
+#undef LOCTEXT_NAMESPACE
diff --git a/Engine/Plugins/Runtime/ModelViewViewModel/Source/ModelViewViewModelEditor/Private/ViewModel/AssetTypeActions_ViewModelBlueprint.h b/Engine/Plugins/Runtime/ModelViewViewModel/Source/ModelViewViewModelEditor/Private/ViewModel/AssetTypeActions_ViewModelBlueprint.h
new file mode 100644
--- /dev/null
+++ b/Engine/Plugins/Runtime/ModelViewViewModel/Source/ModelViewViewModelEditor/Private/ViewModel/AssetTypeActions_ViewModelBlueprint.h
@@ -0,0 +1,29 @@
+// Copyright Epic Games, Inc. All Rights Reserved.
+
+#pragma once
+
+#include "AssetTypeActions_Base.h"
+
+struct FAssetData;
+
+#if UE_MVVM_WITH_VIEWMODEL_EDITOR
+namespace UE::MVVM
+{
+
+class FAssetTypeActions_ViewModelBlueprint : public FAssetTypeActions_Base
+{
+public:
+	// IAssetTypeActions Implementation
+	virtual FText GetName() const override { return NSLOCTEXT("AssetTypeActions", "AssetTypeActions_ViewModelBlueprint", "Viewmodel Blueprint"); }
+	virtual FColor GetTypeColor() const override { return FColor(44, 180, 180); }
+	virtual UClass* GetSupportedClass() const override;
+	virtual void OpenAssetEditor(const TArray<UObject*>& InObjects, TSharedPtr<class IToolkitHost> EditWithinLevelEditor = TSharedPtr<IToolkitHost>()) override;
+	virtual uint32 GetCategories() override { return EAssetTypeCategories::UI; }
+	virtual bool CanLocalize() const override { return false; }
+	virtual void PerformAssetDiff(UObject* Asset1, UObject* Asset2, const struct FRevisionInfo& OldRevision, const struct FRevisionInfo& NewRevision) const override;
+	virtual FText GetAssetDescription( const FAssetData& AssetData ) const override;
+	// End IAssetTypeActions Implementation
+};
+
+} //namespace
+#endif
diff --git a/Engine/Plugins/Runtime/ModelViewViewModel/Source/ModelViewViewModelEditor/Private/ViewModel/MVVMViewModelBlueprintEditor.cpp b/Engine/Plugins/Runtime/ModelViewViewModel/Source/ModelViewViewModelEditor/Private/ViewModel/MVVMViewModelBlueprintEditor.cpp
new file mode 100644
--- /dev/null
+++ b/Engine/Plugins/Runtime/ModelViewViewModel/Source/ModelViewViewModelEditor/Private/ViewModel/MVVMViewModelBlueprintEditor.cpp
@@ -0,0 +1,103 @@
+// Copyright Epic Games, Inc. All Rights Reserved.
+
+#include "ViewModel/MVVMViewModelBlueprintEditor.h"
+#include "ToolMenuContext.h"
+#include "ViewModel/MVVMViewModelBlueprint.h"
+#include "ViewModel/MVVMViewModelBlueprintToolMenuContext.h"
+
+#define LOCTEXT_NAMESPACE "MVVMViewModelBlueprintEditor"
+
+void FMVVMViewModelBlueprintEditor::InitViewModelBlueprintEditor(const EToolkitMode::Type Mode, const TSharedPtr< IToolkitHost >& InitToolkitHost, const TArray<UBlueprint*>& InBlueprints, bool bShouldOpenInDefaultsMode)
+{
+	Super::InitBlueprintEditor(Mode, InitToolkitHost, InBlueprints, bShouldOpenInDefaultsMode);
+}
+
+
+UMVVMViewModelBlueprint* FMVVMViewModelBlueprintEditor::GetViewModelBlueprintObj() const
+{
+	return Cast<UMVVMViewModelBlueprint>(GetBlueprintObj());
+}
+
+
+FName FMVVMViewModelBlueprintEditor::GetToolkitFName() const
+{
+	return FName("ViewModelBlueprintEditor");
+}
+
+
+FText FMVVMViewModelBlueprintEditor::GetBaseToolkitName() const
+{
+	return LOCTEXT("AppLabel", "Viewmodel Editor");
+}
+
+
+FString FMVVMViewModelBlueprintEditor::GetWorldCentricTabPrefix() const
+{
+	return LOCTEXT("WorldCentricTabPrefix", "Viewmodel Editor ").ToString();
+}
+
+
+FLinearColor FMVVMViewModelBlueprintEditor::GetWorldCentricTabColorScale() const
+{
+	return FLinearColor(0.3f, 0.25f, 0.35f, 0.5f);
+}
+
+
+void FMVVMViewModelBlueprintEditor::InitToolMenuContext(FToolMenuContext& MenuContext)
+{
+	Super::InitToolMenuContext(MenuContext);
+
+	UMVVMViewModelBlueprintToolMenuContext* Context = NewObject<UMVVMViewModelBlueprintToolMenuContext>();
+	Context->ViewModelBlueprintEditor = SharedThis(this);
+	MenuContext.AddObject(Context);
+}
+
+
+void FMVVMViewModelBlueprintEditor::RegisterApplicationModes(const TArray<UBlueprint*>& InBlueprints, bool bShouldOpenInDefaultsMode, bool bNewlyCreated/* = false*/)
+{
+	Super::RegisterApplicationModes(InBlueprints, bShouldOpenInDefaultsMode);
+}
+
+
+FGraphAppearanceInfo FMVVMViewModelBlueprintEditor::GetGraphAppearance(UEdGraph* InGraph) const
+{
+	FGraphAppearanceInfo AppearanceInfo = Super::GetGraphAppearance(InGraph);
+
+	if (GetBlueprintObj()->IsA(UMVVMViewModelBlueprint::StaticClass()))
+	{
+		AppearanceInfo.CornerText = LOCTEXT("AppearanceCornerText", "VIEWMODEL BLUEPRINT");
+	}
+
+	return AppearanceInfo;
+}
+
+
+bool FMVVMViewModelBlueprintEditor::NewDocument_IsVisibleForType(ECreatedDocumentType InGraphType) const
+{
+	switch (InGraphType)
+	{
+	case CGT_NewVariable:
+	case CGT_NewFunctionGraph:
+		return true;
+	default:
+		break;
+	}
+	return false;
+}
+
+
+bool FMVVMViewModelBlueprintEditor::IsSectionVisible(NodeSectionID::Type InSectionID) const
+{
+	switch (InSectionID)
+	{
+	case NodeSectionID::VARIABLE:
+	case NodeSectionID::FUNCTION:
+	case NodeSectionID::LOCAL_VARIABLE:
+		return true;
+	default:
+		break;
+	}
+	return false;
+}
+
+#undef LOCTEXT_NAMESPACE
diff --git a/Engine/Plugins/Runtime/ModelViewViewModel/Source/ModelViewViewModelEditor/Private/ViewModel/MVVMViewModelBlueprintFactory.cpp b/Engine/Plugins/Runtime/ModelViewViewModel/Source/ModelViewViewModelEditor/Private/ViewModel/MVVMViewModelBlueprintFactory.cpp
new file mode 100644
--- /dev/null
+++ b/Engine/Plugins/Runtime/ModelViewViewModel/Source/ModelViewViewModelEditor/Private/ViewModel/MVVMViewModelBlueprintFactory.cpp
@@ -0,0 +1,133 @@
+// Copyright Epic Games, Inc. All Rights Reserved.
+
+#include "ViewModel/MVVMViewModelBlueprintFactory.h"
+#include "MVVMViewModelBase.h"
+#include "ViewModel/MVVMViewModelBlueprint.h"
+#include "ViewModel/MVVMViewModelBlueprintGeneratedClass.h"
+
+#include "ClassViewerFilter.h"
+#include "Kismet2/KismetEditorUtilities.h"
+#include "Kismet2/SClassPickerDialog.h"
+#include "Misc/MessageDialog.h"
+#include "Modules/ModuleManager.h"
+
+#include UE_INLINE_GENERATED_CPP_BY_NAME(MVVMViewModelBlueprintFactory)
+
+#define LOCTEXT_NAMESPACE "ViewModelBlueprintFactory"
+
+#if UE_MVVM_WITH_VIEWMODEL_EDITOR
+
+/*------------------------------------------------------------------------------
+	UMVVMViewModelBlueprintFactory implementation.
+------------------------------------------------------------------------------*/
+
+namespace UE::MVVM::Private
+{
+
+class FViewModelClassFilter : public IClassViewerFilter
+{
+public:
+	/** All children of these classes will be included unless filtered out by another setting. */
+	TSet <const UClass*> AllowedChildrenOfClasses;
+
+	/** Disallowed class flags. */
+	EClassFlags DisallowedClassFlags;
+
+	virtual bool IsClassAllowed(const FClassViewerInitializationOptions& InInitOptions, const UClass* InClass, TSharedRef< FClassViewerFilterFuncs > InFilterFuncs) override
+	{
+		return !InClass->HasAnyClassFlags(DisallowedClassFlags)
+			&& InFilterFuncs->IfInChildOfClassesSet(AllowedChildrenOfClasses, InClass) != EFilterReturn::Failed;
+	}
+
+	virtual bool IsUnloadedClassAllowed(const FClassViewerInitializationOptions& InInitOptions, const TSharedRef< const IUnloadedBlueprintData > InUnloadedClassData, TSharedRef< FClassViewerFilterFuncs > InFilterFuncs) override
+	{
+		return !InUnloadedClassData->HasAnyClassFlags(DisallowedClassFlags)
+			&& InFilterFuncs->IfInChildOfClassesSet(AllowedChildrenOfClasses, InUnloadedClassData) != EFilterReturn::Failed;
+	}
+};
+
+}//namespace
+
+#endif
+
+UMVVMViewModelBlueprintFactory::UMVVMViewModelBlueprintFactory()
+{
+#if UE_MVVM_WITH_VIEWMODEL_EDITOR
+	bCreateNew = true;
+#else
+	bCreateNew = false;
+#endif
+	bEditAfterNew = true;
+	SupportedClass = UMVVMViewModelBlueprint::StaticClass();
+	ParentClass = UMVVMViewModelBase::StaticClass();
+}
+
+#if UE_MVVM_WITH_VIEWMODEL_EDITOR
+bool UMVVMViewModelBlueprintFactory::ConfigureProperties()
+{
+	FClassViewerModule& ClassViewerModule = FModuleManager::LoadModuleChecked<FClassViewerModule>("ClassViewer");
+
+	// Fill in options
+	FClassViewerInitializationOptions Options;
+	Options.DisplayMode = EClassViewerDisplayMode::Type::TreeView;
+	Options.Mode = EClassViewerMode::ClassPicker;
+	Options.NameTypeToDisplay = EClassViewerNameTypeToDisplay::DisplayName;
+	Options.bShowNoneOption = false;
+	Options.bExpandAllNodes = true;
+
+	TSharedRef<UE::MVVM::Private::FViewModelClassFilter> Filter = MakeShared<UE::MVVM::Private::FViewModelClassFilter>();
+	Filter->DisallowedClassFlags = CLASS_Deprecated | CLASS_NewerVersionExists;
+	Filter->AllowedChildrenOfClasses.Add(UMVVMViewModelBase::StaticClass());
+
+	Options.ClassFilters.Add(Filter);
+	Options.ExtraPickerCommonClasses.Add(UMVVMViewModelBase::StaticClass());
+
+	const FText TitleText = LOCTEXT("CreateViewmodelBlueprint", "Pick Root Viewmodel for New ViewModel Blueprint");
+	UClass* ChosenParentClass = nullptr;
+	if (SClassPickerDialog::PickClass(TitleText, Options, ChosenParentClass, UMVVMViewModelBase::StaticClass()))
+	{
+		ParentClass = ChosenParentClass ? ChosenParentClass : UMVVMViewModelBase::StaticClass();
+	}
+
+	return true;
+}
+
+bool UMVVMViewModelBlueprintFactory::ShouldShowInNewMenu() const
+{
+	return true;
+}
+
+UObject* UMVVMViewModelBlueprintFactory::FactoryCreateNew(UClass* Class, UObject* InParent, FName Name, EObjectFlags Flags, UObject* Context, FFeedbackContext* Warn, FName CallingContext)
+{
+	// Make sure we are trying to factory a Anim Blueprint, then create and init one
+	check(Class->IsChildOf(UMVVMViewModelBlueprint::StaticClass()));
+
+	UClass* CurrentParentClass = ParentClass;
+	if (CurrentParentClass == nullptr)
+	{
+		CurrentParentClass = UMVVMViewModelBase::StaticClass();
+	}
+
+	if ((CurrentParentClass == nullptr) || !FKismetEditorUtilities::CanCreateBlueprintOfClass(CurrentParentClass) || !CurrentParentClass->IsChildOf(UMVVMViewModelBase::StaticClass()))
+	{
+		FFormatNamedArguments Args;
+		Args.Add(TEXT("ClassName"), CurrentParentClass ? FText::FromString(CurrentParentClass->GetName()) : LOCTEXT("Null", "(null)"));
+		FMessageDialog::Open(EAppMsgType::Ok, FText::Format(LOCTEXT("CannotCreateViewmodelBlueprint", "Cannot create a Viewmodel Blueprint based on the class '{ClassName}'."), Args));
+		return nullptr;
+	}
+	else
+	{
+
+		return CastChecked<UMVVMViewModelBlueprint>(FKismetEditorUtilities::CreateBlueprint(CurrentParentClass, InParent, Name, EBlueprintType::BPTYPE_Normal, UMVVMViewModelBlueprint::StaticClass(), UMVVMViewModelBlueprintGeneratedClass::StaticClass(), CallingContext));
+	}
+}
+
+UObject* UMVVMViewModelBlueprintFactory::FactoryCreateNew(UClass* Class, UObject* InParent, FName Name, EObjectFlags Flags, UObject* Context, FFeedbackContext* Warn)
+{
+	return FactoryCreateNew(Class, InParent, Name, Flags, Context, Warn, NAME_None);
+}
+
+#endif
+
+#undef LOCTEXT_NAMESPACE
+
diff --git a/Engine/Plugins/Runtime/ModelViewViewModel/Source/ModelViewViewModelEditor/Private/ViewModel/MVVMViewModelBlueprintFactory.h b/Engine/Plugins/Runtime/ModelViewViewModel/Source/ModelViewViewModelEditor/Private/ViewModel/MVVMViewModelBlueprintFactory.h
new file mode 100644
--- /dev/null
+++ b/Engine/Plugins/Runtime/ModelViewViewModel/Source/ModelViewViewModelEditor/Private/ViewModel/MVVMViewModelBlueprintFactory.h
@@ -0,0 +1,30 @@
+// Copyright Epic Games, Inc. All Rights Reserved.
+
+#pragma once
+
+#include "Factories/Factory.h"
+
+#include "MVVMViewModelBlueprintFactory.generated.h"
+
+UCLASS(HideCategories=Object, MinimalAPI)
+class UMVVMViewModelBlueprintFactory : public UFactory
+{
+	GENERATED_BODY()
+
+public:
+	UMVVMViewModelBlueprintFactory();
+
+#if UE_MVVM_WITH_VIEWMODEL_EDITOR
+	//~ Begin UFactory Interface
+	virtual bool ConfigureProperties() override;
+	virtual bool ShouldShowInNewMenu() const override;
+	virtual UObject* FactoryCreateNew(UClass* Class, UObject* InParent, FName Name, EObjectFlags Flags, UObject* Context, FFeedbackContext* Warn, FName CallingContext) override;
+	virtual UObject* FactoryCreateNew(UClass* Class, UObject* InParent, FName Name, EObjectFlags Flags, UObject* Context, FFeedbackContext* Warn) override;
+	//~ End UFactory Interface
+#endif
+
+private:
+	// The parent class of the created blueprint
+	UPROPERTY(EditAnywhere, Category = "Viewmodel", meta=(AllowAbstract="true"))
+	TSubclassOf<class UMVVMViewModelBase> ParentClass;
+};
diff --git a/Engine/Plugins/Runtime/ModelViewViewModel/Source/ModelViewViewModelEditor/Public/ViewModel/MVVMViewModelBlueprintEditor.h b/Engine/Plugins/Runtime/ModelViewViewModel/Source/ModelViewViewModelEditor/Public/ViewModel/MVVMViewModelBlueprintEditor.h
new file mode 100644
--- /dev/null
+++ b/Engine/Plugins/Runtime/ModelViewViewModel/Source/ModelViewViewModelEditor/Public/ViewModel/MVVMViewModelBlueprintEditor.h
@@ -0,0 +1,55 @@
+// Copyright Epic Games, Inc. All Rights Reserved.
+
+#pragma once
+
+#include "BlueprintEditor.h"
+
+struct FGraphAppearanceInfo;
+
+class UMVVMViewModelBlueprint;
+
+/**
+ * Viewmodel blueprint editor (extends Blueprint editor)
+ */
+class MODELVIEWVIEWMODELEDITOR_API FMVVMViewModelBlueprintEditor : public FBlueprintEditor
+{
+private:
+	using Super = FBlueprintEditor;
+
+
+public:
+	FMVVMViewModelBlueprintEditor() = default;
+	//virtual ~FMVVMViewModelBlueprintEditor();
+
+	void InitViewModelBlueprintEditor(const EToolkitMode::Type Mode, const TSharedPtr< IToolkitHost >& InitToolkitHost, const TArray<UBlueprint*>& InBlueprints, bool bShouldOpenInDefaultsMode);
+
+	//~ Begin FBlueprintEditor interface
+	//virtual void Tick(float DeltaTime) override;
+	//virtual void PostUndo(bool bSuccessful) override;
+	//virtual void PostRedo(bool bSuccessful) override;
+	//virtual void Compile() override;
+	//~ End FBlueprintEditor interface
+
+	// ~ Begin IToolkit interface
+	virtual FName GetToolkitFName() const override;
+	virtual FText GetBaseToolkitName() const override;
+	virtual FString GetWorldCentricTabPrefix() const override;
+	virtual FLinearColor GetWorldCentricTabColorScale() const override;
+	virtual void InitToolMenuContext(FToolMenuContext& MenuContext) override;
+	// ~ End IToolkit interface
+
+	/** @return The viewmodel blueprint currently being edited in this editor */
+	UMVVMViewModelBlueprint* GetViewModelBlueprintObj() const;
+
+protected:
+	//~ Begin FBlueprintEditor interface
+	virtual void RegisterApplicationModes(const TArray<UBlueprint*>& Blueprints, bool bShouldOpenInDefaultsMode, bool bNewlyCreated = false) override;
+	virtual FGraphAppearanceInfo GetGraphAppearance(class UEdGraph* Graph) const override;
+	virtual bool NewDocument_IsVisibleForType(ECreatedDocumentType GraphType) const override;
+	virtual bool IsSectionVisible(NodeSectionID::Type SectionID) const override;
+	//~ End FBlueprintEditor interface
+};
+
+#if UE_ENABLE_INCLUDE_ORDER_DEPRECATED_IN_5_2
+#include "CoreMinimal.h"
+#endif
diff --git a/Engine/Plugins/Runtime/ModelViewViewModel/Source/ModelViewViewModelEditor/Public/ViewModel/MVVMViewModelBlueprintToolMenuContext.h b/Engine/Plugins/Runtime/ModelViewViewModel/Source/ModelViewViewModelEditor/Public/ViewModel/MVVMViewModelBlueprintToolMenuContext.h
new file mode 100644
--- /dev/null
+++ b/Engine/Plugins/Runtime/ModelViewViewModel/Source/ModelViewViewModelEditor/Public/ViewModel/MVVMViewModelBlueprintToolMenuContext.h
@@ -0,0 +1,21 @@
+// Copyright Epic Games, Inc. All Rights Reserved.
+
+#pragma once
+
+
+#include "MVVMViewModelBlueprintToolMenuContext.generated.h"
+
+class FMVVMViewModelBlueprintEditor;
+
+UCLASS()
+class MODELVIEWVIEWMODELEDITOR_API UMVVMViewModelBlueprintToolMenuContext : public UObject
+{
+	GENERATED_BODY()
+
+public:
+	TWeakPtr<FMVVMViewModelBlueprintEditor> ViewModelBlueprintEditor;
+};
+
+#if UE_ENABLE_INCLUDE_ORDER_DEPRECATED_IN_5_2
+#include "CoreMinimal.h"
+#endif
