MVVM: Uses the BlueprintGeneratedClass for finding the FieldNotify flag make in Blueprint.

diff --git a/Engine/Plugins/Runtime/ModelViewViewModel/Source/ModelViewViewModel/Private/MVVMViewModelBase.cpp b/Engine/Plugins/Runtime/ModelViewViewModel/Source/ModelViewViewModel/Private/MVVMViewModelBase.cpp
--- a/Engine/Plugins/Runtime/ModelViewViewModel/Source/ModelViewViewModel/Private/MVVMViewModelBase.cpp
+++ b/Engine/Plugins/Runtime/ModelViewViewModel/Source/ModelViewViewModel/Private/MVVMViewModelBase.cpp
@@ -63,9 +63,9 @@ const UE::FieldNotification::IClassDescriptor& UMVVMViewModelBase::GetFieldNotif
 
 void UMVVMViewModelBase::FFieldNotificationClassDescriptor::ForEachField(const UClass* Class, TFunctionRef<bool(::UE::FieldNotification::FFieldId FieldId)> Callback) const
 {
-	if (const UMVVMViewModelBlueprintGeneratedClass* ViewModelBPClass = Cast<const UMVVMViewModelBlueprintGeneratedClass>(Class))
+	if (const UBlueprintGeneratedClass* ViewModelBPClass = Cast<const UBlueprintGeneratedClass>(Class))
 	{
-		ViewModelBPClass->ForEachField(Callback);
+		ViewModelBPClass->ForEachFieldNotify(Callback, true);
 	}
 }
 
diff --git a/Engine/Source/Runtime/UMG/Private/Components/Widget.cpp b/Engine/Source/Runtime/UMG/Private/Components/Widget.cpp
--- a/Engine/Source/Runtime/UMG/Private/Components/Widget.cpp
+++ b/Engine/Source/Runtime/UMG/Private/Components/Widget.cpp
@@ -172,9 +172,9 @@ void UWidget::FFieldNotificationClassDescriptor::ForEachField(const UClass* Clas
 			return;
 		}
 	}
-	if (const UWidgetBlueprintGeneratedClass* WidgetBPClass = Cast<const UWidgetBlueprintGeneratedClass>(Class))
+	if (const UBlueprintGeneratedClass* ViewModelBPClass = Cast<const UBlueprintGeneratedClass>(Class))
 	{
-		WidgetBPClass->ForEachField(Callback);
+		ViewModelBPClass->ForEachFieldNotify(Callback, true);
 	}
 }
 UE_FIELD_NOTIFICATION_IMPLEMENT_CLASS_DESCRIPTOR_ThreeFields(UWidget, ToolTipText, Visibility, bIsEnabled);
