MVVM: Remove the specialized viewmodel editor and use the Blueprint base editor instead.

diff --git a/Engine/Plugins/Runtime/ModelViewViewModel/Source/ModelViewViewModel/Private/ViewModel/MVVMViewModelBlueprintGeneratedClass.cpp b/Engine/Plugins/Runtime/ModelViewViewModel/Source/ModelViewViewModel/Private/ViewModel/MVVMViewModelBlueprintGeneratedClass.cpp
--- a/Engine/Plugins/Runtime/ModelViewViewModel/Source/ModelViewViewModel/Private/ViewModel/MVVMViewModelBlueprintGeneratedClass.cpp
+++ b/Engine/Plugins/Runtime/ModelViewViewModel/Source/ModelViewViewModel/Private/ViewModel/MVVMViewModelBlueprintGeneratedClass.cpp
@@ -6,51 +6,3 @@
 
 #include UE_INLINE_GENERATED_CPP_BY_NAME(MVVMViewModelBlueprintGeneratedClass)
 
-UMVVMViewModelBlueprintGeneratedClass::UMVVMViewModelBlueprintGeneratedClass()
-	: FieldNotifyStartBitNumber(INDEX_NONE)
-{}
-
-
-void UMVVMViewModelBlueprintGeneratedClass::PostLoadDefaultObject(UObject* Object)
-{
-	Super::PostLoadDefaultObject(Object);
-	InitializeFieldNotification(Cast<UMVVMViewModelBase>(Object));
-}
-
-
-void UMVVMViewModelBlueprintGeneratedClass::PurgeClass(bool bRecompilingOnLoad)
-{
-	Super::PurgeClass(bRecompilingOnLoad);
-	FieldNotifyNames.Empty();
-}
-
-
-void UMVVMViewModelBlueprintGeneratedClass::InitializeFieldNotification(const UMVVMViewModelBase* ViewModel)
-{
-	FieldNotifyStartBitNumber = 0;
-	if (ViewModel && FieldNotifyNames.Num())
-	{
-		int32 NumberOfField = 0;
-		ViewModel->GetFieldNotificationDescriptor().ForEachField(this, [&NumberOfField](::UE::FieldNotification::FFieldId FielId)
-			{
-				++NumberOfField;
-				return true;
-			});
-		FieldNotifyStartBitNumber = NumberOfField - FieldNotifyNames.Num();
-		ensureMsgf(FieldNotifyStartBitNumber >= 0, TEXT("The FieldNotifyStartIndex is negative. The number of field should be positive."));
-	}
-}
-
-
-void UMVVMViewModelBlueprintGeneratedClass::ForEachField(TFunctionRef<bool(::UE::FieldNotification::FFieldId FieldId)> Callback) const
-{
-	ensureMsgf(FieldNotifyStartBitNumber >= 0, TEXT("The FieldNotifyStartIndex is negative. The number of field should be positive."));
-	for (int32 Index = 0; Index < FieldNotifyNames.Num(); ++Index)
-	{
-		if (!Callback(UE::FieldNotification::FFieldId(FieldNotifyNames[Index].GetFieldName(), Index + FieldNotifyStartBitNumber)))
-		{
-			break;
-		}
-	}
-}
-
diff --git a/Engine/Plugins/Runtime/ModelViewViewModel/Source/ModelViewViewModel/Public/ViewModel/MVVMViewModelBlueprintGeneratedClass.h b/Engine/Plugins/Runtime/ModelViewViewModel/Source/ModelViewViewModel/Public/ViewModel/MVVMViewModelBlueprintGeneratedClass.h
--- a/Engine/Plugins/Runtime/ModelViewViewModel/Source/ModelViewViewModel/Public/ViewModel/MVVMViewModelBlueprintGeneratedClass.h
+++ b/Engine/Plugins/Runtime/ModelViewViewModel/Source/ModelViewViewModel/Public/ViewModel/MVVMViewModelBlueprintGeneratedClass.h
@@ -16,30 +16,11 @@ namespace UE::MVVM
 	class FViewModelBlueprintCompilerContext;
 }//namespace
 
-UCLASS()
-class MODELVIEWVIEWMODEL_API UMVVMViewModelBlueprintGeneratedClass : public UBlueprintGeneratedClass
+UCLASS(Deprecated)
+class UE_DEPRECATED(5.3, "The prototype viewmodel editor is deprecated. Use the regular Blueprint editor.")
+MODELVIEWVIEWMODEL_API UDEPRECATED_MVVMViewModelBlueprintGeneratedClass : public UBlueprintGeneratedClass
 {
 	GENERATED_BODY()
-
-	friend UE::MVVM::FViewModelBlueprintCompilerContext;
-
-public:
-	UMVVMViewModelBlueprintGeneratedClass();
-
-	//~ Begin UBlueprintGeneratedClass interface
-	virtual void PostLoadDefaultObject(UObject* Object) override;
-	virtual void PurgeClass(bool bRecompilingOnLoad) override;
-	//~ End UBlueprintGeneratedClass interface
-
-	void InitializeFieldNotification(const UMVVMViewModelBase* ViewModel);
-	void ForEachField(TFunctionRef<bool(::UE::FieldNotification::FFieldId FieldId)> Callback) const;
-
-private:
-	/** List Field Notifies. No index here on purpose to prevent saving them. */
-	UPROPERTY()
-	TArray<FFieldNotificationId> FieldNotifyNames;
-
-	int32 FieldNotifyStartBitNumber;
 };
 
 #if UE_ENABLE_INCLUDE_ORDER_DEPRECATED_IN_5_2
diff --git a/Engine/Plugins/Runtime/ModelViewViewModel/Source/ModelViewViewModelBlueprint/ModelViewViewModelBlueprint.Build.cs b/Engine/Plugins/Runtime/ModelViewViewModel/Source/ModelViewViewModelBlueprint/ModelViewViewModelBlueprint.Build.cs
--- a/Engine/Plugins/Runtime/ModelViewViewModel/Source/ModelViewViewModelBlueprint/ModelViewViewModelBlueprint.Build.cs
+++ b/Engine/Plugins/Runtime/ModelViewViewModel/Source/ModelViewViewModelBlueprint/ModelViewViewModelBlueprint.Build.cs
@@ -32,7 +32,5 @@ public class ModelViewViewModelBlueprint : ModuleRules
 				"UMGEditor",
 				"UnrealEd",
 			});
-
-		PublicDefinitions.Add("UE_MVVM_WITH_VIEWMODEL_EDITOR=0");
 	}
 }
diff --git a/Engine/Plugins/Runtime/ModelViewViewModel/Source/ModelViewViewModelBlueprint/Private/ModelViewViewModelBlueprintModule.cpp b/Engine/Plugins/Runtime/ModelViewViewModel/Source/ModelViewViewModelBlueprint/Private/ModelViewViewModelBlueprintModule.cpp
--- a/Engine/Plugins/Runtime/ModelViewViewModel/Source/ModelViewViewModelBlueprint/Private/ModelViewViewModelBlueprintModule.cpp
+++ b/Engine/Plugins/Runtime/ModelViewViewModel/Source/ModelViewViewModelBlueprint/Private/ModelViewViewModelBlueprintModule.cpp
@@ -3,36 +3,5 @@
 #include "ModelViewViewModelBlueprintModule.h"
 #include "Modules/ModuleManager.h"
 
-#include "ViewModel/MVVMViewModelBlueprint.h"
-#include "ViewModel/MVVMViewModelBlueprintCompiler.h"
 
-
-class FModelViewViewModelBlueprintModule : public IModuleInterface
-{
-public:
-	virtual void StartupModule() override
-	{
-#if UE_MVVM_WITH_VIEWMODEL_EDITOR
-		IKismetCompilerInterface& KismetCompilerModule = FModuleManager::GetModuleChecked<IKismetCompilerInterface>("KismetCompiler");
-		KismetCompilerModule.GetCompilers().Add(&ViewModelBlueprintCompiler);
-
-		FKismetCompilerContext::RegisterCompilerForBP(UMVVMViewModelBlueprint::StaticClass(), &UMVVMViewModelBlueprint::GetCompilerForViewModelBlueprint);
-#endif
-	}
-
-	virtual void ShutdownModule() override
-	{
-#if UE_MVVM_WITH_VIEWMODEL_EDITOR
-		if (IKismetCompilerInterface* KismetCompilerModule = FModuleManager::GetModulePtr<IKismetCompilerInterface>("KismetCompiler"))
-		{
-			KismetCompilerModule->GetCompilers().Remove(&ViewModelBlueprintCompiler);
-		}
-#endif
-	}
-
-private:
-	UE::MVVM::FViewModelBlueprintCompiler ViewModelBlueprintCompiler;
-};
-
-
-IMPLEMENT_MODULE(FModelViewViewModelBlueprintModule, ModelViewViewModelBlueprint);
+IMPLEMENT_MODULE(FDefaultModuleImpl, ModelViewViewModelBlueprint);
diff --git a/Engine/Plugins/Runtime/ModelViewViewModel/Source/ModelViewViewModelBlueprint/Private/ViewModel/MVVMViewModelBlueprint.cpp b/Engine/Plugins/Runtime/ModelViewViewModel/Source/ModelViewViewModelBlueprint/Private/ViewModel/MVVMViewModelBlueprint.cpp
--- a/Engine/Plugins/Runtime/ModelViewViewModel/Source/ModelViewViewModelBlueprint/Private/ViewModel/MVVMViewModelBlueprint.cpp
+++ b/Engine/Plugins/Runtime/ModelViewViewModel/Source/ModelViewViewModelBlueprint/Private/ViewModel/MVVMViewModelBlueprint.cpp
@@ -1,51 +1,6 @@
 // Copyright Epic Games, Inc. All Rights Reserved.
 
 #include "ViewModel/MVVMViewModelBlueprint.h"
-#include "ViewModel/MVVMViewModelBlueprintCompiler.h"
-#include "ViewModel/MVVMViewModelBlueprintGeneratedClass.h"
 
 
 #include UE_INLINE_GENERATED_CPP_BY_NAME(MVVMViewModelBlueprint)
-
-
-TSharedPtr<FKismetCompilerContext> UMVVMViewModelBlueprint::GetCompilerForViewModelBlueprint(UBlueprint* BP, FCompilerResultsLog& InMessageLog, const FKismetCompilerOptions& InCompileOptions)
-{
-	return TSharedPtr<FKismetCompilerContext>(new UE::MVVM::FViewModelBlueprintCompilerContext(CastChecked<UMVVMViewModelBlueprint>(BP), InMessageLog, InCompileOptions));
-}
-
-
-UMVVMViewModelBlueprintGeneratedClass* UMVVMViewModelBlueprint::GetViewModelBlueprintGeneratedClass() const
-{
-	return Cast<UMVVMViewModelBlueprintGeneratedClass>(*GeneratedClass);
-}
-
-
-UMVVMViewModelBlueprintGeneratedClass* UMVVMViewModelBlueprint::GetViewModelBlueprintSkeletonClass() const
-{
-	return Cast<UMVVMViewModelBlueprintGeneratedClass>(*SkeletonGeneratedClass);
-}
-
-
-bool UMVVMViewModelBlueprint::SupportsMacros() const
-{
-	return false;
-}
-
-
-bool UMVVMViewModelBlueprint::SupportsEventGraphs() const
-{
-	return false;
-}
-
-
-bool UMVVMViewModelBlueprint::SupportsDelegates() const
-{
-	return false;
-}
-
-
-bool UMVVMViewModelBlueprint::SupportsAnimLayers() const
-{
-	return false;
-}
-
diff --git a/Engine/Plugins/Runtime/ModelViewViewModel/Source/ModelViewViewModelBlueprint/Private/ViewModel/MVVMViewModelBlueprintCompiler.cpp b/Engine/Plugins/Runtime/ModelViewViewModel/Source/ModelViewViewModelBlueprint/Private/ViewModel/MVVMViewModelBlueprintCompiler.cpp
deleted file mode 100644
diff --git a/Engine/Plugins/Runtime/ModelViewViewModel/Source/ModelViewViewModelBlueprint/Private/ViewModel/MVVMViewModelBlueprintCompiler.h b/Engine/Plugins/Runtime/ModelViewViewModel/Source/ModelViewViewModelBlueprint/Private/ViewModel/MVVMViewModelBlueprintCompiler.h
deleted file mode 100644
diff --git a/Engine/Plugins/Runtime/ModelViewViewModel/Source/ModelViewViewModelBlueprint/Public/ViewModel/MVVMViewModelBlueprint.h b/Engine/Plugins/Runtime/ModelViewViewModel/Source/ModelViewViewModelBlueprint/Public/ViewModel/MVVMViewModelBlueprint.h
--- a/Engine/Plugins/Runtime/ModelViewViewModel/Source/ModelViewViewModelBlueprint/Public/ViewModel/MVVMViewModelBlueprint.h
+++ b/Engine/Plugins/Runtime/ModelViewViewModel/Source/ModelViewViewModelBlueprint/Public/ViewModel/MVVMViewModelBlueprint.h
@@ -16,31 +16,11 @@ class UMVVMViewModelBlueprintGeneratedClass;
 /**
  * 
  */
-UCLASS(BlueprintType, meta=(IgnoreClassThumbnail))
-class MODELVIEWVIEWMODELBLUEPRINT_API UMVVMViewModelBlueprint : public UBlueprint
+UCLASS(Deprecated)
+class UE_DEPRECATED(5.3, "The prototype viewmodel editor is deprecated. Use the regular Blueprint editor.")
+MODELVIEWVIEWMODELBLUEPRINT_API UDEPRECATED_MVVMViewModelBlueprint : public UBlueprint
 {
 	GENERATED_BODY()
-
-public:
-	/** */
-	static TSharedPtr<FKismetCompilerContext> GetCompilerForViewModelBlueprint(UBlueprint* BP, FCompilerResultsLog& InMessageLog, const FKismetCompilerOptions& InCompileOptions);
-
-	/** Get the (full) generated class for this viewmodel blueprint */
-	UMVVMViewModelBlueprintGeneratedClass* GetViewModelBlueprintGeneratedClass() const;
-
-	/** Get the (skeleton) generated class for this viewmodel blueprint */
-	UMVVMViewModelBlueprintGeneratedClass* GetViewModelBlueprintSkeletonClass() const;
-
-protected:
-	virtual bool SupportsMacros() const override;
-	virtual bool SupportsEventGraphs() const override;
-	virtual bool SupportsDelegates() const override;
-	virtual bool SupportsAnimLayers() const override;
-
-public:
-	// Use during compilation to clean the automatically generated graph.
-	UPROPERTY(Transient)
-	TArray<TObjectPtr<UEdGraph>> TemporaryGraph;
 };
 
 #if UE_ENABLE_INCLUDE_ORDER_DEPRECATED_IN_5_2
diff --git a/Engine/Plugins/Runtime/ModelViewViewModel/Source/ModelViewViewModelEditor/Private/ModelViewViewModelEditorModule.cpp b/Engine/Plugins/Runtime/ModelViewViewModel/Source/ModelViewViewModelEditor/Private/ModelViewViewModelEditorModule.cpp
--- a/Engine/Plugins/Runtime/ModelViewViewModel/Source/ModelViewViewModelEditor/Private/ModelViewViewModelEditorModule.cpp
+++ b/Engine/Plugins/Runtime/ModelViewViewModel/Source/ModelViewViewModelEditor/Private/ModelViewViewModelEditorModule.cpp
@@ -20,7 +20,6 @@
 #include "Tabs/MVVMBindingSummoner.h"
 #include "Tabs/MVVMViewModelSummoner.h"
 #include "UMGEditorModule.h"
-#include "ViewModel/AssetTypeActions_ViewModelBlueprint.h"
 #include "WidgetBlueprintEditor.h"
 #include "WidgetDrawerConfig.h"
 
@@ -36,16 +35,6 @@ void FModelViewViewModelEditorModule::StartupModule()
 	IUMGEditorModule& UMGEditorModule = FModuleManager::LoadModuleChecked<IUMGEditorModule>("UMGEditor");
 	UMGEditorModule.OnRegisterTabsForEditor().AddRaw(this, &FModelViewViewModelEditorModule::HandleRegisterBlueprintEditorTab);
 
-	// Register asset types
-	{
-#if UE_MVVM_WITH_VIEWMODEL_EDITOR
-		// Only remove what is related to the viewmodel editor, not the UMG extention for view.
-		IAssetTools& AssetTools = FModuleManager::LoadModuleChecked<FAssetToolsModule>("AssetTools").Get();
-		ViewModelBlueprintActions = MakeShared<UE::MVVM::FAssetTypeActions_ViewModelBlueprint>();
-		AssetTools.RegisterAssetTypeActions(ViewModelBlueprintActions.ToSharedRef());
-#endif
-	}
-
 	PropertyBindingExtension = MakeShared<FMVVMPropertyBindingExtension>();
 	UMGEditorModule.GetPropertyBindingExtensibilityManager()->AddExtension(PropertyBindingExtension.ToSharedRef());
 
@@ -91,17 +80,6 @@ void FModelViewViewModelEditorModule::ShutdownModule()
 	}
 	PropertyBindingExtension.Reset();
 
-#if UE_MVVM_WITH_VIEWMODEL_EDITOR
-	// Unregister all the asset types that we registered
-	if (FAssetToolsModule* AssetTools = FModuleManager::GetModulePtr<FAssetToolsModule>("AssetTools"))
-	{
-		if (ViewModelBlueprintActions)
-		{
-			AssetTools->Get().UnregisterAssetTypeActions(ViewModelBlueprintActions.ToSharedRef());
-		}
-	}
-#endif
-
 	FMVVMEditorStyle::DestroyInstance();
 
 	FMVVMEditorCommands::Unregister();
diff --git a/Engine/Plugins/Runtime/ModelViewViewModel/Source/ModelViewViewModelEditor/Private/ViewModel/AssetTypeActions_ViewModelBlueprint.cpp b/Engine/Plugins/Runtime/ModelViewViewModel/Source/ModelViewViewModelEditor/Private/ViewModel/AssetTypeActions_ViewModelBlueprint.cpp
deleted file mode 100644
diff --git a/Engine/Plugins/Runtime/ModelViewViewModel/Source/ModelViewViewModelEditor/Private/ViewModel/AssetTypeActions_ViewModelBlueprint.h b/Engine/Plugins/Runtime/ModelViewViewModel/Source/ModelViewViewModelEditor/Private/ViewModel/AssetTypeActions_ViewModelBlueprint.h
deleted file mode 100644
diff --git a/Engine/Plugins/Runtime/ModelViewViewModel/Source/ModelViewViewModelEditor/Private/ViewModel/MVVMViewModelBlueprintEditor.cpp b/Engine/Plugins/Runtime/ModelViewViewModel/Source/ModelViewViewModelEditor/Private/ViewModel/MVVMViewModelBlueprintEditor.cpp
deleted file mode 100644
diff --git a/Engine/Plugins/Runtime/ModelViewViewModel/Source/ModelViewViewModelEditor/Private/ViewModel/MVVMViewModelBlueprintFactory.cpp b/Engine/Plugins/Runtime/ModelViewViewModel/Source/ModelViewViewModelEditor/Private/ViewModel/MVVMViewModelBlueprintFactory.cpp
deleted file mode 100644
diff --git a/Engine/Plugins/Runtime/ModelViewViewModel/Source/ModelViewViewModelEditor/Private/ViewModel/MVVMViewModelBlueprintFactory.h b/Engine/Plugins/Runtime/ModelViewViewModel/Source/ModelViewViewModelEditor/Private/ViewModel/MVVMViewModelBlueprintFactory.h
deleted file mode 100644
diff --git a/Engine/Plugins/Runtime/ModelViewViewModel/Source/ModelViewViewModelEditor/Public/ViewModel/MVVMViewModelBlueprintEditor.h b/Engine/Plugins/Runtime/ModelViewViewModel/Source/ModelViewViewModelEditor/Public/ViewModel/MVVMViewModelBlueprintEditor.h
deleted file mode 100644
diff --git a/Engine/Plugins/Runtime/ModelViewViewModel/Source/ModelViewViewModelEditor/Public/ViewModel/MVVMViewModelBlueprintToolMenuContext.h b/Engine/Plugins/Runtime/ModelViewViewModel/Source/ModelViewViewModelEditor/Public/ViewModel/MVVMViewModelBlueprintToolMenuContext.h
deleted file mode 100644
