MVVM: Unload properties from the view when no widget are using them. That reload the correct properties from a recompiled BP.

diff --git a/Engine/Plugins/Runtime/ModelViewViewModel/Source/ModelViewViewModel/ModelViewViewModel.Build.cs b/Engine/Plugins/Runtime/ModelViewViewModel/Source/ModelViewViewModel/ModelViewViewModel.Build.cs
--- a/Engine/Plugins/Runtime/ModelViewViewModel/Source/ModelViewViewModel/ModelViewViewModel.Build.cs
+++ b/Engine/Plugins/Runtime/ModelViewViewModel/Source/ModelViewViewModel/ModelViewViewModel.Build.cs
@@ -25,6 +25,15 @@ public class ModelViewViewModel : ModuleRules
 				"UMG",
 			});
 
+		if (Target.bBuildEditor)
+		{
+			PrivateDependencyModuleNames.AddRange(
+			new string[]
+			{
+				"UnrealEd",
+			});
+		}
+
 		if (Target.Configuration == UnrealTargetConfiguration.Test || Target.Configuration == UnrealTargetConfiguration.Shipping)
 		{
 			PublicDefinitions.Add("UE_WITH_MVVM_DEBUGGING=0");
diff --git a/Engine/Plugins/Runtime/ModelViewViewModel/Source/ModelViewViewModel/Private/Bindings/MVVMCompiledBindingLibrary.cpp b/Engine/Plugins/Runtime/ModelViewViewModel/Source/ModelViewViewModel/Private/Bindings/MVVMCompiledBindingLibrary.cpp
--- a/Engine/Plugins/Runtime/ModelViewViewModel/Source/ModelViewViewModel/Private/Bindings/MVVMCompiledBindingLibrary.cpp
+++ b/Engine/Plugins/Runtime/ModelViewViewModel/Source/ModelViewViewModel/Private/Bindings/MVVMCompiledBindingLibrary.cpp
@@ -113,22 +113,15 @@ void FMVVMCompiledBindingLibrary::Load()
 			}
 		}
 	}
-
-#if !WITH_EDITOR
-	CompiledFieldNames.Reset(); // Once loaded, we do not need to keep that information.
-	CompiledFields.Reset();
-#endif
 }
 
 
-#if WITH_EDITOR
 void FMVVMCompiledBindingLibrary::Unload()
 {
 	LoadedProperties.Empty();
 	LoadedFunctions.Empty();
 	LoadedFieldIds.Empty();
 }
-#endif
 
 
 TValueOrError<void, FMVVMCompiledBindingLibrary::EExecutionFailingReason> FMVVMCompiledBindingLibrary::Execute(UObject* InExecutionSource, const FMVVMVCompiledBinding& InBinding, EConversionFunctionType InFunctionType) const
diff --git a/Engine/Plugins/Runtime/ModelViewViewModel/Source/ModelViewViewModel/Private/View/MVVMViewClass.cpp b/Engine/Plugins/Runtime/ModelViewViewModel/Source/ModelViewViewModel/Private/View/MVVMViewClass.cpp
--- a/Engine/Plugins/Runtime/ModelViewViewModel/Source/ModelViewViewModel/Private/View/MVVMViewClass.cpp
+++ b/Engine/Plugins/Runtime/ModelViewViewModel/Source/ModelViewViewModel/Private/View/MVVMViewClass.cpp
@@ -14,6 +14,13 @@
 #include "MVVMMessageLog.h"
 #include "MVVMSubsystem.h"
 
+#include <limits>
+
+#if WITH_EDITOR
+#include "Editor.h"
+#include "Editor/EditorEngine.h"
+#endif
+
 #include UE_INLINE_GENERATED_CPP_BY_NAME(MVVMViewClass)
 
 
@@ -481,28 +488,77 @@ FString FMVVMViewClass_CompiledBinding::ToString(const FMVVMCompiledBindingLibra
 
 void UMVVMViewClass::Initialize(UUserWidget* UserWidget)
 {
+	++ViewCounter;
+	ensure(ViewCounter <= std::numeric_limits<int32>::max());
+
 	ensure(UserWidget->GetExtension<UMVVMView>() == nullptr);
 	UMVVMView* View = UserWidget->AddExtension<UMVVMView>();
 	if (ensure(View))
 	{
-		if (!bLoaded)
+		if (ViewCounter == 1)
 		{
 			BindingLibrary.Load();
-			bLoaded = true;
 		}
 
 		View->ConstructView(this);
 	}
+
+#if WITH_EDITOR
+	if (GEditor)
+	{
+		BluerpintCompiledHandle = GEditor->OnBlueprintCompiled().AddUObject(this, &UMVVMViewClass::HandleBlueprintCompiled);
+	}
+#endif
+}
+
+
+void UMVVMViewClass::Destruct(UUserWidget* UserWidget)
+{
+	--ViewCounter;
+	ensure(ViewCounter >= 0);
+	if (ViewCounter == 0)
+	{
+		BindingLibrary.Unload();
+	}
+
+#if WITH_EDITOR
+	if (GEditor && BluerpintCompiledHandle.IsValid())
+	{
+		GEditor->OnBlueprintCompiled().Remove(BluerpintCompiledHandle);
+		BluerpintCompiledHandle.Reset();
+	}
+#endif
 }
 
+#if WITH_EDITOR
+void UMVVMViewClass::BeginDestroy()
+{
+	if (GEditor && BluerpintCompiledHandle.IsValid())
+	{
+		GEditor->OnBlueprintCompiled().Remove(BluerpintCompiledHandle);
+		BluerpintCompiledHandle.Reset();
+	}
+	Super::BeginDestroy();
+}
+
+void UMVVMViewClass::HandleBlueprintCompiled()
+{
+	// do not keep property alive of potential dead BP class.
+	if (ViewCounter > 0)
+	{
+		BindingLibrary.Unload();
+		BindingLibrary.Load();
+	}
+}
+#endif
+
 #if UE_WITH_MVVM_DEBUGGING
 void UMVVMViewClass::Log(FMVVMViewClass_SourceCreator::FToStringArgs SourceArgs, FMVVMViewClass_CompiledBinding::FToStringArgs BindingArgs) const
 {
-	bool bPreviousLoad = bLoaded;
-	if (!bLoaded)
+	bool bPreviousLoad = ViewCounter >= 0;
+	if (!bPreviousLoad)
 	{
 		const_cast<UMVVMViewClass*>(this)->BindingLibrary.Load();
-		const_cast<UMVVMViewClass*>(this)->bLoaded = true;
 	}
 
 	TStringBuilder<2048> Builder;
@@ -530,7 +586,6 @@ void UMVVMViewClass::Log(FMVVMViewClass_SourceCreator::FToStringArgs SourceArgs,
 	{
 #if WITH_EDITOR
 		const_cast<UMVVMViewClass*>(this)->BindingLibrary.Unload();
-		const_cast<UMVVMViewClass*>(this)->bLoaded = false;
 #endif
 	}
 }
diff --git a/Engine/Plugins/Runtime/ModelViewViewModel/Source/ModelViewViewModel/Public/Bindings/MVVMCompiledBindingLibrary.h b/Engine/Plugins/Runtime/ModelViewViewModel/Source/ModelViewViewModel/Public/Bindings/MVVMCompiledBindingLibrary.h
--- a/Engine/Plugins/Runtime/ModelViewViewModel/Source/ModelViewViewModel/Public/Bindings/MVVMCompiledBindingLibrary.h
+++ b/Engine/Plugins/Runtime/ModelViewViewModel/Source/ModelViewViewModel/Public/Bindings/MVVMCompiledBindingLibrary.h
@@ -293,11 +293,8 @@ public:
 
 	/** Fetch the FProperty and UFunction. */
 	void Load();
-
-#if WITH_EDITOR
 	/** Release the acquired FProperty and UFunction. */
 	void Unload();
-#endif
 
 	/**
 	 * Execute a binding, in one direction.
diff --git a/Engine/Plugins/Runtime/ModelViewViewModel/Source/ModelViewViewModel/Public/View/MVVMViewClass.h b/Engine/Plugins/Runtime/ModelViewViewModel/Source/ModelViewViewModel/Public/View/MVVMViewClass.h
--- a/Engine/Plugins/Runtime/ModelViewViewModel/Source/ModelViewViewModel/Public/View/MVVMViewClass.h
+++ b/Engine/Plugins/Runtime/ModelViewViewModel/Source/ModelViewViewModel/Public/View/MVVMViewClass.h
@@ -322,6 +322,10 @@ class MODELVIEWVIEWMODEL_API UMVVMViewClass : public UWidgetBlueprintGeneratedCl
 public:
 	//~ Begin UWidgetBlueprintGeneratedClassExtension
 	virtual void Initialize(UUserWidget* UserWidget) override;
+	virtual void Destruct(UUserWidget* UserWidget) override;
+#if WITH_EDITOR
+	virtual void BeginDestroy() override;
+#endif
 	//~ End UWidgetBlueprintGeneratedClassExtension
 
 public:
@@ -366,6 +370,10 @@ private:
 		return MakeArrayView(CompiledBindings);
 	}
 
+#if WITH_EDITOR
+	void HandleBlueprintCompiled();
+#endif
+
 private:
 	/** Data to retrieve/create the sources (could be viewmodel, widget, ...). */
 	UPROPERTY()
@@ -379,12 +387,16 @@ private:
 	UPROPERTY()
 	FMVVMCompiledBindingLibrary BindingLibrary;
 
+	/** */
+	int32 ViewCounter = 0;
+
 	/** */
 	UPROPERTY()
 	bool bInitializeAtConstruction = true;
 
-	/** */
-	bool bLoaded = false;
+#if WITH_EDITORONLY_DATA
+	FDelegateHandle BluerpintCompiledHandle;
+#endif
 };
 
 #if UE_ENABLE_INCLUDE_ORDER_DEPRECATED_IN_5_2
diff --git a/Engine/Plugins/Runtime/ModelViewViewModel/Source/ModelViewViewModelBlueprint/Private/MVVMBlueprintView.cpp b/Engine/Plugins/Runtime/ModelViewViewModel/Source/ModelViewViewModelBlueprint/Private/MVVMBlueprintView.cpp
--- a/Engine/Plugins/Runtime/ModelViewViewModel/Source/ModelViewViewModelBlueprint/Private/MVVMBlueprintView.cpp
+++ b/Engine/Plugins/Runtime/ModelViewViewModel/Source/ModelViewViewModelBlueprint/Private/MVVMBlueprintView.cpp
@@ -91,14 +91,6 @@ bool UMVVMBlueprintView::RenameViewModel(FName OldViewModelName, FName NewViewMo
 	return ViewModelContext != nullptr;
 }
 
-void UMVVMBlueprintView::SetViewModels(const TArray<FMVVMBlueprintViewModelContext>& ViewModelContexts)
-{
-	AvailableViewModels = ViewModelContexts;
-
-	FBlueprintEditorUtils::MarkBlueprintAsStructurallyModified(GetOuterUMVVMWidgetBlueprintExtension_View()->GetWidgetBlueprint());
-	OnViewModelsUpdated.Broadcast();
-}
-
 const FMVVMBlueprintViewBinding* UMVVMBlueprintView::FindBinding(const UWidget* Widget, const FProperty* Property) const
 {
 	return const_cast<UMVVMBlueprintView*>(this)->FindBinding(Widget, Property);
diff --git a/Engine/Plugins/Runtime/ModelViewViewModel/Source/ModelViewViewModelBlueprint/Public/MVVMBlueprintView.h b/Engine/Plugins/Runtime/ModelViewViewModel/Source/ModelViewViewModelBlueprint/Public/MVVMBlueprintView.h
--- a/Engine/Plugins/Runtime/ModelViewViewModel/Source/ModelViewViewModelBlueprint/Public/MVVMBlueprintView.h
+++ b/Engine/Plugins/Runtime/ModelViewViewModel/Source/ModelViewViewModelBlueprint/Public/MVVMBlueprintView.h
@@ -45,7 +45,6 @@ public:
 	bool RemoveViewModel(FGuid ViewModelId);
 	int32 RemoveViewModels(const TArrayView<FGuid> ViewModelIds);
 	bool RenameViewModel(FName OldViewModelName, FName NewViewModelName);
-	void SetViewModels(const TArray<FMVVMBlueprintViewModelContext>& ViewModelContexts);
 
 	const TArrayView<const FMVVMBlueprintViewModelContext> GetViewModels() const
 	{
