MVVM: Add a debugger for viewmodel. Work in progress. You can see the list of view and viewmodel currently used. You can see the list of binding on a viewmodel.

diff --git a/Engine/Plugins/Runtime/ModelViewViewModel/ModelViewViewModel.uplugin b/Engine/Plugins/Runtime/ModelViewViewModel/ModelViewViewModel.uplugin
--- a/Engine/Plugins/Runtime/ModelViewViewModel/ModelViewViewModel.uplugin
+++ b/Engine/Plugins/Runtime/ModelViewViewModel/ModelViewViewModel.uplugin
@@ -30,6 +30,16 @@
 			"Name": "ModelViewViewModelEditor",
 			"Type": "Editor",
 			"LoadingPhase": "Default"
+		},
+		{
+			"Name": "ModelViewViewModelDebugger",
+			"Type": "UncookedOnly",
+			"LoadingPhase": "Default"
+		},
+		{
+			"Name": "ModelViewViewModelDebuggerEditor",
+			"Type": "Editor",
+			"LoadingPhase": "Default"
 		}
 	]
 }
diff --git a/Engine/Plugins/Runtime/ModelViewViewModel/Source/ModelViewViewModel/Private/MVVMViewModelBase.cpp b/Engine/Plugins/Runtime/ModelViewViewModel/Source/ModelViewViewModel/Private/MVVMViewModelBase.cpp
--- a/Engine/Plugins/Runtime/ModelViewViewModel/Source/ModelViewViewModel/Private/MVVMViewModelBase.cpp
+++ b/Engine/Plugins/Runtime/ModelViewViewModel/Source/ModelViewViewModel/Private/MVVMViewModelBase.cpp
@@ -54,6 +54,12 @@ int32 UMVVMViewModelBase::RemoveAllFieldValueChangedDelegates(UE::FieldNotificat
 }
 
 
+TArray<UE::FieldNotification::FFieldMulticastDelegate::FDelegateView> UMVVMViewModelBase::GetNotificationDelegateView() const
+{
+	return NotificationDelegates.GetView();
+}
+
+
 const UE::FieldNotification::IClassDescriptor& UMVVMViewModelBase::GetFieldNotificationDescriptor() const
 {
 	static FFieldNotificationClassDescriptor Local;
diff --git a/Engine/Plugins/Runtime/ModelViewViewModel/Source/ModelViewViewModel/Private/ViewModel/MVVMFieldNotificationDelegates.cpp b/Engine/Plugins/Runtime/ModelViewViewModel/Source/ModelViewViewModel/Private/ViewModel/MVVMFieldNotificationDelegates.cpp
--- a/Engine/Plugins/Runtime/ModelViewViewModel/Source/ModelViewViewModel/Private/ViewModel/MVVMFieldNotificationDelegates.cpp
+++ b/Engine/Plugins/Runtime/ModelViewViewModel/Source/ModelViewViewModel/Private/ViewModel/MVVMFieldNotificationDelegates.cpp
@@ -97,6 +97,11 @@ void FMVVMFieldNotificationDelegates::BroadcastFieldValueChanged(UObject* InOwne
 	}
 }
 
+TArray<UE::FieldNotification::FFieldMulticastDelegate::FDelegateView> FMVVMFieldNotificationDelegates::GetView() const
+{
+	return Delegates.GetView();
+}
+
 } //namespace
 
 #undef LOCTEXT_NAMESPACE
diff --git a/Engine/Plugins/Runtime/ModelViewViewModel/Source/ModelViewViewModel/Public/MVVMViewModelBase.h b/Engine/Plugins/Runtime/ModelViewViewModel/Source/ModelViewViewModel/Public/MVVMViewModelBase.h
--- a/Engine/Plugins/Runtime/ModelViewViewModel/Source/ModelViewViewModel/Public/MVVMViewModelBase.h
+++ b/Engine/Plugins/Runtime/ModelViewViewModel/Source/ModelViewViewModel/Public/MVVMViewModelBase.h
@@ -52,6 +52,8 @@ public:
 	virtual void BroadcastFieldValueChanged(UE::FieldNotification::FFieldId InFieldId) override;
 	//~ End INotifyFieldValueChanged Interface
 
+	TArray<UE::FieldNotification::FFieldMulticastDelegate::FDelegateView> GetNotificationDelegateView() const;
+
 	UFUNCTION(BlueprintCallable, Category = "FieldNotify", meta = (DisplayName = "Add Field Value Changed Delegate", ScriptName = "AddFieldValueChangedDelegate"))
 	void K2_AddFieldValueChangedDelegate(FFieldNotificationId FieldId, FFieldValueChangedDynamicDelegate Delegate);
 
diff --git a/Engine/Plugins/Runtime/ModelViewViewModel/Source/ModelViewViewModel/Public/View/MVVMView.h b/Engine/Plugins/Runtime/ModelViewViewModel/Source/ModelViewViewModel/Public/View/MVVMView.h
--- a/Engine/Plugins/Runtime/ModelViewViewModel/Source/ModelViewViewModel/Public/View/MVVMView.h
+++ b/Engine/Plugins/Runtime/ModelViewViewModel/Source/ModelViewViewModel/Public/View/MVVMView.h
@@ -76,6 +76,11 @@ public:
 		return ClassExtension;
 	}
 
+	const TArrayView<const FMVVMViewSource> GetSources() const
+	{
+		return Sources;
+	}
+
 	void ExecuteDelayedBinding(const FMVVMViewDelayedBinding& DelayedBinding) const;
 	void ExecuteEveryTickBindings() const;
 
diff --git a/Engine/Plugins/Runtime/ModelViewViewModel/Source/ModelViewViewModel/Public/ViewModel/MVVMFieldNotificationDelegates.h b/Engine/Plugins/Runtime/ModelViewViewModel/Source/ModelViewViewModel/Public/ViewModel/MVVMFieldNotificationDelegates.h
--- a/Engine/Plugins/Runtime/ModelViewViewModel/Source/ModelViewViewModel/Public/ViewModel/MVVMFieldNotificationDelegates.h
+++ b/Engine/Plugins/Runtime/ModelViewViewModel/Source/ModelViewViewModel/Public/ViewModel/MVVMFieldNotificationDelegates.h
@@ -21,6 +21,7 @@ public:
 	int32 RemoveAllFieldValueChangedDelegates(UObject* Owner, const void* InUserObject);
 	int32 RemoveAllFieldValueChangedDelegates(UObject* Owner, UE::FieldNotification::FFieldId InFieldId, const void* InUserObject);
 	void BroadcastFieldValueChanged(UObject* Owner, UE::FieldNotification::FFieldId InFieldId);
+	TArray<UE::FieldNotification::FFieldMulticastDelegate::FDelegateView> GetView() const;
 
 private:
 	UE::FieldNotification::FFieldMulticastDelegate Delegates;
diff --git a/Engine/Plugins/Runtime/ModelViewViewModel/Source/ModelViewViewModelDebugger/ModelViewViewModelDebugger.Build.cs b/Engine/Plugins/Runtime/ModelViewViewModel/Source/ModelViewViewModelDebugger/ModelViewViewModelDebugger.Build.cs
new file mode 100644
--- /dev/null
+++ b/Engine/Plugins/Runtime/ModelViewViewModel/Source/ModelViewViewModelDebugger/ModelViewViewModelDebugger.Build.cs
@@ -0,0 +1,28 @@
+// Copyright Epic Games, Inc. All Rights Reserved.
+
+using System.IO;
+using System.Collections.Generic;
+using UnrealBuildTool;
+
+public class ModelViewViewModelDebugger : ModuleRules 
+{
+	public ModelViewViewModelDebugger(ReadOnlyTargetRules Target) : base(Target)
+	{
+		PublicDependencyModuleNames.AddRange(
+			new string[]
+			{
+				"Core",
+				"CoreUObject",
+				"FieldNotification",
+				"ModelViewViewModel",
+			});
+
+		PrivateDependencyModuleNames.AddRange(
+			new string[]
+			{
+				"SlateCore",
+				"Slate",
+				"UMG",
+			});
+	}
+}
diff --git a/Engine/Plugins/Runtime/ModelViewViewModel/Source/ModelViewViewModelDebugger/Private/MVVMDebugSnapshot.cpp b/Engine/Plugins/Runtime/ModelViewViewModel/Source/ModelViewViewModelDebugger/Private/MVVMDebugSnapshot.cpp
new file mode 100644
--- /dev/null
+++ b/Engine/Plugins/Runtime/ModelViewViewModel/Source/ModelViewViewModelDebugger/Private/MVVMDebugSnapshot.cpp
@@ -0,0 +1,183 @@
+// Copyright Epic Games, Inc. All Rights Reserved.
+
+#include "MVVMDebugSnapshot.h"
+#include "MVVMDebugView.h"
+#include "MVVMDebugViewClass.h"
+
+#include "Bindings/MVVMCompiledBindingLibrary.h"
+#include "Blueprint/UserWidget.h"
+#include "Engine/LocalPlayer.h"
+#include "Engine/World.h"
+#include "Templates/ValueOrError.h"
+#include "MVVMDebugView.h"
+#include "MVVMDebugViewModel.h"
+#include "MVVMViewModelBase.h"
+#include "UObject/UObjectIterator.h"
+#include "View/MVVMView.h"
+#include "View/MVVMViewClass.h"
+
+namespace UE::MVVM
+{
+namespace Private
+{
+FMVVMViewSourceDebugEntry CreateSourceInstanceEntry(const FMVVMViewSource& ViewSource)
+{
+	FMVVMViewSourceDebugEntry Result;
+	Result.SourceInstanceName = ViewSource.SourceName;
+	Result.SourceAsset = ViewSource.Source ? FAssetData(ViewSource.Source->GetClass()) : FAssetData();
+	Result.ViewModelDebugId = FGuid::NewGuid();
+	Result.LiveSource = ViewSource.Source;
+	return Result;
+}
+
+FMVVMViewBindingDebugEntry CreateViewBindingDebugEntry(const UMVVMViewClass* ViewClass, const FMVVMViewClass_CompiledBinding& Binding, int32 Index)
+{
+	FMVVMViewBindingDebugEntry Result;
+#if WITH_EDITOR
+	Result.BlueprintViewBindingId = Binding.GetEditorId();
+#endif
+	if (Binding.GetSourceFieldId().IsValid())
+	{
+		TValueOrError<UE::FieldNotification::FFieldId, void> FieldIdValue = ViewClass->GetBindingLibrary().GetFieldId(Binding.GetSourceFieldId());
+		Result.FieldId = FieldIdValue.HasValue() ? FieldIdValue.GetValue().GetName() : FName();
+	}
+	if (Binding.GetBinding().GetSourceFieldPath().IsValid())
+	{
+		TValueOrError<FString, FString> SourceFieldPathValue = ViewClass->GetBindingLibrary().FieldPathToString(Binding.GetBinding().GetSourceFieldPath(), false);
+		Result.SourceFieldPath = SourceFieldPathValue.HasValue() ? SourceFieldPathValue.StealValue() : SourceFieldPathValue.StealError();
+	}
+	if (Binding.GetBinding().GetDestinationFieldPath().IsValid())
+	{
+		TValueOrError<FString, FString> DestinationFieldPathValue = ViewClass->GetBindingLibrary().FieldPathToString(Binding.GetBinding().GetDestinationFieldPath(), false);
+		Result.DestinationFieldPath = DestinationFieldPathValue.HasValue() ? DestinationFieldPathValue.StealValue() : DestinationFieldPathValue.StealError();
+	}
+	if (Binding.GetBinding().GetConversionFunctionFieldPath().IsValid())
+	{
+		TValueOrError<FString, FString> ConversionFieldPathValue = ViewClass->GetBindingLibrary().FieldPathToString(Binding.GetBinding().GetConversionFunctionFieldPath(), false);
+		Result.ConversionFunctionFieldPath = ConversionFieldPathValue.HasValue() ? ConversionFieldPathValue.StealValue() : ConversionFieldPathValue.StealError();
+	}
+	Result.CompiledBindingIndex = Index;
+	return Result;
+}
+
+FMVVMViewSourceDebugEntry CreateViewModelEntry(const FMVVMViewSource& ViewSource)
+{
+	FMVVMViewSourceDebugEntry Result;
+	Result.SourceInstanceName = ViewSource.SourceName;
+	Result.SourceAsset = ViewSource.Source ? FAssetData(ViewSource.Source->GetClass()) : FAssetData();
+	Result.ViewModelDebugId = FGuid::NewGuid();
+	Result.LiveSource = ViewSource.Source;
+	return Result;
+}
+
+FMVVMViewModelFieldBoundDebugEntry CreateViewModelFieldBoundDebugEntry(const UE::FieldNotification::FFieldMulticastDelegate::FDelegateView& DelegateView)
+{
+	FMVVMViewModelFieldBoundDebugEntry Result;
+	Result.ObjectName = DelegateView.Object ? DelegateView.Object->GetPathName() : FString();
+	Result.FieldId = FFieldNotificationId(DelegateView.Field.GetName());
+	Result.FunctionName = DelegateView.FunctionName;
+	Result.LiveInstanceObject = DelegateView.Object;
+	return Result;
+}
+}//namespace private
+
+
+TSharedPtr<FMVVMViewDebugEntry> FDebugSnapshot::FindView(FGuid Id) const
+{
+	const TSharedPtr<FMVVMViewDebugEntry>* Result = Views.FindByPredicate([Id](const TSharedPtr<FMVVMViewDebugEntry>& Other){ return Other->ViewInstanceDebugId == Id; });
+	return Result ? *Result : TSharedPtr<FMVVMViewDebugEntry>();
+}
+
+
+TSharedPtr<FMVVMViewModelDebugEntry> FDebugSnapshot::FindViewModel(FGuid Id) const
+{
+	const TSharedPtr<FMVVMViewModelDebugEntry>* Result = ViewModels.FindByPredicate([Id](const TSharedPtr<FMVVMViewModelDebugEntry>& Other) { return Other->ViewModelDebugId == Id; });
+	return Result ? *Result : TSharedPtr<FMVVMViewModelDebugEntry>();
+}
+
+
+TSharedPtr<FDebugSnapshot> FDebugSnapshot::CreateSnapshot()
+{
+	TSharedPtr<FDebugSnapshot> Snapshot = MakeShared<FDebugSnapshot>();
+
+	for (FThreadSafeObjectIterator It(UMVVMView::StaticClass()); It; ++It)
+	{
+		if (It->IsTemplate(RF_ClassDefaultObject))
+		{
+			continue;
+		}
+
+		UMVVMView* View = CastChecked<UMVVMView>(*It);
+		UUserWidget* UserWidget = View->GetOuterUUserWidget();
+
+		TSharedPtr<FMVVMViewDebugEntry> DebugEntry = MakeShared<FMVVMViewDebugEntry>();
+		if (UserWidget)
+		{
+			DebugEntry->UserWidgetInstanceName = UserWidget->GetFName();
+			ULocalPlayer* LocalPlayer = UserWidget->GetOwningLocalPlayer();
+			DebugEntry->LocalPlayerName = LocalPlayer ? LocalPlayer->GetFName() : FName();
+			UWorld* World = UserWidget->GetWorld();
+			DebugEntry->WorldName = World ? World->GetFName() : FName();
+			DebugEntry->UserWidgetAsset = FAssetData(UserWidget->GetClass());
+		}
+		for (const FMVVMViewSource& ViewSource : View->GetSources())
+		{
+			DebugEntry->Sources.Add(UE::MVVM::Private::CreateSourceInstanceEntry(ViewSource));
+		}
+		DebugEntry->ViewClassDebugId = Snapshot->FindOrAddViewClassEntry(View->GetViewClass())->ViewClassDebugId;
+		DebugEntry->ViewInstanceDebugId = FGuid::NewGuid();
+		DebugEntry->LiveView = View;
+
+		Snapshot->Views.Add(DebugEntry);
+	}
+
+	for (FThreadSafeObjectIterator It(UMVVMViewModelBase::StaticClass()); It; ++It)
+	{
+		if (It->IsTemplate(RF_ClassDefaultObject))
+		{
+			continue;
+		}
+
+		UMVVMViewModelBase* ViewModel = CastChecked<UMVVMViewModelBase>(*It);
+
+		TSharedPtr<FMVVMViewModelDebugEntry> DebugEntry = MakeShared<FMVVMViewModelDebugEntry>();
+		DebugEntry->Name = ViewModel->GetFName();
+		DebugEntry->FullName = ViewModel->GetPathName();
+		DebugEntry->ViewModelAsset = FAssetData(ViewModel->GetClass());
+		for (const UE::FieldNotification::FFieldMulticastDelegate::FDelegateView& DelegateView : ViewModel->GetNotificationDelegateView())
+		{
+			DebugEntry->FieldBound.Add(UE::MVVM::Private::CreateViewModelFieldBoundDebugEntry(DelegateView));
+		}
+		//DebugEntry.PropertyBag;
+		DebugEntry->ViewModelDebugId = FGuid::NewGuid();
+		DebugEntry->LiveViewModel = ViewModel;
+
+		Snapshot->ViewModels.Add(DebugEntry);
+	}
+
+	return Snapshot;
+}
+
+
+TSharedRef<FMVVMViewClassDebugEntry> FDebugSnapshot::FindOrAddViewClassEntry(const UMVVMViewClass* ViewClass)
+{
+	TSharedPtr<FMVVMViewClassDebugEntry>* Entry = ViewClasses.FindByPredicate([ViewClass](const TSharedPtr<FMVVMViewClassDebugEntry>& Entry) { return Entry->LiveViewClass == ViewClass; });
+	if (Entry)
+	{
+		return Entry->ToSharedRef();
+	}
+
+	TSharedRef<FMVVMViewClassDebugEntry> NewEntry = MakeShared<FMVVMViewClassDebugEntry>();
+	ViewClasses.Add(NewEntry);
+	NewEntry->ViewClassDebugId = FGuid::NewGuid();
+	NewEntry->LiveViewClass = ViewClass;
+
+	int32 Index = 0;
+	for (const FMVVMViewClass_CompiledBinding& Binding : ViewClass->GetCompiledBindings())
+	{
+		NewEntry->Bindings.Add(Private::CreateViewBindingDebugEntry(ViewClass, Binding, Index));
+		++Index;
+	}
+	return NewEntry;
+}
+}//namespace UE::MVVM
\ No newline at end of file
diff --git a/Engine/Plugins/Runtime/ModelViewViewModel/Source/ModelViewViewModelDebugger/Private/ModelViewViewModelDebuggerModule.cpp b/Engine/Plugins/Runtime/ModelViewViewModel/Source/ModelViewViewModelDebugger/Private/ModelViewViewModelDebuggerModule.cpp
new file mode 100644
--- /dev/null
+++ b/Engine/Plugins/Runtime/ModelViewViewModel/Source/ModelViewViewModelDebugger/Private/ModelViewViewModelDebuggerModule.cpp
@@ -0,0 +1,5 @@
+// Copyright Epic Games, Inc. All Rights Reserved.
+
+#include "Modules/ModuleManager.h"
+
+IMPLEMENT_MODULE(FDefaultModuleImpl, ModelViewViewModelDebugger);
diff --git a/Engine/Plugins/Runtime/ModelViewViewModel/Source/ModelViewViewModelDebugger/Public/MVVMDebugItemId.h b/Engine/Plugins/Runtime/ModelViewViewModel/Source/ModelViewViewModelDebugger/Public/MVVMDebugItemId.h
new file mode 100644
--- /dev/null
+++ b/Engine/Plugins/Runtime/ModelViewViewModel/Source/ModelViewViewModelDebugger/Public/MVVMDebugItemId.h
@@ -0,0 +1,32 @@
+// Copyright Epic Games, Inc. All Rights Reserved.
+
+#pragma once
+
+#include "CoreMinimal.h"
+
+namespace UE::MVVM
+{
+
+struct FDebugItemId
+{
+	enum class EType : uint8
+	{
+		Invalid,
+		View,
+		ViewModel,
+	} Type = EType::Invalid;
+	FGuid Id;
+
+
+	FDebugItemId() = default;
+	FDebugItemId(EType InType, FGuid InId)
+		: Type(InType), Id(InId)
+	{ }
+
+	bool operator==(const FDebugItemId& Other) const
+	{
+		return Other.Type == Type && Other.Id == Id;
+	}
+};
+
+}
diff --git a/Engine/Plugins/Runtime/ModelViewViewModel/Source/ModelViewViewModelDebugger/Public/MVVMDebugSnapshot.h b/Engine/Plugins/Runtime/ModelViewViewModel/Source/ModelViewViewModelDebugger/Public/MVVMDebugSnapshot.h
new file mode 100644
--- /dev/null
+++ b/Engine/Plugins/Runtime/ModelViewViewModel/Source/ModelViewViewModelDebugger/Public/MVVMDebugSnapshot.h
@@ -0,0 +1,44 @@
+// Copyright Epic Games, Inc. All Rights Reserved.
+
+#pragma once
+
+#include "CoreMinimal.h"
+#include "AssetRegistry/AssetData.h"
+#include "UObject/ObjectMacros.h"
+#include "MVVMDebugItemId.h"
+
+struct FMVVMViewDebugEntry;
+struct FMVVMViewClassDebugEntry;
+struct FMVVMViewModelDebugEntry;
+class UMVVMViewClass;
+
+namespace UE::MVVM
+{
+
+class MODELVIEWVIEWMODELDEBUGGER_API FDebugSnapshot
+{
+private:
+	TArray<TSharedPtr<FMVVMViewDebugEntry>> Views;
+	TArray<TSharedPtr<FMVVMViewClassDebugEntry>> ViewClasses;
+	TArray<TSharedPtr<FMVVMViewModelDebugEntry>> ViewModels;
+
+public:
+	TArrayView<TSharedPtr<FMVVMViewDebugEntry>> GetViews()
+	{
+		return Views;
+	}
+	TArrayView<TSharedPtr<FMVVMViewModelDebugEntry>> GetViewModels()
+	{
+		return ViewModels;
+	}
+
+	TSharedPtr<FMVVMViewDebugEntry> FindView(FGuid Id) const;
+	TSharedPtr<FMVVMViewModelDebugEntry> FindViewModel(FGuid Id) const;
+
+	static TSharedPtr<FDebugSnapshot> CreateSnapshot();
+
+private:
+	TSharedRef<FMVVMViewClassDebugEntry> FindOrAddViewClassEntry(const UMVVMViewClass* ViewClass);
+};
+
+};
\ No newline at end of file
diff --git a/Engine/Plugins/Runtime/ModelViewViewModel/Source/ModelViewViewModelDebugger/Public/MVVMDebugView.h b/Engine/Plugins/Runtime/ModelViewViewModel/Source/ModelViewViewModelDebugger/Public/MVVMDebugView.h
new file mode 100644
--- /dev/null
+++ b/Engine/Plugins/Runtime/ModelViewViewModel/Source/ModelViewViewModelDebugger/Public/MVVMDebugView.h
@@ -0,0 +1,59 @@
+// Copyright Epic Games, Inc. All Rights Reserved.
+
+#pragma once
+
+#include "CoreMinimal.h"
+#include "AssetRegistry/AssetData.h"
+#include "UObject/ObjectMacros.h"
+
+#include "MVVMDebugView.generated.h"
+
+class UMVVMView;
+
+
+USTRUCT()
+struct FMVVMViewSourceDebugEntry
+{
+	GENERATED_BODY()
+
+	UPROPERTY(VisibleAnywhere, Category = "Viewmodel")
+	FName SourceInstanceName;
+
+	UPROPERTY(VisibleAnywhere, Category = "Viewmodel")
+	FAssetData SourceAsset;
+
+	UPROPERTY(VisibleAnywhere, Category = "Viewmodel")
+	FGuid ViewModelDebugId; // if it exist and if it's a viewmodel
+
+	TWeakObjectPtr<UObject> LiveSource; // TScriptInterface<INotifyFieldValueChanged>
+};
+
+
+USTRUCT()
+struct FMVVMViewDebugEntry
+{
+	GENERATED_BODY()
+
+	UPROPERTY(VisibleAnywhere, Category = "Viewmodel")
+	FName UserWidgetInstanceName;
+
+	UPROPERTY(VisibleAnywhere, Category = "Viewmodel")
+	FName LocalPlayerName;
+
+	UPROPERTY(VisibleAnywhere, Category = "Viewmodel")
+	FName WorldName;
+
+	UPROPERTY(VisibleAnywhere, Category = "Viewmodel")
+	FAssetData UserWidgetAsset;
+
+	UPROPERTY(VisibleAnywhere, Category = "Viewmodel")
+	TArray<FMVVMViewSourceDebugEntry> Sources;
+
+	UPROPERTY(VisibleAnywhere, Category = "Viewmodel")
+	FGuid ViewClassDebugId;
+	
+	UPROPERTY()
+	FGuid ViewInstanceDebugId;
+
+	TWeakObjectPtr<UMVVMView> LiveView;
+};
\ No newline at end of file
diff --git a/Engine/Plugins/Runtime/ModelViewViewModel/Source/ModelViewViewModelDebugger/Public/MVVMDebugViewClass.h b/Engine/Plugins/Runtime/ModelViewViewModel/Source/ModelViewViewModelDebugger/Public/MVVMDebugViewClass.h
new file mode 100644
--- /dev/null
+++ b/Engine/Plugins/Runtime/ModelViewViewModel/Source/ModelViewViewModelDebugger/Public/MVVMDebugViewClass.h
@@ -0,0 +1,51 @@
+// Copyright Epic Games, Inc. All Rights Reserved.
+
+#pragma once
+
+#include "CoreMinimal.h"
+#include "AssetRegistry/AssetData.h"
+#include "UObject/ObjectMacros.h"
+
+#include "MVVMDebugViewClass.generated.h"
+
+class UMVVMViewClass;
+
+
+USTRUCT()
+struct FMVVMViewBindingDebugEntry
+{
+	GENERATED_BODY()
+
+	UPROPERTY(VisibleAnywhere, Category = "Viewmodel")
+	FGuid BlueprintViewBindingId;
+
+	UPROPERTY(VisibleAnywhere, Category = "Viewmodel")
+	FName FieldId;
+	
+	UPROPERTY(VisibleAnywhere, Category = "Viewmodel")
+	FString SourceFieldPath;
+	
+	UPROPERTY(VisibleAnywhere, Category = "Viewmodel")
+	FString DestinationFieldPath;
+	
+	UPROPERTY(VisibleAnywhere, Category = "Viewmodel")
+	FString ConversionFunctionFieldPath;
+
+	UPROPERTY(VisibleAnywhere, Category = "Viewmodel")
+	int32 CompiledBindingIndex;
+};
+
+
+USTRUCT()
+struct FMVVMViewClassDebugEntry
+{
+	GENERATED_BODY()
+
+	UPROPERTY(VisibleAnywhere, Category = "Viewmodel")
+	TArray<FMVVMViewBindingDebugEntry> Bindings;
+
+	UPROPERTY()
+	FGuid ViewClassDebugId;
+
+	TWeakObjectPtr<const UMVVMViewClass> LiveViewClass;
+};
\ No newline at end of file
diff --git a/Engine/Plugins/Runtime/ModelViewViewModel/Source/ModelViewViewModelDebugger/Public/MVVMDebugViewModel.h b/Engine/Plugins/Runtime/ModelViewViewModel/Source/ModelViewViewModelDebugger/Public/MVVMDebugViewModel.h
new file mode 100644
--- /dev/null
+++ b/Engine/Plugins/Runtime/ModelViewViewModel/Source/ModelViewViewModelDebugger/Public/MVVMDebugViewModel.h
@@ -0,0 +1,56 @@
+// Copyright Epic Games, Inc. All Rights Reserved.
+
+#pragma once
+
+#include "CoreMinimal.h"
+#include "AssetRegistry/AssetData.h"
+#include "UObject/ObjectMacros.h"
+
+#include "FieldNotificationId.h"
+
+#include "MVVMDebugViewModel.generated.h"
+
+
+USTRUCT()
+struct FMVVMViewModelFieldBoundDebugEntry
+{
+	GENERATED_BODY()
+
+	UPROPERTY(VisibleAnywhere, Category = "Viewmodel")
+	FString ObjectName;
+
+	UPROPERTY(VisibleAnywhere, Category = "Viewmodel")
+	FFieldNotificationId FieldId;
+
+	UPROPERTY(VisibleAnywhere, Category = "Viewmodel")
+	FName FunctionName;
+
+	TWeakObjectPtr<const UObject> LiveInstanceObject;
+};
+
+
+USTRUCT()
+struct FMVVMViewModelDebugEntry
+{
+	GENERATED_BODY()
+
+	UPROPERTY(VisibleAnywhere, Category = "Viewmodel")
+	FName Name;
+
+	UPROPERTY(VisibleAnywhere, Category = "Viewmodel")
+	FString FullName;
+
+	UPROPERTY(VisibleAnywhere, Category = "Viewmodel")
+	FAssetData ViewModelAsset;
+
+	UPROPERTY(VisibleAnywhere, Category = "Viewmodel")
+	TArray<FMVVMViewModelFieldBoundDebugEntry> FieldBound;
+
+	//UPROPERTY(VisibleAnywhere, Category = "Viewmodel")
+	//FInstancedPropertyBag PropertyBag;
+
+	UPROPERTY(VisibleAnywhere, Category = "Viewmodel")
+	FGuid ViewModelDebugId;
+
+	TWeakObjectPtr<UObject> LiveViewModel;
+};
\ No newline at end of file
diff --git a/Engine/Plugins/Runtime/ModelViewViewModel/Source/ModelViewViewModelDebuggerEditor/ModelViewViewModelDebuggerEditor.Build.cs b/Engine/Plugins/Runtime/ModelViewViewModel/Source/ModelViewViewModelDebuggerEditor/ModelViewViewModelDebuggerEditor.Build.cs
new file mode 100644
--- /dev/null
+++ b/Engine/Plugins/Runtime/ModelViewViewModel/Source/ModelViewViewModelDebuggerEditor/ModelViewViewModelDebuggerEditor.Build.cs
@@ -0,0 +1,34 @@
+// Copyright Epic Games, Inc. All Rights Reserved.
+
+using System.IO;
+using System.Collections.Generic;
+using UnrealBuildTool;
+
+public class ModelViewViewModelDebuggerEditor : ModuleRules 
+{
+	public ModelViewViewModelDebuggerEditor(ReadOnlyTargetRules Target) : base(Target)
+	{
+		PublicDependencyModuleNames.AddRange(
+			new string[]
+			{
+				"Core",
+				"CoreUObject",
+				"ModelViewViewModel",
+				"ModelViewViewModelDebugger",
+			});
+
+		PrivateDependencyModuleNames.AddRange(
+			new string[]
+			{
+				"InputCore",
+				"LevelEditor",
+				"MessageLog",
+				"Projects",
+				"SlateCore",
+				"Slate",
+				"ToolMenus",
+				"ToolWidgets",
+				"WorkspaceMenuStructure",
+			});
+	}
+}
diff --git a/Engine/Plugins/Runtime/ModelViewViewModel/Source/ModelViewViewModelDebuggerEditor/Private/ModelViewViewModelDebuggerEditorModule.cpp b/Engine/Plugins/Runtime/ModelViewViewModel/Source/ModelViewViewModelDebuggerEditor/Private/ModelViewViewModelDebuggerEditorModule.cpp
new file mode 100644
--- /dev/null
+++ b/Engine/Plugins/Runtime/ModelViewViewModel/Source/ModelViewViewModelDebuggerEditor/Private/ModelViewViewModelDebuggerEditorModule.cpp
@@ -0,0 +1,60 @@
+// Copyright Epic Games, Inc. All Rights Reserved.
+
+#include "LevelEditor.h"
+
+#include "Modules/ModuleInterface.h"
+#include "Modules/ModuleManager.h"
+#include "ToolMenus.h"
+#include "WorkspaceMenuStructure.h"
+#include "WorkspaceMenuStructureModule.h"
+
+#include "Framework/Docking/TabManager.h"
+#include "Styling/MVVMDebuggerEditorStyle.h"
+#include "Widgets/SMainDebugTab.h"
+#include "Widgets/Docking/SDockTab.h"
+
+#define LOCTEXT_NAMESPACE "MVVMDebuggerModule"
+
+namespace UE::MVVM::Private
+{
+	const FName MainTabName("MVVMDebug");
+}
+
+class FMVVMDebuggerEditorModule : public IModuleInterface
+{
+public:
+	virtual void StartupModule() override
+	{
+		UE::MVVM::FMVVMDebuggerEditorStyle::CreateInstance();
+
+		UToolMenu* ModesToolbar = UToolMenus::Get()->RegisterMenu("ModelViewViewModel.Debug.Toolbar", FName(), EMultiBoxType::ToolBar);
+		ModesToolbar->StyleName = "AssetEditorToolbar";
+
+		FGlobalTabmanager::Get()->RegisterNomadTabSpawner(
+			UE::MVVM::Private::MainTabName, FOnSpawnTab::CreateRaw(this, &FMVVMDebuggerEditorModule::SpawnDebugTab))
+			.SetDisplayName(LOCTEXT("TabTitle", "Viewmodel Debugger"))
+			.SetIcon(FSlateIcon("MVVMDebuggerEditorStyle", "Viewmodel.TabIcon"))
+			.SetTooltipText(LOCTEXT("TooltipText", "Opens Viewmodel Debugger."))
+			.SetGroup(WorkspaceMenu::GetMenuStructure().GetDeveloperToolsDebugCategory());
+	}
+
+	virtual void ShutdownModule() override
+	{
+		FGlobalTabmanager::Get()->UnregisterNomadTabSpawner(UE::MVVM::Private::MainTabName);
+		UE::MVVM::FMVVMDebuggerEditorStyle::DestroyInstance();
+	}
+
+private:
+	TSharedRef<SDockTab> SpawnDebugTab(const FSpawnTabArgs& SpawnTabArgs)
+	{
+		const TSharedRef<SDockTab> MajorTab = SNew(SDockTab)
+		.TabRole(ETabRole::NomadTab);
+		MajorTab->SetContent(SNew(UE::MVVM::SMainDebug, MajorTab));
+		return MajorTab;
+	}
+};
+
+
+IMPLEMENT_MODULE(FMVVMDebuggerEditorModule, ModelViewViewModelDebuggerEditor);
+
+#undef LOCTEXT_NAMESPACE
diff --git a/Engine/Plugins/Runtime/ModelViewViewModel/Source/ModelViewViewModelDebuggerEditor/Private/Styling/MVVMDebuggerEditorStyle.cpp b/Engine/Plugins/Runtime/ModelViewViewModel/Source/ModelViewViewModelDebuggerEditor/Private/Styling/MVVMDebuggerEditorStyle.cpp
new file mode 100644
--- /dev/null
+++ b/Engine/Plugins/Runtime/ModelViewViewModel/Source/ModelViewViewModelDebuggerEditor/Private/Styling/MVVMDebuggerEditorStyle.cpp
@@ -0,0 +1,62 @@
+// Copyright Epic Games, Inc. All Rights Reserved.
+
+#include "Styling/MVVMDebuggerEditorStyle.h"
+#include "Brushes/SlateColorBrush.h"
+#include "Brushes/SlateImageBrush.h"
+#include "Styling/SlateStyleMacros.h"
+#include "Brushes/SlateRoundedBoxBrush.h"
+#include "Styling/CoreStyle.h"
+#include "Styling/SlateStyleRegistry.h"
+#include "Styling/SlateTypes.h"
+#include "Styling/StyleColors.h"
+#include "Interfaces/IPluginManager.h"
+
+namespace UE::MVVM
+{
+
+TUniquePtr<FMVVMDebuggerEditorStyle, FMVVMDebuggerEditorStyle::FCustomDeleter> FMVVMDebuggerEditorStyle::Instance;
+
+
+FMVVMDebuggerEditorStyle::FMVVMDebuggerEditorStyle()
+	: FSlateStyleSet("MVVMDebuggerEditorStyle")
+{
+	const FVector2D Icon10x10(10.0f, 10.0f);
+	const FVector2D Icon14x14(14.0f, 14.0f);
+	const FVector2D Icon16x16(16.0f, 16.0f);
+	const FVector2D Icon20x20(20.0f, 20.0f);
+	const FVector2D Icon24x24(24.0f, 24.0f);
+	const FVector2D Icon32x32(32.0f, 32.0f);
+	const FVector2D Icon40x40(40.0f, 40.0f);
+
+	TSharedPtr<IPlugin> MVVMPlugin = IPluginManager::Get().FindPlugin("ModelViewViewModel");
+	if (ensure(MVVMPlugin))
+	{
+		SetContentRoot(MVVMPlugin->GetContentDir() / TEXT("Editor"));
+	}
+
+	Set("Viewmodel.TabIcon", new IMAGE_BRUSH_SVG("Slate/ViewModel", Icon16x16));
+	Set("Icon.TakeSnapshot", new IMAGE_BRUSH_SVG("Starship/Common/SaveThumbnail", Icon24x24));
+
+
+	FSlateStyleRegistry::RegisterSlateStyle(*this);
+}
+
+
+FMVVMDebuggerEditorStyle::~FMVVMDebuggerEditorStyle()
+{
+	FSlateStyleRegistry::UnRegisterSlateStyle(*this);
+}
+
+
+void FMVVMDebuggerEditorStyle::CreateInstance()
+{
+	Instance.Reset(new FMVVMDebuggerEditorStyle);
+}
+
+
+void FMVVMDebuggerEditorStyle::DestroyInstance()
+{
+	Instance.Reset();
+}
+
+} //namespace
\ No newline at end of file
diff --git a/Engine/Plugins/Runtime/ModelViewViewModel/Source/ModelViewViewModelDebuggerEditor/Private/Styling/MVVMDebuggerEditorStyle.h b/Engine/Plugins/Runtime/ModelViewViewModel/Source/ModelViewViewModelDebuggerEditor/Private/Styling/MVVMDebuggerEditorStyle.h
new file mode 100644
--- /dev/null
+++ b/Engine/Plugins/Runtime/ModelViewViewModel/Source/ModelViewViewModelDebuggerEditor/Private/Styling/MVVMDebuggerEditorStyle.h
@@ -0,0 +1,41 @@
+// Copyright Epic Games, Inc. All Rights Reserved.
+
+#pragma once
+
+#include "Styling/SlateStyle.h"
+
+namespace UE::MVVM
+{
+
+/**
+ * 
+ */
+class FMVVMDebuggerEditorStyle : public FSlateStyleSet
+{
+
+public:
+	static void CreateInstance();
+	static void DestroyInstance();
+
+	static FMVVMDebuggerEditorStyle& Get()
+	{
+		check(Instance);
+		return *Instance;
+	}
+
+private:
+	FMVVMDebuggerEditorStyle();
+	FMVVMDebuggerEditorStyle(const FMVVMDebuggerEditorStyle&) = delete;
+	FMVVMDebuggerEditorStyle& operator=(const FMVVMDebuggerEditorStyle&) = delete;
+	~FMVVMDebuggerEditorStyle();
+	struct FCustomDeleter
+	{
+		void operator()(FMVVMDebuggerEditorStyle* Ptr) const
+		{
+			delete Ptr;
+		}
+	};
+	static TUniquePtr<FMVVMDebuggerEditorStyle, FCustomDeleter> Instance;
+};
+
+} // namespace
diff --git a/Engine/Plugins/Runtime/ModelViewViewModel/Source/ModelViewViewModelDebuggerEditor/Private/Widgets/SDetailsTab.cpp b/Engine/Plugins/Runtime/ModelViewViewModel/Source/ModelViewViewModelDebuggerEditor/Private/Widgets/SDetailsTab.cpp
new file mode 100644
--- /dev/null
+++ b/Engine/Plugins/Runtime/ModelViewViewModel/Source/ModelViewViewModelDebuggerEditor/Private/Widgets/SDetailsTab.cpp
@@ -0,0 +1,45 @@
+// Copyright Epic Games, Inc. All Rights Reserved.
+
+#include "Widgets/SDetailsTab.h"
+
+#include "IDetailsView.h"
+#include "Modules/ModuleManager.h"
+#include "PropertyEditorModule.h"
+
+#include "Widgets/SBoxPanel.h"
+
+#define LOCTEXT_NAMESPACE "MVVMDebuggerSDetailsTab"
+
+namespace UE::MVVM
+{
+
+void SDetailsTab::Construct(const FArguments& InArgs)
+{
+	FPropertyEditorModule& PropertyEditorModule = FModuleManager::GetModuleChecked<FPropertyEditorModule>("PropertyEditor");
+	FDetailsViewArgs DetailsViewArgs;
+	{
+		DetailsViewArgs.bAllowSearch = true;
+		DetailsViewArgs.bShowOptions = true;
+		DetailsViewArgs.bAllowMultipleTopLevelObjects = false;
+		DetailsViewArgs.bAllowFavoriteSystem = true;
+		DetailsViewArgs.bShowObjectLabel = false;
+		DetailsViewArgs.bHideSelectionTip = true;
+	}
+
+	TSharedRef<IDetailsView> PropertyView = PropertyEditorModule.CreateDetailView(DetailsViewArgs);
+	DetailView = PropertyView;
+
+	ChildSlot
+	[
+		PropertyView
+	];
+}
+
+void SDetailsTab::SetObjects(const TArray<UObject*>& InObjects)
+{
+	DetailView->SetObjects(InObjects);
+}
+
+} //namespace
+
+#undef LOCTEXT_NAMESPACE
\ No newline at end of file
diff --git a/Engine/Plugins/Runtime/ModelViewViewModel/Source/ModelViewViewModelDebuggerEditor/Private/Widgets/SDetailsTab.h b/Engine/Plugins/Runtime/ModelViewViewModel/Source/ModelViewViewModelDebuggerEditor/Private/Widgets/SDetailsTab.h
new file mode 100644
--- /dev/null
+++ b/Engine/Plugins/Runtime/ModelViewViewModel/Source/ModelViewViewModelDebuggerEditor/Private/Widgets/SDetailsTab.h
@@ -0,0 +1,26 @@
+// Copyright Epic Games, Inc. All Rights Reserved.
+#pragma once
+
+#include "CoreMinimal.h"
+#include "Widgets/SCompoundWidget.h"
+
+class IDetailsView;
+
+namespace UE::MVVM
+{
+
+class SDetailsTab : public SCompoundWidget
+{
+public:
+	SLATE_BEGIN_ARGS(SDetailsTab) { }
+	SLATE_END_ARGS()
+
+public:
+	void Construct(const FArguments& InArgs);
+	void SetObjects(const TArray<UObject*>& InObjects);
+
+private:
+	TSharedPtr<IDetailsView> DetailView;
+};
+
+} //namespace
diff --git a/Engine/Plugins/Runtime/ModelViewViewModel/Source/ModelViewViewModelDebuggerEditor/Private/Widgets/SMainDebugTab.cpp b/Engine/Plugins/Runtime/ModelViewViewModel/Source/ModelViewViewModelDebuggerEditor/Private/Widgets/SMainDebugTab.cpp
new file mode 100644
--- /dev/null
+++ b/Engine/Plugins/Runtime/ModelViewViewModel/Source/ModelViewViewModelDebuggerEditor/Private/Widgets/SMainDebugTab.cpp
@@ -0,0 +1,349 @@
+// Copyright Epic Games, Inc. All Rights Reserved.
+
+#include "Widgets/SMainDebugTab.h"
+
+#include "MVVMDebugSnapshot.h"
+
+#include "ToolMenus.h"
+
+#include "Framework/Docking/LayoutExtender.h"
+#include "Framework/Docking/LayoutService.h"
+#include "Framework/Docking/TabManager.h"
+#include "Framework/MultiBox/MultiBoxBuilder.h"
+#include "Styling/AppStyle.h"
+#include "Widgets/Docking/SDockTab.h"
+#include "Widgets/Layout/SBorder.h"
+#include "Widgets/SBoxPanel.h"
+#include "Widgets/SDetailsTab.h"
+#include "Widgets/SMessagesLog.h"
+#include "Widgets/SSelectionTab.h"
+#include "Widgets/SViewModelBindingDetail.h"
+
+
+#define LOCTEXT_NAMESPACE "MVVMDebuggerMainDebug"
+
+namespace UE::MVVM
+{
+
+namespace Private
+{
+const FLazyName Stack_Selection = "Selection";
+const FLazyName Stack_Binding = "Binding";
+const FLazyName Stack_Detail = "Detail";
+const FLazyName Stack_Messages = "Messages";
+}
+
+void SMainDebug::Construct(const FArguments& InArgs, const TSharedRef<SDockTab>& InParentTab)
+{
+	//ULevelEditorMenuContext* LevelEditorMenuContext = NewObject<ULevelEditorMenuContext>();
+	//FToolMenuContext MenuContext = FToolMenuContext(LevelEditorMenuContext);
+	FToolMenuContext MenuContext;
+	BuildToolMenu();
+
+	FMenuBarBuilder MenuBarBuilder = FMenuBarBuilder(TSharedPtr<FUICommandList>());
+	MenuBarBuilder.AddPullDownMenu(
+		LOCTEXT("WindowMenuLabel", "Window"),
+		FText::GetEmpty(),
+		FNewMenuDelegate::CreateSP(this, &SMainDebug::HandlePullDownWindowMenu),
+		"Window"
+	);
+
+	const TSharedRef<SWidget> MenuWidget = MenuBarBuilder.MakeWidget();
+
+	TSharedRef<SWidget> DockingArea = CreateDockingArea(InParentTab);
+	TabManager->SetMenuMultiBox(MenuBarBuilder.GetMultiBox(), MenuWidget);
+
+	ChildSlot
+	[
+		SNew(SVerticalBox)
+		+ SVerticalBox::Slot()
+		.AutoHeight()
+		[
+			MenuWidget
+		]
+		+ SVerticalBox::Slot()
+		.Padding(FMargin(0.0f, 0.0f, 0.0f, 2.0f))
+		.AutoHeight()
+		[
+			SNew(SBorder)
+			.Padding(0)
+			.BorderImage(FAppStyle::Get().GetBrush("NoBorder"))
+			[
+				UToolMenus::Get()->GenerateWidget("ModelViewViewModel.Debug.Toolbar", MenuContext)
+			]
+			]
+			+ SVerticalBox::Slot()
+			.Padding(4.0f, 2.f, 4.f, 2.f)
+			.FillHeight(1.0f)
+			[
+				DockingArea
+			]
+	];
+}
+
+
+void SMainDebug::HandlePullDownWindowMenu(FMenuBuilder& MenuBuilder)
+{
+	if (!TabManager.IsValid())
+	{
+		return;
+	}
+
+	TabManager->PopulateLocalTabSpawnerMenu(MenuBuilder);
+}
+
+
+void SMainDebug::BuildToolMenu()
+{
+	UToolMenu* ModesToolbar = UToolMenus::Get()->ExtendMenu("ModelViewViewModel.Debug.Toolbar");
+	check(ModesToolbar);
+
+	{
+		FToolMenuSection& Section = ModesToolbar->AddSection("Snapshot"); 
+
+		FToolMenuEntry TakeSnapshotButtonEntry = FToolMenuEntry::InitToolBarButton(
+			"TakeSnapshot",
+			FUIAction(
+				FExecuteAction::CreateSP(this, &SMainDebug::HandleTakeSnapshot)
+			),
+			LOCTEXT("TakesnapshotLabel", "Take snapshot"),
+			LOCTEXT("TakesnapshotTooltip", "Take snapshot"),
+			FSlateIcon("MVVMDebuggerEditorStyle", "Viewmodel.TabIcon")
+		);
+		Section.AddEntry(TakeSnapshotButtonEntry);
+
+		FToolMenuEntry ConfigureSnapshotMenuEntry = FToolMenuEntry::InitComboButton(
+			"ConfigureSnapshot",
+			FUIAction(),
+			FOnGetContent::CreateSP(this, &SMainDebug::HandleSnapshotMenuContent),
+			LOCTEXT("ConfigureSnapshotLabel", "Configure snapshot"),
+			LOCTEXT("ConfigureSnapshotTooltip", "Configure snapshot"),
+			FSlateIcon(),
+			true //bInSimpleComboBox
+		);
+		ConfigureSnapshotMenuEntry.StyleNameOverride = "CalloutToolbar";
+		Section.AddEntry(ConfigureSnapshotMenuEntry);
+
+		//Section.AddEntry(FToolMenuEntry::InitToolBarButton(
+		//	"LoadSnapshot",
+		//	FUIAction(
+		//		FExecuteAction::CreateSP(this, &SMainDebug::HandleLoadSnapshot)
+		//		),
+		//	LOCTEXT("LoadSnapshotLabel", "Load Snapshot"),
+		//	LOCTEXT("LoadSnapshotTooltip", "Load Snapshot"),
+		//	FSlateIcon(FAppStyle::Get().GetStyleSetName(), "Icons.Import")
+		//));
+
+		//Section.AddEntry(FToolMenuEntry::InitToolBarButton(
+		//	"SaveSnapshot",
+		//	FUIAction(
+		//		FExecuteAction::CreateSP(this, &SMainDebug::HandleSaveSnapshot),
+		//		FCanExecuteAction::CreateSP(this, &SMainDebug::HasValidSnapshot)
+		//	),
+		//	LOCTEXT("SaveSnapshotLabel", "Save Snapshot"),
+		//	LOCTEXT("SaveSnapshotTooltip", "Save Snapshot"),
+		//	FSlateIcon(FAppStyle::Get().GetStyleSetName(), "Icons.Save")
+		//));
+	}
+}
+
+
+TSharedRef<SWidget> SMainDebug::HandleSnapshotMenuContent()
+{
+	const bool bShouldCloseWindowAfterMenuSelection = true;
+	FMenuBuilder MenuBuilder(bShouldCloseWindowAfterMenuSelection, nullptr);
+
+	MenuBuilder.BeginSection("Profile", LOCTEXT("SnapshotContextMenuSectionName", "Snapshot"));
+	{
+		//MenuBuilder.AddMenuEntry(
+		//	LOCTEXT("CreateMenuLabel", "New Empty Media Profile"),
+		//	LOCTEXT("CreateMenuTooltip", "Create a new Media Profile asset."),
+		//	FSlateIcon(FMediaFrameworkUtilitiesEditorStyle::GetStyleSetName(), TEXT("ClassIcon.MediaProfile")),
+		//	FUIAction(FExecuteAction::CreateRaw(this, &FMediaProfileMenuEntryImpl::CreateNewProfile))
+		//);
+	}
+	MenuBuilder.EndSection();
+
+	return MenuBuilder.MakeWidget();
+}
+
+
+void SMainDebug::HandleTakeSnapshot()
+{
+	Snapshot = FDebugSnapshot::CreateSnapshot();
+	if (TSharedPtr<SSelectionTab> SelectionViewPtr = SelectionView.Pin())
+	{
+		SelectionViewPtr->SetSnapshot(Snapshot);
+	}
+}
+
+
+void SMainDebug::HandleLoadSnapshot()
+{
+	//FNotificationInfo Info(LOCTEXT("LoadFail", "Failed to load snapshot from disk."));
+	//Info.ExpireDuration = 2.0f;
+	//FSlateNotificationManager::Get().AddNotification(Info)->SetCompletionState(SNotificationItem::CS_Fail);
+}
+
+
+void SMainDebug::HandleSaveSnapshot()
+{}
+
+
+bool SMainDebug::HasValidSnapshot() const
+{
+	return Snapshot.IsValid();
+}
+
+
+void SMainDebug::HandleObjectSelectionChanged()
+{
+	if (Snapshot == nullptr)
+	{
+		if (TSharedPtr<SDetailsTab> DetailViewPtr = DetailView.Pin())
+		{
+			DetailViewPtr->SetObjects(TArray<UObject*>());
+		}
+		if (TSharedPtr<SViewModelBindingDetail> ViewModelBindingDetailPtr = ViewModelBindingDetail.Pin())
+		{
+			ViewModelBindingDetailPtr->SetViewModels(TArray<TSharedPtr<FMVVMViewModelDebugEntry>>());
+		}
+	}
+	else
+	{
+		if (TSharedPtr<SSelectionTab> SelectionPtr = SelectionView.Pin())
+		{
+			if (TSharedPtr<SDetailsTab> DetailViewPtr = DetailView.Pin())
+			{
+				DetailViewPtr->SetObjects(SelectionPtr->GetSelectedObjects());
+			}
+
+			if (TSharedPtr<SViewModelBindingDetail> ViewModelBindingDetailPtr = ViewModelBindingDetail.Pin())
+			{
+				TArray<TSharedPtr<FMVVMViewModelDebugEntry>> ViewModels;
+				for (FDebugItemId DebugItem : SelectionPtr->GetSelectedItems())
+				{
+					if (DebugItem.Type == FDebugItemId::EType::ViewModel)
+					{
+						if (TSharedPtr<FMVVMViewModelDebugEntry> Found = Snapshot->FindViewModel(DebugItem.Id))
+						{
+							ViewModels.Add(Found);
+						}
+					}
+				}
+				ViewModelBindingDetailPtr->SetViewModels(ViewModels);
+			}
+		}
+	}
+}
+
+
+TSharedRef<SWidget> SMainDebug::CreateDockingArea(const TSharedRef<SDockTab>& InParentTab)
+{
+	const FName LayoutName = TEXT("MVVMDebugger_Layout_v1.8");
+	TSharedRef<FTabManager::FLayout> DefaultLayer = FTabManager::NewLayout(LayoutName)
+		->AddArea
+		(
+			FTabManager::NewPrimaryArea()
+			->SetOrientation(Orient_Vertical)
+			->Split
+			(
+				FTabManager::NewSplitter()
+				->SetOrientation(Orient_Horizontal)
+				->Split
+				(
+					FTabManager::NewStack()
+					->SetSizeCoefficient(0.25f)
+					->SetHideTabWell(true)
+					->AddTab(Private::Stack_Selection, ETabState::OpenedTab)
+				)
+				->Split
+				(
+					FTabManager::NewStack()
+					->SetSizeCoefficient(0.5f)
+					->SetHideTabWell(true)
+					->AddTab(Private::Stack_Binding, ETabState::OpenedTab)
+				)
+				->Split
+				(
+					FTabManager::NewStack()
+					->SetSizeCoefficient(0.25f)
+					->SetHideTabWell(true)
+					->AddTab(Private::Stack_Detail, ETabState::OpenedTab)
+				)
+			)
+			->Split
+			(
+				FTabManager::NewStack()
+				->SetSizeCoefficient(1.0f)
+				->SetHideTabWell(true)
+				->AddTab(Private::Stack_Messages, ETabState::OpenedTab)
+			)
+		);
+
+	TSharedRef<FTabManager::FLayout> Layout = FLayoutSaveRestore::LoadFromConfig(GEditorLayoutIni, DefaultLayer);
+	FLayoutExtender LayoutExtender;
+	Layout->ProcessExtensions(LayoutExtender);
+
+	TabManager = FGlobalTabmanager::Get()->NewTabManager(InParentTab);
+
+	TabManager->RegisterTabSpawner(Private::Stack_Selection, FOnSpawnTab::CreateSP(this, &SMainDebug::SpawnSelectionTab))
+		.SetDisplayName(LOCTEXT("SelectionTab", "Selection"));
+	TabManager->RegisterTabSpawner(Private::Stack_Binding, FOnSpawnTab::CreateSP(this, &SMainDebug::SpawnBindingTab))
+		.SetDisplayName(LOCTEXT("BindingTab", "Bindings"));
+	TabManager->RegisterTabSpawner(Private::Stack_Detail, FOnSpawnTab::CreateSP(this, &SMainDebug::SpawnDetailTab))
+		.SetDisplayName(LOCTEXT("DetailTab", "Details"));
+	TabManager->RegisterTabSpawner(Private::Stack_Messages, FOnSpawnTab::CreateSP(this, &SMainDebug::SpawnMessagesTab))
+		.SetDisplayName(LOCTEXT("MessageLogTab", "Messages Log"));
+
+	return TabManager->RestoreFrom(Layout, nullptr).ToSharedRef();
+}
+
+
+TSharedRef<SDockTab> SMainDebug::SpawnSelectionTab(const FSpawnTabArgs& Args)
+{
+	TSharedRef<SSelectionTab> LocalSelectionView = SNew(SSelectionTab)
+		.OnSelectionChanged(this, &SMainDebug::HandleObjectSelectionChanged);
+	SelectionView = LocalSelectionView;
+	LocalSelectionView->SetSnapshot(Snapshot);
+	return SNew(SDockTab)
+		.Label(LOCTEXT("BindingTab", "Bindings"))
+		[
+			LocalSelectionView
+		];
+}
+
+
+TSharedRef<SDockTab> SMainDebug::SpawnBindingTab(const FSpawnTabArgs& Args)
+{
+	return SNew(SDockTab)
+		.Label(LOCTEXT("DetailTab", "Bindings"))
+		[
+			SAssignNew(ViewModelBindingDetail, SViewModelBindingDetail)
+		];
+}
+
+
+TSharedRef<SDockTab> SMainDebug::SpawnDetailTab(const FSpawnTabArgs& Args)
+{
+	return SNew(SDockTab)
+		.Label(LOCTEXT("DetailTab", "Details"))
+		[
+			SAssignNew(DetailView, SDetailsTab)
+		];
+}
+
+
+TSharedRef<SDockTab> SMainDebug::SpawnMessagesTab(const FSpawnTabArgs& Args)
+{
+	return SNew(SDockTab)
+		.Label(LOCTEXT("MessageLogTab", "Messages Log"))
+		.ShouldAutosize(true)
+		[
+			SAssignNew(MessageLog, SMessagesLog)
+		];
+}
+
+} //namespace
+
+#undef LOCTEXT_NAMESPACE
\ No newline at end of file
diff --git a/Engine/Plugins/Runtime/ModelViewViewModel/Source/ModelViewViewModelDebuggerEditor/Private/Widgets/SMainDebugTab.h b/Engine/Plugins/Runtime/ModelViewViewModel/Source/ModelViewViewModelDebuggerEditor/Private/Widgets/SMainDebugTab.h
new file mode 100644
--- /dev/null
+++ b/Engine/Plugins/Runtime/ModelViewViewModel/Source/ModelViewViewModelDebuggerEditor/Private/Widgets/SMainDebugTab.h
@@ -0,0 +1,62 @@
+// Copyright Epic Games, Inc. All Rights Reserved.
+#pragma once
+
+#include "CoreMinimal.h"
+#include "Widgets/SCompoundWidget.h"
+
+
+class FMenuBuilder;
+class FSpawnTabArgs;
+class FTabManager;
+class SDockTab;
+class UToolMenu;
+
+namespace UE::MVVM
+{
+
+class FDebugSnapshot;
+
+class SDetailsTab;
+class SMessagesLog;
+class SSelectionTab;
+class SViewModelBindingDetail;
+
+class SMainDebug : public SCompoundWidget
+{
+public:
+	SLATE_BEGIN_ARGS(SMainDebug) { }
+	SLATE_END_ARGS()
+
+public:
+	void Construct(const FArguments& InArgs, const TSharedRef<SDockTab>& InParentTab);
+
+private:
+	void BuildToolMenu();
+	void HandlePullDownWindowMenu(FMenuBuilder& MenuBuilder);
+	TSharedRef<SWidget> HandleSnapshotMenuContent();
+
+	void HandleTakeSnapshot();
+	void HandleLoadSnapshot();
+	void HandleSaveSnapshot();
+	bool HasValidSnapshot() const;
+
+	void HandleObjectSelectionChanged();
+
+	TSharedRef<SWidget> CreateDockingArea(const TSharedRef<SDockTab>& InParentTab);
+	TSharedRef<SDockTab> SpawnSelectionTab(const FSpawnTabArgs& Args);
+	TSharedRef<SDockTab> SpawnBindingTab(const FSpawnTabArgs& Args);
+	TSharedRef<SDockTab> SpawnDetailTab(const FSpawnTabArgs& Args);
+	TSharedRef<SDockTab> SpawnMessagesTab(const FSpawnTabArgs& Args);
+
+private:
+	TSharedPtr<FTabManager> TabManager;
+
+	TWeakPtr<SDetailsTab> DetailView;
+	TWeakPtr<SMessagesLog> MessageLog;
+	TWeakPtr<SSelectionTab> SelectionView;
+	TWeakPtr<SViewModelBindingDetail> ViewModelBindingDetail;
+
+	TSharedPtr<FDebugSnapshot> Snapshot;
+};
+
+} //namespace
diff --git a/Engine/Plugins/Runtime/ModelViewViewModel/Source/ModelViewViewModelDebuggerEditor/Private/Widgets/SMessagesLog.cpp b/Engine/Plugins/Runtime/ModelViewViewModel/Source/ModelViewViewModelDebuggerEditor/Private/Widgets/SMessagesLog.cpp
new file mode 100644
--- /dev/null
+++ b/Engine/Plugins/Runtime/ModelViewViewModel/Source/ModelViewViewModelDebuggerEditor/Private/Widgets/SMessagesLog.cpp
@@ -0,0 +1,64 @@
+// Copyright Epic Games, Inc. All Rights Reserved.
+
+#include "Widgets/SMessagesLog.h"
+
+#include "IMessageLogListing.h"
+#include "MessageLogModule.h"
+#include "Modules/ModuleManager.h"
+
+#include "Widgets/SBoxPanel.h"
+#include "Widgets/SOverlay.h"
+#include "Widgets/Text/STextBlock.h"
+
+#define LOCTEXT_NAMESPACE "MVVMDebuggerMessagesLog"
+
+namespace UE::MVVM
+{
+
+void SMessagesLog::Construct(const FArguments& InArgs)
+{
+	const FName LogName = "Live Link";
+	TSharedPtr<IMessageLogListing> MessageLogListing;
+
+	FMessageLogModule& MessageLogModule = FModuleManager::LoadModuleChecked<FMessageLogModule>("MessageLog");
+	if (MessageLogModule.IsRegisteredLogListing(LogName))
+	{
+		MessageLogListing = MessageLogModule.GetLogListing(LogName);
+	}
+	TSharedRef<class SWidget> MessageLogListingWidget = MessageLogListing.IsValid() ? MessageLogModule.CreateLogListingWidget(MessageLogListing.ToSharedRef()) : SNullWidget::NullWidget;
+
+	ChildSlot
+	[
+		SNew(SOverlay)
+		+SOverlay::Slot()
+		[
+			MessageLogListingWidget
+		]
+		+SOverlay::Slot()
+		.HAlign(HAlign_Left)
+		.VAlign(VAlign_Bottom)
+		[
+			SNew(SHorizontalBox)
+			+SHorizontalBox::Slot()
+			.AutoWidth()
+			.Padding(10, 4, 4, 10)
+			[
+				SNew(STextBlock)
+				.Text(this, &SMessagesLog::GetMessageCountText)
+			]
+		]
+	];
+}
+
+FText SMessagesLog::GetMessageCountText() const
+{
+	int32 ErrorCount = 0;
+	int32 WarningCount = 0;
+	int32 InfoCount = 0;
+	//FLiveLinkLog::GetInstance()->GetLogCount(ErrorCount, WarningCount, InfoCount);
+	return FText::Format(LOCTEXT("MessageCountText", "{0} Error(s)  {1} Warning(s)"), FText::AsNumber(ErrorCount), FText::AsNumber(WarningCount));
+}
+
+} //namespace
+
+#undef LOCTEXT_NAMESPACE
\ No newline at end of file
diff --git a/Engine/Plugins/Runtime/ModelViewViewModel/Source/ModelViewViewModelDebuggerEditor/Private/Widgets/SMessagesLog.h b/Engine/Plugins/Runtime/ModelViewViewModel/Source/ModelViewViewModelDebuggerEditor/Private/Widgets/SMessagesLog.h
new file mode 100644
--- /dev/null
+++ b/Engine/Plugins/Runtime/ModelViewViewModel/Source/ModelViewViewModelDebuggerEditor/Private/Widgets/SMessagesLog.h
@@ -0,0 +1,24 @@
+// Copyright Epic Games, Inc. All Rights Reserved.
+#pragma once
+
+#include "CoreMinimal.h"
+#include "Widgets/SCompoundWidget.h"
+
+
+namespace UE::MVVM
+{
+
+class SMessagesLog : public SCompoundWidget
+{
+public:
+	SLATE_BEGIN_ARGS(SMessagesLog) { }
+	SLATE_END_ARGS()
+
+public:
+	void Construct(const FArguments& InArgs);
+
+private:
+	FText GetMessageCountText() const;
+};
+
+} //namespace
diff --git a/Engine/Plugins/Runtime/ModelViewViewModel/Source/ModelViewViewModelDebuggerEditor/Private/Widgets/SSelectionBase.cpp b/Engine/Plugins/Runtime/ModelViewViewModel/Source/ModelViewViewModelDebuggerEditor/Private/Widgets/SSelectionBase.cpp
new file mode 100644
--- /dev/null
+++ b/Engine/Plugins/Runtime/ModelViewViewModel/Source/ModelViewViewModelDebuggerEditor/Private/Widgets/SSelectionBase.cpp
@@ -0,0 +1,248 @@
+// Copyright Epic Games, Inc. All Rights Reserved.
+
+#include "Widgets/SSelectionBase.h"
+
+#include "MVVMDebugSnapshot.h"
+#include "MVVMDebugView.h"
+#include "MVVMDebugViewClass.h"
+
+#include "Filters/SFilterSearchBox.h"
+#include "Framework/MultiBox/MultiBoxBuilder.h"
+#include "Widgets/SBoxPanel.h"
+#include "Widgets/Input/SComboButton.h"
+#include "Widgets/Text/STextBlock.h"
+#include "Widgets/Views/SListView.h"
+
+#define LOCTEXT_NAMESPACE "MVVMDebuggerViewSelection"
+
+namespace UE::MVVM::Private
+{
+
+void SSelectionBase::Construct(const FArguments& InArgs)
+{
+	//VerticalBox->AddSlot()
+	//.AutoHeight()
+	//[
+	//	SNew(SVerticalBox)
+	//	+ SVerticalBox::Slot()
+	//	.AutoHeight()
+	//	[
+	//		SNew(SMultiLineEditableTextBox)
+	//		.IsReadOnly(true)
+	//		.Visibility_Lambda([this]() { return Mode->HasErrors() ? EVisibility::Visible : EVisibility::Collapsed; })
+	//		.Font(IDetailLayoutBuilder::GetDetailFontBold())
+	//		.BackgroundColor(FAppStyle::GetColor("ErrorReporting.WarningBackgroundColor"))
+	//		.Text(Mode->GetErrorsText())
+	//		.AutoWrapText(true)
+	//	]
+	//	+ SVerticalBox::Slot()
+	//	.AutoHeight()
+	//	[
+	//		SNew(SButton)
+	//		.OnClicked_Lambda([this] { Mode->RepairErrors(); return FReply::Handled(); })
+	//		.Visibility_Lambda([this]() { return Mode->HasErrors() ? EVisibility::Visible : EVisibility::Collapsed; })
+	//		.HAlign(HAlign_Center)
+	//		.Text(LOCTEXT("SceneOutlinerRepairErrors", "Repair Errors"))
+	//	]
+	//];
+
+	TSharedRef<SHorizontalBox> Toolbar = SNew(SHorizontalBox)
+	//+SHorizontalBox::Slot().VAlign(VAlign_Center)
+	//.Padding(0.0f, 0.0f, 2.0f, 0.0f)
+	//.AutoWidth()
+	//[
+	//	SSceneOutlinerFilterBar::MakeAddFilterButton(FilterBar.ToSharedRef())
+	//]
+	+ SHorizontalBox::Slot()
+	.VAlign(VAlign_Center)
+	[
+		SAssignNew(SearchBox, SFilterSearchBox)
+		.HintText(LOCTEXT("FilterSearch", "Search..."))
+		.ToolTipText(LOCTEXT("FilterSearchHint", "Type here to search (pressing enter selects the results)"))
+		.OnTextChanged(this, &SSelectionBase::HandleFilterTextChanged)
+	]
+	+SHorizontalBox::Slot()
+	.VAlign(VAlign_Center)
+	.AutoWidth()
+	[
+		SNew(SComboButton)
+		.ComboButtonStyle(FAppStyle::Get(), "SimpleComboButtonWithIcon") // Use the tool bar item style for this button
+		.OnGetMenuContent(this, &SSelectionBase::GetViewButtonContent)
+		.HasDownArrow(false)
+		.ButtonContent()
+		[
+			SNew(SImage)
+			.ColorAndOpacity(FSlateColor::UseForeground())
+			.Image(FAppStyle::Get().GetBrush("Icons.Settings"))
+		]
+	];
+	//+SHorizontalBox::Slot()
+	//.VAlign(VAlign_Center)
+	//.AutoWidth()
+	//[
+	//	SNew(SButton)
+	//	.ButtonStyle(FAppStyle::Get(), "SimpleButton")
+	//	.OnClicked(this, &SSelectionBase::HandleLockView)
+	//	.ContentPadding(FMargin(4, 2))
+	//	.HAlign(HAlign_Center)
+	//	.VAlign(VAlign_Center)
+	//	.ToolTipText(LOCTEXT("LockSelectionButton_ToolTip", "Locks the current selection into the Details panel"))
+	//	[
+	//		SNew(SImage)
+	//		.ColorAndOpacity(FSlateColor::UseForeground())
+	//		.Image(this, &SDetailNameArea::OnGetLockButtonImageResource)
+	//	]
+	//];
+
+
+	TSharedRef<SVerticalBox> VerticalBox = SNew(SVerticalBox)
+	+SVerticalBox::Slot()
+	.AutoHeight()
+	.Padding(8.0f, 8.0f, 8.0f, 4.0f)
+	[
+		Toolbar
+	];
+
+	//	VerticalBox->AddSlot()
+	//		.AutoHeight()
+	//		.Padding(0.0f, 0.0f, 0.0f, 4.0f)
+	//		[
+	//			FilterBar.ToSharedRef()
+	//		];
+	//}
+
+	VerticalBox->AddSlot()
+	.FillHeight(1.0)
+	[
+		SNew(SOverlay)
+		//+ SOverlay::Slot()
+		//.HAlign(HAlign_Center)
+		//[
+		//	SNew(STextBlock)
+		//	.Visibility(this, &SSceneOutliner::GetEmptyLabelVisibility)
+		//	.Text(LOCTEXT("EmptyLabel", "Empty"))
+		//	.ColorAndOpacity(FLinearColor(0.4f, 1.0f, 0.4f))
+		//]
+		+ SOverlay::Slot()
+		[
+			SNew(SBorder).BorderImage(FAppStyle::Get().GetBrush("Brushes.Recessed"))
+		]
+		+ SOverlay::Slot()
+		[
+			BuildListImpl()
+		]
+	];
+
+
+	//// Bottom panel status bar, if enabled by the mode
+	//VerticalBox->AddSlot()
+	//.AutoHeight()
+	//[
+	//	SNew(SBorder)
+	//	.BorderImage(FAppStyle::Get().GetBrush("Brushes.Header"))
+	//	.VAlign(VAlign_Center)
+	//	.HAlign(HAlign_Left)
+	//	.Padding(FMargin(14, 9))
+	//	[
+	//		SNew(STextBlock)
+	//		.Text(this, &SSelectionBase::GetFilterStatusText)
+	//		.ColorAndOpacity(this, &SSelectionBase::GetFilterStatusTextColor)
+	//	]
+	//];
+
+	ChildSlot
+	[
+		VerticalBox
+	];
+}
+
+
+void SSelectionBase::SetSnapshot(const TSharedPtr<FDebugSnapshot> InSnapshot)
+{
+	Snapshot = InSnapshot;
+	UpdateSource();
+}
+
+
+void SSelectionBase::HandleFilterTextChanged(const FText& InFilterText)
+{
+	//for (const TSharedPtr<FTreeNode>& Node : FilteredTreeSource)
+	//{
+	//	SetHighlightTextRecursive(Node, InFilterText);
+	//}
+
+	FText ErrorText = FilterTextChangedImpl(InFilterText);
+	if (SearchBox)
+	{
+		SearchBox->SetError(ErrorText);
+	}
+	UpdateSource();
+}
+
+
+TSharedRef<SWidget> SSelectionBase::GetViewButtonContent() const
+{
+	// Menu should stay open on selection if filters are not being shown
+	FMenuBuilder MenuBuilder(true, nullptr);
+
+	//MenuBuilder.BeginSection("OutlinerSettings", LOCTEXT("HierarchyHeading", "Hierarchy"));
+	//{
+	//	MenuBuilder.AddMenuEntry(
+	//		LOCTEXT("ExpandAll", "Refresh"),
+	//		LOCTEXT("ExpandAllToolTip", "Expand All Items in the Hierarchy"),
+	//		FSlateIcon(),
+	//		FUIAction(FExecuteAction::CreateSP(this, &SSelectionBase::ExpandAll)));
+
+	//	MenuBuilder.AddMenuEntry(
+	//		LOCTEXT("CollapseAll", "Collapse All"),
+	//		LOCTEXT("CollapseAllToolTip", "Collapse All Items in the Hierarchy"),
+	//		FSlateIcon(),
+	//		FUIAction(FExecuteAction::CreateSP(this, &SSceneOutliner::CollapseAll)));
+
+	//	MenuBuilder.AddMenuEntry(
+	//		LOCTEXT("ShowHierarchy", "Stack Hierarchy Headers"),
+	//		LOCTEXT("ShowHierarchyToolTip", "Toggle pinning of the hierarchy of items at the top of the outliner"),
+	//		FSlateIcon(),
+	//		FUIAction(
+	//			FExecuteAction::CreateRaw(this, &SSceneOutliner::ToggleStackHierarchyHeaders),
+	//			FCanExecuteAction(),
+	//			FIsActionChecked::CreateRaw(this, &SSceneOutliner::ShouldStackHierarchyHeaders)
+	//		),
+	//		NAME_None,
+	//		EUserInterfaceActionType::ToggleButton
+	//	);
+	//}
+	//MenuBuilder.EndSection();
+
+	//if (bShowFilters)
+	//{
+	//	MenuBuilder.BeginSection("AssetThumbnails", LOCTEXT("ShowHeading", "Show"));
+	//	{
+	//		// Add mode filters
+	//		for (auto& ModeFilterInfo : Mode->GetFilterInfos())
+	//		{
+	//			ModeFilterInfo.Value.AddMenu(MenuBuilder);
+	//		}
+	//	}
+	//	MenuBuilder.EndSection();
+	//}
+	//Mode->CreateViewContent(MenuBuilder);
+
+	return MenuBuilder.MakeWidget();
+}
+
+
+FText SSelectionBase::GetFilterStatusText() const
+{
+	return FText::GetEmpty();
+}
+
+
+FSlateColor SSelectionBase::GetFilterStatusTextColor() const
+{
+	return FSlateColor();
+}
+
+} //namespace
+
+#undef LOCTEXT_NAMESPACE
\ No newline at end of file
diff --git a/Engine/Plugins/Runtime/ModelViewViewModel/Source/ModelViewViewModelDebuggerEditor/Private/Widgets/SSelectionBase.h b/Engine/Plugins/Runtime/ModelViewViewModel/Source/ModelViewViewModelDebuggerEditor/Private/Widgets/SSelectionBase.h
new file mode 100644
--- /dev/null
+++ b/Engine/Plugins/Runtime/ModelViewViewModel/Source/ModelViewViewModelDebuggerEditor/Private/Widgets/SSelectionBase.h
@@ -0,0 +1,146 @@
+// Copyright Epic Games, Inc. All Rights Reserved.
+#pragma once
+
+#include "CoreMinimal.h"
+#include "Misc/TextFilter.h"
+#include "Widgets/SCompoundWidget.h"
+#include "Widgets/Views/SListView.h"
+
+class SFilterSearchBox;
+class ITableRow;
+
+class SHeader;
+class SHeaderRow;
+class STableViewBase;
+
+namespace UE::MVVM
+{
+
+class FDebugSnapshot;
+
+namespace Private
+{
+
+class SSelectionBase : public SCompoundWidget
+{
+public:
+	SLATE_BEGIN_ARGS(SSelectionBase) { }
+	SLATE_EVENT(FSimpleDelegate, OnSelectionChanged)
+	SLATE_END_ARGS()
+
+	void Construct(const FArguments& InArgs);
+
+	void SetSnapshot(const TSharedPtr<FDebugSnapshot> Snapshot);
+	TSharedPtr<FDebugSnapshot> GetSnapshot() const
+	{
+		return Snapshot;
+	}
+
+protected:
+	TSharedPtr<SFilterSearchBox> SearchBox;
+	TSharedPtr<FDebugSnapshot> Snapshot;
+	FSimpleDelegate OnSelectionChanged;
+	bool bIsSelectionReentrant = false;
+
+protected:
+	virtual void UpdateSource() = 0;
+	virtual TSharedRef<STableViewBase> BuildListImpl() = 0;
+	virtual FText FilterTextChangedImpl(const FText& InFilterText) = 0;
+
+private:
+	void HandleFilterTextChanged(const FText& InFilterText);
+	TSharedRef<SWidget> GetViewButtonContent() const;
+	FText GetFilterStatusText() const;
+	FSlateColor GetFilterStatusTextColor() const;
+};
+}
+
+template <typename ItemType, typename ItemWidgetType>
+class SSelectionBase : public Private::SSelectionBase
+{
+public:
+	SLATE_BEGIN_ARGS(SSelectionBase) { }
+	SLATE_EVENT(FSimpleDelegate, OnSelectionChanged)
+	SLATE_END_ARGS()
+
+public:
+	void Construct(const FArguments& InArgs)
+	{
+		OnSelectionChanged = InArgs._OnSelectionChanged;
+		SearchFilter = MakeShared<FTextFilter>(FTextFilter::FItemToStringArray::CreateSP(this, &SSelectionBase::HandleGetFilterStrings));
+		Private::SSelectionBase::Construct(Private::SSelectionBase::FArguments());
+	}
+
+	TArray<ItemType> GetSelectedItems() const
+	{
+		TArray<ItemType> Result;
+		ListView->GetSelectedItems(Result);
+		return Result;
+	}
+
+protected:
+	virtual void GatherFilterStringsImpl(ItemType Item, TArray<FString>& OutStrings) const = 0;
+	virtual TSharedRef<SHeaderRow> BuildHeaderRowImpl() const = 0;
+	virtual TArray<ItemType> UpdateSourceImpl(TSharedPtr<TTextFilter<ItemType>> TextFilter) = 0;
+
+	virtual FText FilterTextChangedImpl(const FText& InFilterText) override
+	{
+		SearchFilter->SetRawFilterText(InFilterText);
+		return SearchFilter->GetFilterErrorText();
+	}
+
+private:
+	virtual void UpdateSource() override
+	{
+		FilteredSource = UpdateSourceImpl(SearchFilter);
+		//FilteredItems.Sort(&FSortPlaceableItems::SortItemsByName);
+		ListView->RequestListRefresh();
+	}
+
+	virtual TSharedRef<STableViewBase> BuildListImpl() override
+	{
+		return SAssignNew(ListView, SListView<ItemType>)
+			.ListItemsSource(&FilteredSource)
+			//.SelectionMode(ESelectionMode::Single)
+			.OnGenerateRow(this, &SSelectionBase::HandleGenerateWidgetForItem)
+			.OnSelectionChanged(this, &SSelectionBase::HandleListSelectionChanged)
+			//.OnMouseButtonDoubleClick(this, &SSceneOutliner::OnOutlinerTreeDoubleClick)
+			//.OnContextMenuOpening(this, &SLiveLinkClientPanel::OnOpenVirtualSubjectContextMenu)
+			.HeaderRow(BuildHeaderRowImpl());
+	}
+
+	void HandleGetFilterStrings(ItemType Item, TArray<FString>& OutStrings)
+	{
+		GatherFilterStringsImpl(Item, OutStrings);
+	}
+
+	void HandleListSelectionChanged(ItemType TreeItem, ESelectInfo::Type SelectInfo)
+	{
+		if (SelectInfo == ESelectInfo::Direct)
+		{
+			return;
+		}
+
+		if (!bIsSelectionReentrant)
+		{
+			TGuardValue<bool> ReentrantGuard(bIsSelectionReentrant, true);
+			OnSelectionChanged.ExecuteIfBound();
+		}
+	}
+
+	TSharedRef<ITableRow> HandleGenerateWidgetForItem(ItemType Entry, const TSharedRef<STableViewBase>& OwnerTable) const
+	{
+		return SNew(ItemWidgetType, OwnerTable)
+			.Entry(Entry);
+	}
+
+
+protected:
+	TArray<ItemType> FilteredSource;
+	TSharedPtr<SListView<ItemType>> ListView;
+
+	using FTextFilter = TTextFilter<ItemType>;
+	TSharedPtr<FTextFilter> SearchFilter;
+};
+
+} //namespace
diff --git a/Engine/Plugins/Runtime/ModelViewViewModel/Source/ModelViewViewModelDebuggerEditor/Private/Widgets/SSelectionTab.cpp b/Engine/Plugins/Runtime/ModelViewViewModel/Source/ModelViewViewModelDebuggerEditor/Private/Widgets/SSelectionTab.cpp
new file mode 100644
--- /dev/null
+++ b/Engine/Plugins/Runtime/ModelViewViewModel/Source/ModelViewViewModelDebuggerEditor/Private/Widgets/SSelectionTab.cpp
@@ -0,0 +1,118 @@
+// Copyright Epic Games, Inc. All Rights Reserved.
+
+#include "Widgets/SSelectionTab.h"
+
+#include "MVVMDebugView.h"
+#include "MVVMDebugViewModel.h"
+#include "MVVMViewModelBase.h"
+#include "View/MVVMView.h"
+
+#include "Widgets/Layout/SSplitter.h"
+#include "Widgets/SBoxPanel.h"
+#include "Widgets/SViewSelection.h"
+#include "Widgets/SViewModelSelection.h"
+
+#define LOCTEXT_NAMESPACE "MVVMDebuggerSelectionTab"
+
+namespace UE::MVVM
+{
+
+void SSelectionTab::Construct(const FArguments& InArgs)
+{
+	OnSelectionChanged = InArgs._OnSelectionChanged;
+
+	ChildSlot
+	[
+		SNew(SBox)
+		[
+			SNew(SSplitter)
+			.Orientation(EOrientation::Orient_Vertical)
+			+ SSplitter::Slot()
+			[
+				SAssignNew(ViewSelection, SViewSelection)
+				.OnSelectionChanged(this, &SSelectionTab::HandleViewSelectionChanged)
+			]
+			+ SSplitter::Slot()
+			[
+				SAssignNew(ViewModelSelection, SViewModelSelection)
+				.OnSelectionChanged(this, &SSelectionTab::HandleViewModleSelectionChanged)
+			]
+		]
+	];
+}
+
+
+void SSelectionTab::SetSnapshot(TSharedPtr<FDebugSnapshot> Snapshot)
+{
+	if (ViewSelection)
+	{
+		ViewSelection->SetSnapshot(Snapshot);
+	}
+	if (ViewModelSelection)
+	{
+		ViewModelSelection->SetSnapshot(Snapshot);
+	}
+	CurrentSelection = ESelection::None;
+}
+
+
+TArray<FDebugItemId> SSelectionTab::GetSelectedItems() const
+{
+	TArray<FDebugItemId> Result;
+	if (CurrentSelection == ESelection::View)
+	{
+		for (const TSharedPtr<FMVVMViewDebugEntry>& Selection : ViewSelection->GetSelectedItems())
+		{
+			Result.Add(FDebugItemId(FDebugItemId::EType::View, Selection->ViewInstanceDebugId));
+		}
+	}
+	else if (CurrentSelection == ESelection::ViewModel)
+	{
+		for (const TSharedPtr<FMVVMViewModelDebugEntry>& Selection : ViewModelSelection->GetSelectedItems())
+		{
+			Result.Add(FDebugItemId(FDebugItemId::EType::ViewModel, Selection->ViewModelDebugId));
+		}
+	}
+	return Result;
+}
+
+
+TArray<UObject*> SSelectionTab::GetSelectedObjects() const
+{
+	TArray<UObject*> Result;
+	if (CurrentSelection == ESelection::View)
+	{
+		for (const TSharedPtr<FMVVMViewDebugEntry>& Selection : ViewSelection->GetSelectedItems())
+		{
+			Result.Add(Selection->LiveView.Get());
+		}
+	}
+	else if (CurrentSelection == ESelection::ViewModel)
+	{
+		for (const TSharedPtr<FMVVMViewModelDebugEntry>& Selection : ViewModelSelection->GetSelectedItems())
+		{
+			Result.Add(Selection->LiveViewModel.Get());
+		}
+	}
+	return Result;
+}
+
+
+void SSelectionTab::HandleViewSelectionChanged()
+{
+	CurrentSelection = ESelection::View;
+	OnSelectionChanged.ExecuteIfBound();
+	ViewModelSelection->SetSelection(FGuid());
+}
+
+
+void SSelectionTab::HandleViewModleSelectionChanged()
+{
+	CurrentSelection = ESelection::ViewModel;
+	OnSelectionChanged.ExecuteIfBound();
+	ViewSelection->SetSelection(FGuid());
+}
+
+} //namespace
+
+#undef LOCTEXT_NAMESPACE
\ No newline at end of file
diff --git a/Engine/Plugins/Runtime/ModelViewViewModel/Source/ModelViewViewModelDebuggerEditor/Private/Widgets/SSelectionTab.h b/Engine/Plugins/Runtime/ModelViewViewModel/Source/ModelViewViewModelDebuggerEditor/Private/Widgets/SSelectionTab.h
new file mode 100644
--- /dev/null
+++ b/Engine/Plugins/Runtime/ModelViewViewModel/Source/ModelViewViewModelDebuggerEditor/Private/Widgets/SSelectionTab.h
@@ -0,0 +1,46 @@
+// Copyright Epic Games, Inc. All Rights Reserved.
+#pragma once
+
+#include "CoreMinimal.h"
+#include "Widgets/SCompoundWidget.h"
+#include "MVVMDebugItemId.h"
+
+
+namespace UE::MVVM
+{
+class FDebugSnapshot;
+class SViewSelection;
+class SViewModelSelection;
+
+class SSelectionTab : public SCompoundWidget
+{
+public:
+	SLATE_BEGIN_ARGS(SSelectionTab) { }
+	SLATE_EVENT(FSimpleDelegate, OnSelectionChanged)
+	SLATE_END_ARGS()
+
+public:
+	void Construct(const FArguments& InArgs);
+
+	void SetSnapshot(TSharedPtr<FDebugSnapshot> Snapshot);
+	TArray<FDebugItemId> GetSelectedItems() const;
+	TArray<UObject*> GetSelectedObjects() const;
+
+private:
+	void HandleViewSelectionChanged();
+	void HandleViewModleSelectionChanged();
+
+private:
+	TSharedPtr<SViewSelection> ViewSelection;
+	TSharedPtr<SViewModelSelection> ViewModelSelection;
+
+	FSimpleDelegate OnSelectionChanged;
+	enum class ESelection
+	{
+		None,
+		View,
+		ViewModel,
+	} CurrentSelection = ESelection::None;
+};
+
+} //namespace
diff --git a/Engine/Plugins/Runtime/ModelViewViewModel/Source/ModelViewViewModelDebuggerEditor/Private/Widgets/SViewModelBindingDetail.cpp b/Engine/Plugins/Runtime/ModelViewViewModel/Source/ModelViewViewModelDebuggerEditor/Private/Widgets/SViewModelBindingDetail.cpp
new file mode 100644
--- /dev/null
+++ b/Engine/Plugins/Runtime/ModelViewViewModel/Source/ModelViewViewModelDebuggerEditor/Private/Widgets/SViewModelBindingDetail.cpp
@@ -0,0 +1,133 @@
+// Copyright Epic Games, Inc. All Rights Reserved.
+
+#include "Widgets/SViewModelBindingDetail.h"
+
+#include "MVVMDebugSnapshot.h"
+#include "MVVMDebugView.h"
+#include "MVVMDebugViewClass.h"
+
+#include "Widgets/SBoxPanel.h"
+#include "Widgets/Text/STextBlock.h"
+#include "Widgets/Views/SListView.h"
+
+//ToolWidgets
+
+#define LOCTEXT_NAMESPACE "MVVMViewModelBindingDetail"
+
+
+namespace UE::MVVM
+{
+namespace Private
+{
+
+namespace ViewColumnName
+{
+FLazyName FieldId = "FieldId";
+FLazyName ObjectName = "ObjectName";
+}
+
+class SViewModelBindingDetailItem : public SMultiColumnTableRow<SViewModelBindingDetail::FViewModelBindingId>
+{
+	SLATE_BEGIN_ARGS(SViewModelBindingDetailItem) {}
+	SLATE_ARGUMENT(SViewModelBindingDetail::FViewModelBindingId, BindingId)
+	SLATE_END_ARGS()
+
+	void Construct(const FArguments& Args, const TSharedRef<STableViewBase>& OwnerTableView)
+	{
+		BindingId = Args._BindingId;
+
+		SMultiColumnTableRow<SViewModelBindingDetail::FViewModelBindingId>::Construct(
+			FSuperRowType::FArguments()
+			.Padding(1.0f),
+			OwnerTableView
+		);
+	}
+
+	virtual TSharedRef<SWidget> GenerateWidgetForColumn(const FName& ColumnName) override
+	{
+		if (BindingId.ViewModel && BindingId.ViewModel->FieldBound.IsValidIndex(BindingId.Index))
+		{
+			if (ColumnName == ViewColumnName::FieldId)
+			{
+				return SNew(STextBlock)
+					.Text(FText::FromName(BindingId.ViewModel->FieldBound[BindingId.Index].FieldId.GetFieldName()));
+			}
+			if (ColumnName == ViewColumnName::ObjectName)
+			{
+				return SNew(STextBlock)
+					.Text(FText::FromString(BindingId.ViewModel->FullName));
+			}
+		}
+
+		check(false);
+		return SNew(STextBlock);
+	}
+
+private:
+	SViewModelBindingDetail::FViewModelBindingId BindingId;
+};
+} //namespace Private
+
+
+void SViewModelBindingDetail::Construct(const FArguments& InArgs)
+{
+	TSharedRef<SHeaderRow> HeaderRow = SNew(SHeaderRow)
+		+ SHeaderRow::Column(Private::ViewColumnName::FieldId)
+		.FillWidth(0.25f)
+		.DefaultLabel(LOCTEXT("FieldIdColumnHeaderName", "Field"))
+		+ SHeaderRow::Column(Private::ViewColumnName::ObjectName)
+		.FillWidth(0.75f)
+		.DefaultLabel(LOCTEXT("ObjectColumnHeaderName", "Object"));
+
+	SAssignNew(ListView, SListView<FViewModelBindingId>)
+		.ListItemsSource(&Entries)
+		//.SelectionMode(ESelectionMode::Single)
+		.OnGenerateRow(this, &SViewModelBindingDetail::HandleGenerateWidgetForItem)
+		.OnSelectionChanged(this, &SViewModelBindingDetail::HandleListSelectionChanged)
+		//.OnMouseButtonDoubleClick(this, &SSceneOutliner::OnOutlinerTreeDoubleClick)
+		//.OnContextMenuOpening(this, &SLiveLinkClientPanel::OnOpenVirtualSubjectContextMenu)
+		.HeaderRow(HeaderRow);
+
+	ChildSlot
+	[
+		ListView.ToSharedRef()
+	];
+}
+
+
+void SViewModelBindingDetail::SetViewModels(TArray<TSharedPtr<FMVVMViewModelDebugEntry>> InViewModel)
+{
+	if (ViewModels != InViewModel)
+	{
+		ViewModels = InViewModel;
+		Entries.Reset();
+
+		FViewModelBindingId NewItem;
+		for (TSharedPtr<FMVVMViewModelDebugEntry> Entry : ViewModels)
+		{
+			NewItem.ViewModel = Entry;
+			for (int32 Index = 0; Index < Entry->FieldBound.Num(); ++Index)
+			{
+				NewItem.Index = Index;
+				Entries.Add(NewItem);
+			}
+		}
+
+		ListView->RebuildList(); //RequestListRefresh
+	}
+}
+
+
+TSharedRef<ITableRow> SViewModelBindingDetail::HandleGenerateWidgetForItem(FViewModelBindingId Entry, const TSharedRef<STableViewBase>& OwnerTable)
+{
+	return SNew(Private::SViewModelBindingDetailItem, OwnerTable)
+		.BindingId(Entry);
+}
+
+
+void SViewModelBindingDetail::HandleListSelectionChanged(FViewModelBindingId TreeItem, ESelectInfo::Type SelectInfo)
+{}
+
+} //namespace
+
+#undef LOCTEXT_NAMESPACE
\ No newline at end of file
diff --git a/Engine/Plugins/Runtime/ModelViewViewModel/Source/ModelViewViewModelDebuggerEditor/Private/Widgets/SViewModelBindingDetail.h b/Engine/Plugins/Runtime/ModelViewViewModel/Source/ModelViewViewModelDebuggerEditor/Private/Widgets/SViewModelBindingDetail.h
new file mode 100644
--- /dev/null
+++ b/Engine/Plugins/Runtime/ModelViewViewModel/Source/ModelViewViewModelDebuggerEditor/Private/Widgets/SViewModelBindingDetail.h
@@ -0,0 +1,115 @@
+// Copyright Epic Games, Inc. All Rights Reserved.
+#pragma once
+
+#include "CoreMinimal.h"
+#include "Widgets/SCompoundWidget.h"
+
+#include "Framework/Views/TableViewTypeTraits.h"
+
+#include "MVVMDebugView.h"
+#include "MVVMDebugViewModel.h"
+
+template<typename ItemType>
+class SListView;
+class ITableRow;
+class STableViewBase;
+
+namespace UE::MVVM
+{
+
+class SViewModelBindingDetail : public SCompoundWidget
+{
+public:
+	struct FViewModelBindingId
+	{
+		TSharedPtr<FMVVMViewModelDebugEntry> ViewModel;
+		int32 Index = INDEX_NONE;
+
+		bool operator==(const FViewModelBindingId& Other) const
+		{
+			return Index == Other.Index && ViewModel == Other.ViewModel;
+		}
+
+		friend uint32 GetTypeHash(const FViewModelBindingId& Value)
+		{
+			return HashCombine(GetTypeHash(Value.ViewModel), GetTypeHash(Value.Index));
+		}
+	};
+
+public:
+	SLATE_BEGIN_ARGS(SViewModelBindingDetail) { }
+	SLATE_END_ARGS()
+
+public:
+	void Construct(const FArguments& InArgs);
+
+	void SetViewModels(TArray<TSharedPtr<FMVVMViewModelDebugEntry>> ViewModel);
+
+private:
+	TSharedRef<ITableRow> HandleGenerateWidgetForItem(FViewModelBindingId Entry, const TSharedRef<STableViewBase>& OwnerTable);
+	void HandleListSelectionChanged(FViewModelBindingId TreeItem, ESelectInfo::Type SelectInfo);
+
+private:
+	TArray<TSharedPtr<FMVVMViewModelDebugEntry>> ViewModels;
+	TArray<FViewModelBindingId> Entries;
+	TSharedPtr<SListView<FViewModelBindingId>> ListView;
+};
+
+} //namespace
+
+
+template <>
+struct TListTypeTraits<UE::MVVM::SViewModelBindingDetail::SViewModelBindingDetail::FViewModelBindingId>
+{
+public:
+	typedef UE::MVVM::SViewModelBindingDetail::FViewModelBindingId NullableType;
+
+	using MapKeyFuncs = TDefaultMapHashableKeyFuncs<UE::MVVM::SViewModelBindingDetail::SViewModelBindingDetail::FViewModelBindingId, TSharedRef<ITableRow>, false>;
+	using MapKeyFuncsSparse = TDefaultMapHashableKeyFuncs<UE::MVVM::SViewModelBindingDetail::SViewModelBindingDetail::FViewModelBindingId, FSparseItemInfo, false>;
+	using SetKeyFuncs = DefaultKeyFuncs<UE::MVVM::SViewModelBindingDetail::SViewModelBindingDetail::FViewModelBindingId>;
+
+	template<typename U>
+	static void AddReferencedObjects(FReferenceCollector& Collector
+		, TArray<UE::MVVM::SViewModelBindingDetail::FViewModelBindingId>& ItemsWithGeneratedWidgets
+		, TSet<UE::MVVM::SViewModelBindingDetail::FViewModelBindingId>& SelectedItems
+		, TMap<const U*, UE::MVVM::SViewModelBindingDetail::FViewModelBindingId>& WidgetToItemMap)
+	{
+	}
+
+	static bool IsPtrValid(const UE::MVVM::SViewModelBindingDetail::FViewModelBindingId& InValue)
+	{
+		return InValue.ViewModel.IsValid() && InValue.Index != INDEX_NONE;
+	}
+
+	static void ResetPtr(UE::MVVM::SViewModelBindingDetail::FViewModelBindingId& InValue)
+	{
+		InValue.ViewModel.Reset();
+		InValue.Index = INDEX_NONE;
+	}
+
+	static UE::MVVM::SViewModelBindingDetail::FViewModelBindingId MakeNullPtr()
+	{
+		return UE::MVVM::SViewModelBindingDetail::FViewModelBindingId();
+	}
+
+	static UE::MVVM::SViewModelBindingDetail::FViewModelBindingId NullableItemTypeConvertToItemType(const UE::MVVM::SViewModelBindingDetail::FViewModelBindingId& InValue)
+	{
+		return InValue;
+	}
+
+	static FString DebugDump(UE::MVVM::SViewModelBindingDetail::FViewModelBindingId InValue)
+	{
+		return FString();
+	}
+
+	class SerializerType {};
+};
+
+template <>
+struct TIsValidListItem<UE::MVVM::SViewModelBindingDetail::SViewModelBindingDetail::FViewModelBindingId>
+{
+	enum
+	{
+		Value = true
+	};
+};
diff --git a/Engine/Plugins/Runtime/ModelViewViewModel/Source/ModelViewViewModelDebuggerEditor/Private/Widgets/SViewModelSelection.cpp b/Engine/Plugins/Runtime/ModelViewViewModel/Source/ModelViewViewModelDebuggerEditor/Private/Widgets/SViewModelSelection.cpp
new file mode 100644
--- /dev/null
+++ b/Engine/Plugins/Runtime/ModelViewViewModel/Source/ModelViewViewModelDebuggerEditor/Private/Widgets/SViewModelSelection.cpp
@@ -0,0 +1,130 @@
+// Copyright Epic Games, Inc. All Rights Reserved.
+
+#include "Widgets/SViewModelSelection.h"
+
+#include "MVVMDebugSnapshot.h"
+#include "MVVMDebugView.h"
+#include "MVVMDebugViewClass.h"
+
+#include "Misc/TextFilter.h"
+
+#include "Widgets/SBoxPanel.h"
+#include "Widgets/Text/STextBlock.h"
+#include "Widgets/Views/SListView.h"
+
+#define LOCTEXT_NAMESPACE "MVVMDebuggerViewModelSelection"
+
+namespace UE::MVVM
+{
+namespace Private
+{
+namespace ViewModelColumnName
+{
+FLazyName Name = "Name";
+FLazyName Class = "Class";
+}
+
+
+void SViewModelSelectionItem::Construct(const FArguments& Args, const TSharedRef<STableViewBase>& OwnerTableView)
+{
+	Entry = Args._Entry;
+
+	SMultiColumnTableRow<TSharedPtr<FMVVMViewModelDebugEntry>>::Construct(
+		FSuperRowType::FArguments()
+		.Padding(1.0f),
+		OwnerTableView
+	);
+}
+
+TSharedRef<SWidget> SViewModelSelectionItem::GenerateWidgetForColumn(const FName& ColumnName)
+{
+	if (ColumnName == ViewModelColumnName::Name)
+	{
+		return SNew(STextBlock)
+			.Text(FText::FromName(Entry->Name));
+	}
+	if (ColumnName == ViewModelColumnName::Class)
+	{
+		return SNew(STextBlock)
+			.Text(FText::FromString(Entry->ViewModelAsset.GetObjectPathString()));
+	}
+
+	check(false);
+	return SNew(STextBlock);
+}
+
+} //namespace Private
+
+
+void SViewModelSelection::Construct(const FArguments& InArgs)
+{
+	Super::Construct(Super::FArguments()
+		.OnSelectionChanged(InArgs._OnSelectionChanged)
+	);
+}
+
+
+void SViewModelSelection::SetSelection(FGuid Id)
+{
+	bool bClear = true;
+	if (GetSnapshot())
+	{
+		if (TSharedPtr<FMVVMViewModelDebugEntry> Found = GetSnapshot()->FindViewModel(Id))
+		{
+			ListView->SetSelection(Found);
+			bClear = false;
+		}
+	}
+	if (bClear)
+	{
+		ListView->ClearSelection();
+	}
+}
+
+
+TArray<TSharedPtr<FMVVMViewModelDebugEntry>> SViewModelSelection::UpdateSourceImpl(TSharedPtr<TTextFilter<TSharedPtr<FMVVMViewModelDebugEntry>>> InTextFilter)
+{
+	TArray<TSharedPtr<FMVVMViewModelDebugEntry>> Result;
+	if (GetSnapshot())
+	{
+		if (!InTextFilter->GetRawFilterText().IsEmpty())
+		{
+			Result.Reserve(Snapshot->GetViewModels().Num());
+			for (const TSharedPtr<FMVVMViewModelDebugEntry>& Item : GetSnapshot()->GetViewModels())
+			{
+				if (InTextFilter->PassesFilter(Item))
+				{
+					Result.Add(Item);
+				}
+			}
+		}
+		else
+		{
+			Result = Snapshot->GetViewModels();
+		}
+	}
+	return Result;
+}
+
+
+TSharedRef<SHeaderRow> SViewModelSelection::BuildHeaderRowImpl() const
+{
+	return SNew(SHeaderRow)
+		+ SHeaderRow::Column(Private::ViewModelColumnName::Name)
+		.FillWidth(0.75f)
+		.DefaultLabel(LOCTEXT("ViewmodelColumnHeaderName", "Name"))
+		+ SHeaderRow::Column(Private::ViewModelColumnName::Class)
+		.FillWidth(0.25f)
+		.DefaultLabel(LOCTEXT("ClassColumnHeaderName", "Class"));
+}
+
+
+void SViewModelSelection::GatherFilterStringsImpl(TSharedPtr<FMVVMViewModelDebugEntry> Item, TArray<FString>& OutStrings) const
+{
+	OutStrings.Add(Item->FullName);
+	OutStrings.Add(Item->ViewModelAsset.GetObjectPathString());
+}
+
+} //namespace
+
+#undef LOCTEXT_NAMESPACE
\ No newline at end of file
diff --git a/Engine/Plugins/Runtime/ModelViewViewModel/Source/ModelViewViewModelDebuggerEditor/Private/Widgets/SViewModelSelection.h b/Engine/Plugins/Runtime/ModelViewViewModel/Source/ModelViewViewModelDebuggerEditor/Private/Widgets/SViewModelSelection.h
new file mode 100644
--- /dev/null
+++ b/Engine/Plugins/Runtime/ModelViewViewModel/Source/ModelViewViewModelDebuggerEditor/Private/Widgets/SViewModelSelection.h
@@ -0,0 +1,50 @@
+// Copyright Epic Games, Inc. All Rights Reserved.
+#pragma once
+
+#include "CoreMinimal.h"
+#include "Widgets/SSelectionBase.h"
+#include "Widgets/Views/STableRow.h"
+
+#include "MVVMDebugViewModel.h"
+
+
+namespace UE::MVVM
+{
+
+namespace Private
+{
+class SViewModelSelectionItem : public SMultiColumnTableRow<TSharedPtr<FMVVMViewModelDebugEntry>>
+{
+	SLATE_BEGIN_ARGS(SViewModelSelectionItem) {}
+	SLATE_ARGUMENT(TSharedPtr<FMVVMViewModelDebugEntry>, Entry)
+	SLATE_END_ARGS()
+
+	void Construct(const FArguments& Args, const TSharedRef<STableViewBase>& OwnerTableView);
+	virtual TSharedRef<SWidget> GenerateWidgetForColumn(const FName& ColumnName) override;
+
+private:
+	TSharedPtr<FMVVMViewModelDebugEntry> Entry;
+};
+}
+
+class SViewModelSelection : public SSelectionBase<TSharedPtr<FMVVMViewModelDebugEntry>, Private::SViewModelSelectionItem>
+{
+private:
+	using Super = SSelectionBase<TSharedPtr<FMVVMViewModelDebugEntry>, Private::SViewModelSelectionItem>;
+
+public:
+	SLATE_BEGIN_ARGS(SViewModelSelection) { }
+	SLATE_EVENT(FSimpleDelegate, OnSelectionChanged)
+	SLATE_END_ARGS()
+
+public:
+	void Construct(const FArguments& InArgs);
+	void SetSelection(FGuid);
+
+protected:
+	virtual void GatherFilterStringsImpl(TSharedPtr<FMVVMViewModelDebugEntry> Item, TArray<FString>& OutStrings) const override;
+	virtual TSharedRef<SHeaderRow> BuildHeaderRowImpl() const override;
+	virtual TArray<TSharedPtr<FMVVMViewModelDebugEntry>> UpdateSourceImpl(TSharedPtr<TTextFilter<TSharedPtr<FMVVMViewModelDebugEntry>>> TextFilter) override;
+};
+
+} //namespace
diff --git a/Engine/Plugins/Runtime/ModelViewViewModel/Source/ModelViewViewModelDebuggerEditor/Private/Widgets/SViewSelection.cpp b/Engine/Plugins/Runtime/ModelViewViewModel/Source/ModelViewViewModelDebuggerEditor/Private/Widgets/SViewSelection.cpp
new file mode 100644
--- /dev/null
+++ b/Engine/Plugins/Runtime/ModelViewViewModel/Source/ModelViewViewModelDebuggerEditor/Private/Widgets/SViewSelection.cpp
@@ -0,0 +1,143 @@
+// Copyright Epic Games, Inc. All Rights Reserved.
+
+#include "Widgets/SViewSelection.h"
+
+#include "MVVMDebugSnapshot.h"
+#include "MVVMDebugView.h"
+#include "MVVMDebugViewClass.h"
+
+#include "Misc/TextFilter.h"
+
+#include "Widgets/SBoxPanel.h"
+#include "Widgets/Text/STextBlock.h"
+#include "Widgets/Views/SListView.h"
+
+//ToolWidgets
+
+#define LOCTEXT_NAMESPACE "MVVMDebuggerViewSelection"
+
+namespace UE::MVVM
+{
+namespace Private
+{
+namespace ViewColumnName
+{
+FLazyName UserWidget = "UserWidget";
+FLazyName LocalPlayer = "LocalPlayer";
+FLazyName WorldName = "WorldName";
+}
+
+void SViewSelectionItem::Construct(const FArguments& Args, const TSharedRef<STableViewBase>& OwnerTableView)
+{
+	Entry = Args._Entry;
+
+	SMultiColumnTableRow<TSharedPtr<FMVVMViewDebugEntry>>::Construct(
+		FSuperRowType::FArguments()
+		.Padding(1.0f),
+		OwnerTableView
+	);
+}
+
+TSharedRef<SWidget> SViewSelectionItem::GenerateWidgetForColumn(const FName& ColumnName)
+{
+	if (ColumnName == ViewColumnName::UserWidget)
+	{
+		return SNew(STextBlock)
+			.Text(FText::FromName(Entry->UserWidgetInstanceName));
+	}
+	if (ColumnName == ViewColumnName::LocalPlayer)
+	{
+		return SNew(STextBlock)
+			.Text(FText::FromName(Entry->LocalPlayerName));
+	}
+	if (ColumnName == ViewColumnName::WorldName)
+	{
+		return SNew(STextBlock)
+			.Text(FText::FromName(Entry->WorldName));
+	}
+
+	check(false);
+	return SNew(STextBlock);
+}
+
+} //namespace Private
+
+
+void SViewSelection::Construct(const FArguments& InArgs)
+{
+	Super::Construct(Super::FArguments()
+		.OnSelectionChanged(InArgs._OnSelectionChanged)
+		);
+}
+
+
+void SViewSelection::SetSelection(FGuid Id)
+{
+	bool bClear = true;
+	if (GetSnapshot())
+	{
+		if (TSharedPtr<FMVVMViewDebugEntry> Found = GetSnapshot()->FindView(Id))
+		{
+			ListView->SetSelection(Found);
+			bClear= false;
+		}
+	}
+	if (bClear)
+	{
+		ListView->ClearSelection();
+	}
+}
+
+
+TArray<TSharedPtr<FMVVMViewDebugEntry>> SViewSelection::UpdateSourceImpl(TSharedPtr<TTextFilter<TSharedPtr<FMVVMViewDebugEntry>>> InTextFilter)
+{
+	TArray<TSharedPtr<FMVVMViewDebugEntry>> Result;
+	if (GetSnapshot())
+	{
+		if (!InTextFilter->GetRawFilterText().IsEmpty())
+		{
+			Result.Reserve(Snapshot->GetViews().Num());
+			for (const TSharedPtr<FMVVMViewDebugEntry>& Item : GetSnapshot()->GetViews())
+			{
+				if (InTextFilter->PassesFilter(Item))
+				{
+					Result.Add(Item);
+				}
+			}
+		}
+		else
+		{
+			Result = Snapshot->GetViews();
+		}
+	}
+	return Result;
+}
+
+
+TSharedRef<SHeaderRow> SViewSelection::BuildHeaderRowImpl() const
+{
+	return SNew(SHeaderRow)
+		+ SHeaderRow::Column(Private::ViewColumnName::UserWidget)
+		.FillWidth(0.5f)
+		.DefaultLabel(LOCTEXT("UserWidgetColumnHeaderName", "User Widget"))
+		+ SHeaderRow::Column(Private::ViewColumnName::LocalPlayer)
+		.FillWidth(0.25f)
+		.DefaultLabel(LOCTEXT("LocalPlayerColumnHeaderName", "Player"))
+		+ SHeaderRow::Column(Private::ViewColumnName::WorldName)
+		.FillWidth(0.25f)
+		.DefaultLabel(LOCTEXT("WorldColumnHeaderName", "World"));
+}
+
+
+void SViewSelection::GatherFilterStringsImpl(TSharedPtr<FMVVMViewDebugEntry> Item, TArray<FString>& OutStrings) const
+{
+	OutStrings.Add(Item->UserWidgetInstanceName.ToString());
+	OutStrings.Add(Item->LocalPlayerName.ToString());
+	//OutStrings.Add(Item->UserWidgetAsset.ObjectPath.ToString());
+	OutStrings.Add(Item->UserWidgetAsset.AssetClassPath.ToString());
+	OutStrings.Add(Item->ViewClassDebugId.ToString());
+}
+
+} //namespace
+
+#undef LOCTEXT_NAMESPACE
\ No newline at end of file
diff --git a/Engine/Plugins/Runtime/ModelViewViewModel/Source/ModelViewViewModelDebuggerEditor/Private/Widgets/SViewSelection.h b/Engine/Plugins/Runtime/ModelViewViewModel/Source/ModelViewViewModelDebuggerEditor/Private/Widgets/SViewSelection.h
new file mode 100644
--- /dev/null
+++ b/Engine/Plugins/Runtime/ModelViewViewModel/Source/ModelViewViewModelDebuggerEditor/Private/Widgets/SViewSelection.h
@@ -0,0 +1,49 @@
+// Copyright Epic Games, Inc. All Rights Reserved.
+#pragma once
+
+#include "CoreMinimal.h"
+#include "Widgets/SSelectionBase.h"
+#include "Widgets/Views/STableRow.h"
+
+#include "MVVMDebugView.h"
+
+namespace UE::MVVM
+{
+
+namespace Private
+{
+class SViewSelectionItem : public SMultiColumnTableRow<TSharedPtr<FMVVMViewDebugEntry>>
+{
+	SLATE_BEGIN_ARGS(SViewSelectionItem) {}
+	SLATE_ARGUMENT(TSharedPtr<FMVVMViewDebugEntry>, Entry)
+	SLATE_END_ARGS()
+
+	void Construct(const FArguments& Args, const TSharedRef<STableViewBase>& OwnerTableView);
+	TSharedRef<SWidget> GenerateWidgetForColumn(const FName& ColumnName);
+
+private:
+	TSharedPtr<FMVVMViewDebugEntry> Entry;
+};
+}
+
+class SViewSelection : public SSelectionBase<TSharedPtr<FMVVMViewDebugEntry>, Private::SViewSelectionItem>
+{
+private:
+	using Super = SSelectionBase<TSharedPtr<FMVVMViewDebugEntry>, Private::SViewSelectionItem>;
+
+public:
+	SLATE_BEGIN_ARGS(SViewSelection) { }
+	SLATE_EVENT(FSimpleDelegate, OnSelectionChanged)
+	SLATE_END_ARGS()
+
+public:
+	void Construct(const FArguments& InArgs);
+	void SetSelection(FGuid Id);
+
+protected:
+	virtual void GatherFilterStringsImpl(TSharedPtr<FMVVMViewDebugEntry> Item, TArray<FString>& OutStrings) const override;
+	virtual TSharedRef<SHeaderRow> BuildHeaderRowImpl() const override;
+	virtual TArray<TSharedPtr<FMVVMViewDebugEntry>> UpdateSourceImpl(TSharedPtr<TTextFilter<TSharedPtr<FMVVMViewDebugEntry>>> TextFilter) override;
+};
+
+} //namespace
diff --git a/Engine/Source/Runtime/FieldNotification/Private/FieldNotificationDelegate.cpp b/Engine/Source/Runtime/FieldNotification/Private/FieldNotificationDelegate.cpp
--- a/Engine/Source/Runtime/FieldNotification/Private/FieldNotificationDelegate.cpp
+++ b/Engine/Source/Runtime/FieldNotification/Private/FieldNotificationDelegate.cpp
@@ -506,4 +506,20 @@ void FFieldMulticastDelegate::ExecuteLockOperations()
 	}
 }
 
+
+TArray<FFieldMulticastDelegate::FDelegateView> FFieldMulticastDelegate::GetView() const
+{
+	TArray<FFieldMulticastDelegate::FDelegateView> Result;
+	Result.Reserve(Delegates.Num());
+	for (const FInvocationElement& Element : Delegates)
+	{
+		if (Element.Delegate.IsBound())
+		{
+			//Result.Emplace(Element.Key.Object.Get(), Element.Key.Id, Element.Delegate.TryGetBoundFunctionName());
+			Result.Emplace(Element.Key.Object.Get(), Element.Key.Id, Element.Key.DynamicName);
+		}
+	}
+	return Result;
+}
+
 } //namespace
\ No newline at end of file
diff --git a/Engine/Source/Runtime/FieldNotification/Public/FieldNotificationDelegate.h b/Engine/Source/Runtime/FieldNotification/Public/FieldNotificationDelegate.h
--- a/Engine/Source/Runtime/FieldNotification/Public/FieldNotificationDelegate.h
+++ b/Engine/Source/Runtime/FieldNotification/Public/FieldNotificationDelegate.h
@@ -103,6 +103,17 @@ public:
 
 	void Reset();
 
+	struct FDelegateView
+	{
+		FDelegateView(const UObject* Obj, FFieldId Id, FName Function)
+			: Object(Obj), Field(Id), FunctionName(Function)
+		{}
+		const UObject* Object;
+		FFieldId Field;
+		FName FunctionName;
+	};
+	TArray<FDelegateView> GetView() const;
+
 private:
 	int32 LowerBound(FFieldId InFieldId) const
 	{
