MVVM: Uses Construct instead of Initialize to create to load the Binding Library.

diff --git a/Engine/Plugins/Runtime/ModelViewViewModel/Source/ModelViewViewModel/Private/View/MVVMViewClass.cpp b/Engine/Plugins/Runtime/ModelViewViewModel/Source/ModelViewViewModel/Private/View/MVVMViewClass.cpp
--- a/Engine/Plugins/Runtime/ModelViewViewModel/Source/ModelViewViewModel/Private/View/MVVMViewClass.cpp
+++ b/Engine/Plugins/Runtime/ModelViewViewModel/Source/ModelViewViewModel/Private/View/MVVMViewClass.cpp
@@ -486,12 +486,11 @@ FString FMVVMViewClass_CompiledBinding::ToString(const FMVVMCompiledBindingLibra
 // 
 ///////////////////////////////////////////////////////////////////////
 
-void UMVVMViewClass::Initialize(UUserWidget* UserWidget)
+void UMVVMViewClass::Construct(UUserWidget* UserWidget)
 {
 	++ViewCounter;
 	ensure(ViewCounter <= std::numeric_limits<int32>::max());
 
-	ensure(UserWidget->GetExtension<UMVVMView>() == nullptr);
 	UMVVMView* View = UserWidget->AddExtension<UMVVMView>();
 	if (ensure(View))
 	{
diff --git a/Engine/Plugins/Runtime/ModelViewViewModel/Source/ModelViewViewModel/Public/View/MVVMViewClass.h b/Engine/Plugins/Runtime/ModelViewViewModel/Source/ModelViewViewModel/Public/View/MVVMViewClass.h
--- a/Engine/Plugins/Runtime/ModelViewViewModel/Source/ModelViewViewModel/Public/View/MVVMViewClass.h
+++ b/Engine/Plugins/Runtime/ModelViewViewModel/Source/ModelViewViewModel/Public/View/MVVMViewClass.h
@@ -321,7 +321,7 @@ class MODELVIEWVIEWMODEL_API UMVVMViewClass : public UWidgetBlueprintGeneratedCl
 
 public:
 	//~ Begin UWidgetBlueprintGeneratedClassExtension
-	virtual void Initialize(UUserWidget* UserWidget) override;
+	virtual void Construct(UUserWidget* UserWidget) override;
 	virtual void Destruct(UUserWidget* UserWidget) override;
 #if WITH_EDITOR
 	virtual void BeginDestroy() override;
