MVVM: Always create the View on the userwidget at initialization before the PreConstruct and Destruct function. Only load the binding library when the first widget that uses that view is constructed (by TakeWidget) and release the binding library when the last widget that uses that view is destroyed (RealseSlateResource).

diff --git a/Engine/Plugins/Runtime/ModelViewViewModel/Source/ModelViewViewModel/Private/View/MVVMViewClass.cpp b/Engine/Plugins/Runtime/ModelViewViewModel/Source/ModelViewViewModel/Private/View/MVVMViewClass.cpp
--- a/Engine/Plugins/Runtime/ModelViewViewModel/Source/ModelViewViewModel/Private/View/MVVMViewClass.cpp
+++ b/Engine/Plugins/Runtime/ModelViewViewModel/Source/ModelViewViewModel/Private/View/MVVMViewClass.cpp
@@ -486,20 +486,24 @@ FString FMVVMViewClass_CompiledBinding::ToString(const FMVVMCompiledBindingLibra
 // 
 ///////////////////////////////////////////////////////////////////////
 
+void UMVVMViewClass::Initialize(UUserWidget* UserWidget)
+{
+	ensure(UserWidget->GetExtension<UMVVMView>() == nullptr);
+	UMVVMView* View = UserWidget->AddExtension<UMVVMView>();
+	if (ensure(View))
+	{
+		View->ConstructView(this);
+	}
+}
+
 void UMVVMViewClass::Construct(UUserWidget* UserWidget)
 {
 	++ViewCounter;
 	ensure(ViewCounter <= std::numeric_limits<int32>::max());
 
-	UMVVMView* View = UserWidget->AddExtension<UMVVMView>();
-	if (ensure(View))
+	if (ViewCounter == 1)
 	{
-		if (ViewCounter == 1)
-		{
-			BindingLibrary.Load();
-		}
-
-		View->ConstructView(this);
+		BindingLibrary.Load();
 	}
 
 #if WITH_EDITOR
diff --git a/Engine/Plugins/Runtime/ModelViewViewModel/Source/ModelViewViewModel/Public/View/MVVMViewClass.h b/Engine/Plugins/Runtime/ModelViewViewModel/Source/ModelViewViewModel/Public/View/MVVMViewClass.h
--- a/Engine/Plugins/Runtime/ModelViewViewModel/Source/ModelViewViewModel/Public/View/MVVMViewClass.h
+++ b/Engine/Plugins/Runtime/ModelViewViewModel/Source/ModelViewViewModel/Public/View/MVVMViewClass.h
@@ -321,6 +321,7 @@ class MODELVIEWVIEWMODEL_API UMVVMViewClass : public UWidgetBlueprintGeneratedCl
 
 public:
 	//~ Begin UWidgetBlueprintGeneratedClassExtension
+	virtual void Initialize(UUserWidget* UserWidget) override;
 	virtual void Construct(UUserWidget* UserWidget) override;
 	virtual void Destruct(UUserWidget* UserWidget) override;
 #if WITH_EDITOR
