Temporary stopgap to prevent blueprint compilation from ALT in editor from trying to preload objects

#rnx

diff --git a/Engine/Plugins/Runtime/ModelViewViewModel/Source/ModelViewViewModelBlueprint/Private/MVVMWidgetBlueprintExtension_View.cpp b/Engine/Plugins/Runtime/ModelViewViewModel/Source/ModelViewViewModelBlueprint/Private/MVVMWidgetBlueprintExtension_View.cpp
--- a/Engine/Plugins/Runtime/ModelViewViewModel/Source/ModelViewViewModelBlueprint/Private/MVVMWidgetBlueprintExtension_View.cpp
+++ b/Engine/Plugins/Runtime/ModelViewViewModel/Source/ModelViewViewModelBlueprint/Private/MVVMWidgetBlueprintExtension_View.cpp
@@ -44,7 +44,7 @@ void UMVVMWidgetBlueprintExtension_View::PostLoad()
 
 void UMVVMWidgetBlueprintExtension_View::HandlePreloadObjectsForCompilation(UBlueprint* OwningBlueprint)
 {
-	if (BlueprintView)
+	if (IsInGameThread() && BlueprintView)
 	{
 		BlueprintView->ConditionalPostLoad();
 
