Fix  FMVVMViewBindingDebugEntry variable initialization.
#rnx

diff --git a/Engine/Plugins/Runtime/ModelViewViewModel/Source/ModelViewViewModelDebugger/Public/MVVMDebugViewClass.h b/Engine/Plugins/Runtime/ModelViewViewModel/Source/ModelViewViewModelDebugger/Public/MVVMDebugViewClass.h
--- a/Engine/Plugins/Runtime/ModelViewViewModel/Source/ModelViewViewModelDebugger/Public/MVVMDebugViewClass.h
+++ b/Engine/Plugins/Runtime/ModelViewViewModel/Source/ModelViewViewModelDebugger/Public/MVVMDebugViewClass.h
@@ -32,7 +32,7 @@ struct FMVVMViewBindingDebugEntry
 	FString ConversionFunctionFieldPath;
 
 	UPROPERTY(VisibleAnywhere, Category = "Viewmodel")
-	int32 CompiledBindingIndex;
+	int32 CompiledBindingIndex = INDEX_NONE;
 };
 
 
