Fix localization issue with MVVM Debugger
#rnx

diff --git a/Engine/Plugins/Runtime/ModelViewViewModel/Source/ModelViewViewModelDebuggerEditor/Private/Widgets/SMainDebugTab.cpp b/Engine/Plugins/Runtime/ModelViewViewModel/Source/ModelViewViewModelDebuggerEditor/Private/Widgets/SMainDebugTab.cpp
--- a/Engine/Plugins/Runtime/ModelViewViewModel/Source/ModelViewViewModelDebuggerEditor/Private/Widgets/SMainDebugTab.cpp
+++ b/Engine/Plugins/Runtime/ModelViewViewModel/Source/ModelViewViewModelDebuggerEditor/Private/Widgets/SMainDebugTab.cpp
@@ -317,7 +317,7 @@ TSharedRef<SDockTab> SMainDebug::SpawnSelectionTab(const FSpawnTabArgs& Args)
 TSharedRef<SDockTab> SMainDebug::SpawnBindingTab(const FSpawnTabArgs& Args)
 {
 	return SNew(SDockTab)
-		.Label(LOCTEXT("DetailTab", "Bindings"))
+		.Label(LOCTEXT("BindingTab", "Bindings"))
 		[
 			SAssignNew(ViewModelBindingDetail, SViewModelBindingDetail)
 		];
