MVVM: Add unittest for FieldNotificationDelegate. Fix bug with Remove that didn't return the correct state of the bindings.

diff --git a/Engine/Plugins/Runtime/ModelViewViewModel/Source/ModelViewViewModelEditor/Private/Tests/MVVMFieldValueChangedTest.cpp b/Engine/Plugins/Runtime/ModelViewViewModel/Source/ModelViewViewModelEditor/Private/Tests/MVVMFieldValueChangedTest.cpp
new file mode 100644
--- /dev/null
+++ b/Engine/Plugins/Runtime/ModelViewViewModel/Source/ModelViewViewModelEditor/Private/Tests/MVVMFieldValueChangedTest.cpp
@@ -0,0 +1,337 @@
+// Copyright Epic Games, Inc. All Rights Reserved.
+
+#include "Tests/MVVMFieldValueChangedTest.h"
+#include "Tests/MVVMFieldPathHelperTest.h"
+#include "Misc/AutomationTest.h"
+
+#include "UObject/GCObjectScopeGuard.h"
+
+#include UE_INLINE_GENERATED_CPP_BY_NAME(MVVMFieldValueChangedTest)
+
+
+#if WITH_AUTOMATION_WORKER
+
+#define LOCTEXT_NAMESPACE "MVVMBindingExecuteTest"
+
+IMPLEMENT_SIMPLE_AUTOMATION_TEST(FMVVMFieldValueChangedTest, "System.Plugins.MVVM.FieldValueChanged", EAutomationTestFlags::ApplicationContextMask | EAutomationTestFlags::EngineFilter)
+
+
+bool FMVVMFieldValueChangedTest::RunTest(const FString& Parameters)
+{
+	UMVVMFieldValueChangedTest* SourceObj = NewObject<UMVVMFieldValueChangedTest>();
+	TGCObjectScopeGuard<UObject> ScopeObject = TGCObjectScopeGuard<UObject>(SourceObj);
+
+	struct FCallbackHandler
+	{
+		void OnPropertyIntChanged(UObject* InObject, UE::FieldNotification::FFieldId InFieldId)
+		{
+			if (InObject != SourceObj)
+			{
+				Test->AddError(TEXT("Wrong SourceObject in OnPropertyIntChanged"));
+			}
+			if (InFieldId != UMVVMFieldValueChangedTest::FFieldNotificationClassDescriptor::PropertyInt)
+			{
+				Test->AddError(TEXT("Wrong FFieldId in OnPropertyIntChanged"));
+			}
+			++NumberPropertyIntChanged;
+		}
+
+		void OnPropertyFloatChanged(UObject* InObject, UE::FieldNotification::FFieldId InFieldId)
+		{
+			if (InObject != SourceObj)
+			{
+				Test->AddError(TEXT("Wrong SourceObject in OnPropertyFloatChanged"));
+			}
+			if (InFieldId != UMVVMFieldValueChangedTest::FFieldNotificationClassDescriptor::PropertyFloat)
+			{
+				Test->AddError(TEXT("Wrong FFieldId in OnPropertyFloatChanged"));
+			}
+			++NumberPropertyFloatChanged;
+		}
+
+		void OnFunctionIntChanged(UObject* InObject, UE::FieldNotification::FFieldId InFieldId)
+		{
+			if (InObject != SourceObj)
+			{
+				Test->AddError(TEXT("Wrong SourceObject in OnFunctionIntChanged"));
+			}
+			if (InFieldId != UMVVMFieldValueChangedTest::FFieldNotificationClassDescriptor::FunctionInt)
+			{
+				Test->AddError(TEXT("Wrong FFieldId in OnFunctionIntChanged"));
+			}
+			++NumberFunctionIntChanged;
+		}
+
+		void OnFunctionFloatChanged(UObject* InObject, UE::FieldNotification::FFieldId InFieldId)
+		{
+			if (InObject != SourceObj)
+			{
+				Test->AddError(TEXT("Wrong SourceObject in OnFunctionFloatChanged"));
+			}
+			if (InFieldId != UMVVMFieldValueChangedTest::FFieldNotificationClassDescriptor::FunctionFloat)
+			{
+				Test->AddError(TEXT("Wrong FFieldId in OnFunctionFloatChanged"));
+			}
+			++NumberFunctionFloatChanged;
+		}
+
+		void ResetValues()
+		{
+			NumberPropertyIntChanged = 0;
+			NumberPropertyFloatChanged = 0;
+			NumberFunctionIntChanged = 0;
+			NumberFunctionFloatChanged = 0;
+		}
+		FMVVMFieldValueChangedTest* Test = nullptr;
+		UMVVMFieldValueChangedTest* SourceObj = nullptr;
+		int32 NumberPropertyIntChanged = 0;
+		int32 NumberPropertyFloatChanged = 0;
+		int32 NumberFunctionIntChanged = 0;
+		int32 NumberFunctionFloatChanged = 0;
+	};
+
+	const int32 NumInstances = 3;
+	TArray<FCallbackHandler> CallbackInstances;
+	for (int32 Index = 0; Index < NumInstances; ++Index)
+	{
+		FCallbackHandler& Instance = CallbackInstances.AddDefaulted_GetRef();
+		Instance.Test = this;
+		Instance.SourceObj = SourceObj;
+	}
+
+	auto ResetAllValues = [&]()
+	{
+		for (FCallbackHandler& Instance : CallbackInstances)
+		{
+			Instance.ResetValues();
+		}
+	};
+
+	auto TestAllExpectedCountValues = [&](int32 InPropertyInt, int32 InPropertyFloat, int32 InFunctionInt, int32 InFunctionFloat)
+	{
+		int32 NumberPropertyIntChanged = 0;
+		int32 NumberPropertyFloatChanged = 0;
+		int32 NumberFunctionIntChanged = 0;
+		int32 NumberFunctionFloatChanged = 0;
+		for (FCallbackHandler& Instance : CallbackInstances)
+		{
+			NumberPropertyIntChanged += Instance.NumberPropertyIntChanged;
+			NumberPropertyFloatChanged += Instance.NumberPropertyFloatChanged;
+			NumberFunctionIntChanged += Instance.NumberFunctionIntChanged;
+			NumberFunctionFloatChanged += Instance.NumberFunctionFloatChanged;
+		}
+
+		if (NumberPropertyIntChanged != InPropertyInt)
+		{
+			AddError(TEXT("Wrong number of callbacked called with PropertyInt"));
+		}
+		if (NumberPropertyFloatChanged != InPropertyFloat)
+		{
+			AddError(TEXT("Wrong number of callbacked called with PropertyFloat"));
+		}
+		if (NumberFunctionIntChanged != InFunctionInt)
+		{
+			AddError(TEXT("Wrong number of callbacked called with FunctionInt"));
+		}
+		if (NumberFunctionFloatChanged != InFunctionFloat)
+		{
+			AddError(TEXT("Wrong number of callbacked called with FunctionFloat"));
+		}
+	};
+
+	auto TestExpectedCountValues = [&](int32 EntryIndex, int32 InPropertyInt, int32 InPropertyFloat, int32 InFunctionInt, int32 InFunctionFloat)
+	{
+		FCallbackHandler& Instance = CallbackInstances[EntryIndex];
+		if (Instance.NumberPropertyIntChanged != InPropertyInt)
+		{
+			AddError(TEXT("Wrong number of callbacked called with PropertyInt"));
+		}
+		if (Instance.NumberPropertyFloatChanged != InPropertyFloat)
+		{
+			AddError(TEXT("Wrong number of callbacked called with PropertyFloat"));
+		}
+		if (Instance.NumberFunctionIntChanged != InFunctionInt)
+		{
+			AddError(TEXT("Wrong number of callbacked called with FunctionInt"));
+		}
+		if (Instance.NumberFunctionFloatChanged != InFunctionFloat)
+		{
+			AddError(TEXT("Wrong number of callbacked called with FunctionFloat"));
+		}
+	};
+
+	auto TestBroadcastAll = [&]()
+	{
+		SourceObj->BroadcastFieldValueChanged(UMVVMFieldValueChangedTest::FFieldNotificationClassDescriptor::PropertyInt);
+		SourceObj->BroadcastFieldValueChanged(UMVVMFieldValueChangedTest::FFieldNotificationClassDescriptor::PropertyFloat);
+		SourceObj->BroadcastFieldValueChanged(UMVVMFieldValueChangedTest::FFieldNotificationClassDescriptor::FunctionInt);
+		SourceObj->BroadcastFieldValueChanged(UMVVMFieldValueChangedTest::FFieldNotificationClassDescriptor::FunctionFloat);
+	};
+
+
+	FDelegateHandle HandleA_0 = SourceObj->AddFieldValueChangedDelegate(UMVVMFieldValueChangedTest::FFieldNotificationClassDescriptor::PropertyInt, UMVVMFieldValueChangedTest::FFieldValueChangedDelegate::CreateRaw(&CallbackInstances[0], &FCallbackHandler::OnPropertyIntChanged));
+	FDelegateHandle HandleB_0 = SourceObj->AddFieldValueChangedDelegate(UMVVMFieldValueChangedTest::FFieldNotificationClassDescriptor::PropertyFloat, UMVVMFieldValueChangedTest::FFieldValueChangedDelegate::CreateRaw(&CallbackInstances[0], &FCallbackHandler::OnPropertyFloatChanged));
+	FDelegateHandle HandleC_0 = SourceObj->AddFieldValueChangedDelegate(UMVVMFieldValueChangedTest::FFieldNotificationClassDescriptor::FunctionInt, UMVVMFieldValueChangedTest::FFieldValueChangedDelegate::CreateRaw(&CallbackInstances[0], &FCallbackHandler::OnFunctionIntChanged));
+	FDelegateHandle HandleD_0 = SourceObj->AddFieldValueChangedDelegate(UMVVMFieldValueChangedTest::FFieldNotificationClassDescriptor::FunctionFloat, UMVVMFieldValueChangedTest::FFieldValueChangedDelegate::CreateRaw(&CallbackInstances[0], &FCallbackHandler::OnFunctionFloatChanged));
+
+	TestAllExpectedCountValues(0, 0, 0, 0);
+	ResetAllValues();
+
+	SourceObj->BroadcastFieldValueChanged(UMVVMFieldValueChangedTest::FFieldNotificationClassDescriptor::PropertyInt);
+	TestAllExpectedCountValues(1, 0, 0, 0);
+	TestExpectedCountValues(0, 1, 0, 0, 0);
+	ResetAllValues();
+
+	TestBroadcastAll();
+	TestAllExpectedCountValues(1, 1, 1, 1);
+	TestExpectedCountValues(0, 1, 1, 1, 1);
+	ResetAllValues();
+
+	SourceObj->RemoveFieldValueChangedDelegate(UMVVMFieldValueChangedTest::FFieldNotificationClassDescriptor::PropertyInt, HandleB_0);
+	TestBroadcastAll();
+	TestAllExpectedCountValues(1, 1, 1, 1);
+	TestExpectedCountValues(0, 1, 1, 1, 1);
+	ResetAllValues();
+	
+	SourceObj->RemoveFieldValueChangedDelegate(UMVVMFieldValueChangedTest::FFieldNotificationClassDescriptor::PropertyInt, HandleA_0);
+	TestBroadcastAll();
+	TestAllExpectedCountValues(0, 1, 1, 1);
+	TestExpectedCountValues(0, 0, 1, 1, 1);
+	ResetAllValues();
+
+	HandleA_0 = SourceObj->AddFieldValueChangedDelegate(UMVVMFieldValueChangedTest::FFieldNotificationClassDescriptor::PropertyInt, UMVVMFieldValueChangedTest::FFieldValueChangedDelegate::CreateRaw(&CallbackInstances[0], &FCallbackHandler::OnPropertyIntChanged));
+	TestBroadcastAll();
+	TestAllExpectedCountValues(1, 1, 1, 1);
+	TestExpectedCountValues(0, 1, 1, 1, 1);
+	ResetAllValues();
+
+	FDelegateHandle HandleA_1 = SourceObj->AddFieldValueChangedDelegate(UMVVMFieldValueChangedTest::FFieldNotificationClassDescriptor::PropertyInt, UMVVMFieldValueChangedTest::FFieldValueChangedDelegate::CreateRaw(&CallbackInstances[0], &FCallbackHandler::OnPropertyIntChanged));
+	FDelegateHandle HandleB_1 = SourceObj->AddFieldValueChangedDelegate(UMVVMFieldValueChangedTest::FFieldNotificationClassDescriptor::PropertyFloat, UMVVMFieldValueChangedTest::FFieldValueChangedDelegate::CreateRaw(&CallbackInstances[0], &FCallbackHandler::OnPropertyFloatChanged));
+	FDelegateHandle HandleC_1 = SourceObj->AddFieldValueChangedDelegate(UMVVMFieldValueChangedTest::FFieldNotificationClassDescriptor::FunctionInt, UMVVMFieldValueChangedTest::FFieldValueChangedDelegate::CreateRaw(&CallbackInstances[0], &FCallbackHandler::OnFunctionIntChanged));
+	FDelegateHandle HandleD_1 = SourceObj->AddFieldValueChangedDelegate(UMVVMFieldValueChangedTest::FFieldNotificationClassDescriptor::FunctionFloat, UMVVMFieldValueChangedTest::FFieldValueChangedDelegate::CreateRaw(&CallbackInstances[0], &FCallbackHandler::OnFunctionFloatChanged));
+
+	TestBroadcastAll();
+	TestAllExpectedCountValues(2, 2, 2, 2);
+	TestExpectedCountValues(0, 2, 2, 2, 2);
+	ResetAllValues();
+
+	FDelegateHandle HandleA_2 = SourceObj->AddFieldValueChangedDelegate(UMVVMFieldValueChangedTest::FFieldNotificationClassDescriptor::PropertyInt, UMVVMFieldValueChangedTest::FFieldValueChangedDelegate::CreateRaw(&CallbackInstances[1], &FCallbackHandler::OnPropertyIntChanged));
+	FDelegateHandle HandleB_2 = SourceObj->AddFieldValueChangedDelegate(UMVVMFieldValueChangedTest::FFieldNotificationClassDescriptor::PropertyFloat, UMVVMFieldValueChangedTest::FFieldValueChangedDelegate::CreateRaw(&CallbackInstances[1], &FCallbackHandler::OnPropertyFloatChanged));
+	FDelegateHandle HandleC_2 = SourceObj->AddFieldValueChangedDelegate(UMVVMFieldValueChangedTest::FFieldNotificationClassDescriptor::FunctionInt, UMVVMFieldValueChangedTest::FFieldValueChangedDelegate::CreateRaw(&CallbackInstances[1], &FCallbackHandler::OnFunctionIntChanged));
+	FDelegateHandle HandleD_2 = SourceObj->AddFieldValueChangedDelegate(UMVVMFieldValueChangedTest::FFieldNotificationClassDescriptor::FunctionFloat, UMVVMFieldValueChangedTest::FFieldValueChangedDelegate::CreateRaw(&CallbackInstances[1], &FCallbackHandler::OnFunctionFloatChanged));
+
+	TestBroadcastAll();
+	TestAllExpectedCountValues(3, 3, 3, 3);
+	TestExpectedCountValues(0, 2, 2, 2, 2);
+	ResetAllValues();
+
+	SourceObj->RemoveFieldValueChangedDelegate(UMVVMFieldValueChangedTest::FFieldNotificationClassDescriptor::FunctionInt, HandleA_0);
+	TestBroadcastAll();
+	TestAllExpectedCountValues(3, 3, 3, 3);
+	TestExpectedCountValues(0, 2, 2, 2, 2);
+	ResetAllValues();
+
+	SourceObj->RemoveFieldValueChangedDelegate(UMVVMFieldValueChangedTest::FFieldNotificationClassDescriptor::FunctionInt, HandleC_0);
+	TestBroadcastAll();
+	TestAllExpectedCountValues(3, 3, 2, 3);
+	TestExpectedCountValues(0, 2, 2, 1, 2);
+	ResetAllValues();
+
+	SourceObj->RemoveFieldValueChangedDelegate(UMVVMFieldValueChangedTest::FFieldNotificationClassDescriptor::FunctionInt, HandleC_0);
+	TestBroadcastAll();
+	TestAllExpectedCountValues(3, 3, 2, 3);
+	TestExpectedCountValues(0, 2, 2, 1, 2);
+	ResetAllValues();
+
+	SourceObj->RemoveFieldValueChangedDelegate(UMVVMFieldValueChangedTest::FFieldNotificationClassDescriptor::PropertyInt, HandleA_0);
+	TestBroadcastAll();
+	TestAllExpectedCountValues(2, 3, 2, 3);
+	TestExpectedCountValues(0, 1, 2, 1, 2);
+	ResetAllValues();
+
+	SourceObj->RemoveAllFieldValueChangedDelegates(&CallbackInstances[0]);
+	TestBroadcastAll();
+	TestAllExpectedCountValues(1, 1, 1, 1);
+	TestExpectedCountValues(0, 0, 0, 0, 0);
+	ResetAllValues();
+
+	HandleA_0 = SourceObj->AddFieldValueChangedDelegate(UMVVMFieldValueChangedTest::FFieldNotificationClassDescriptor::PropertyInt, UMVVMFieldValueChangedTest::FFieldValueChangedDelegate::CreateRaw(&CallbackInstances[0], &FCallbackHandler::OnPropertyIntChanged));
+	HandleB_0 = SourceObj->AddFieldValueChangedDelegate(UMVVMFieldValueChangedTest::FFieldNotificationClassDescriptor::PropertyFloat, UMVVMFieldValueChangedTest::FFieldValueChangedDelegate::CreateRaw(&CallbackInstances[0], &FCallbackHandler::OnPropertyFloatChanged));
+	HandleC_0 = SourceObj->AddFieldValueChangedDelegate(UMVVMFieldValueChangedTest::FFieldNotificationClassDescriptor::FunctionInt, UMVVMFieldValueChangedTest::FFieldValueChangedDelegate::CreateRaw(&CallbackInstances[0], &FCallbackHandler::OnFunctionIntChanged));
+	HandleD_0 = SourceObj->AddFieldValueChangedDelegate(UMVVMFieldValueChangedTest::FFieldNotificationClassDescriptor::FunctionFloat, UMVVMFieldValueChangedTest::FFieldValueChangedDelegate::CreateRaw(&CallbackInstances[0], &FCallbackHandler::OnFunctionFloatChanged));
+
+	FDelegateHandle HandleA_3 = SourceObj->AddFieldValueChangedDelegate(UMVVMFieldValueChangedTest::FFieldNotificationClassDescriptor::PropertyInt, UMVVMFieldValueChangedTest::FFieldValueChangedDelegate::CreateRaw(&CallbackInstances[2], &FCallbackHandler::OnPropertyIntChanged));
+	FDelegateHandle HandleB_3 = SourceObj->AddFieldValueChangedDelegate(UMVVMFieldValueChangedTest::FFieldNotificationClassDescriptor::PropertyFloat, UMVVMFieldValueChangedTest::FFieldValueChangedDelegate::CreateRaw(&CallbackInstances[2], &FCallbackHandler::OnPropertyFloatChanged));
+	FDelegateHandle HandleC_3 = SourceObj->AddFieldValueChangedDelegate(UMVVMFieldValueChangedTest::FFieldNotificationClassDescriptor::FunctionInt, UMVVMFieldValueChangedTest::FFieldValueChangedDelegate::CreateRaw(&CallbackInstances[2], &FCallbackHandler::OnFunctionIntChanged));
+	FDelegateHandle HandleD_3 = SourceObj->AddFieldValueChangedDelegate(UMVVMFieldValueChangedTest::FFieldNotificationClassDescriptor::FunctionFloat, UMVVMFieldValueChangedTest::FFieldValueChangedDelegate::CreateRaw(&CallbackInstances[2], &FCallbackHandler::OnFunctionFloatChanged));
+
+	TestBroadcastAll();
+	TestAllExpectedCountValues(3, 3, 3, 3);
+	TestExpectedCountValues(0, 1, 1, 1, 1);
+	TestExpectedCountValues(1, 1, 1, 1, 1);
+	TestExpectedCountValues(2, 1, 1, 1, 1);
+	ResetAllValues();
+
+	SourceObj->BroadcastFieldValueChanged(UMVVMFieldValueChangedTest::FFieldNotificationClassDescriptor::PropertyInt);
+	TestAllExpectedCountValues(3, 0, 0, 0);
+	TestExpectedCountValues(0, 1, 0, 0, 0);
+	TestExpectedCountValues(1, 1, 0, 0, 0);
+	TestExpectedCountValues(2, 1, 0, 0, 0);
+	ResetAllValues();
+
+	SourceObj->RemoveAllFieldValueChangedDelegates(&CallbackInstances[0]);
+	TestBroadcastAll();
+	TestAllExpectedCountValues(2, 2, 2, 2);
+	TestExpectedCountValues(0, 0, 0, 0, 0);
+	TestExpectedCountValues(1, 1, 1, 1, 1);
+	ResetAllValues();
+
+	SourceObj->RemoveAllFieldValueChangedDelegates(&CallbackInstances[2]);
+	TestBroadcastAll();
+	TestAllExpectedCountValues(1, 1, 1, 1);
+	TestExpectedCountValues(1, 1, 1, 1, 1);
+	ResetAllValues();
+
+	SourceObj->RemoveAllFieldValueChangedDelegates(&CallbackInstances[0]);
+	TestBroadcastAll();
+	TestAllExpectedCountValues(1, 1, 1, 1);
+	TestExpectedCountValues(1, 1, 1, 1, 1);
+	ResetAllValues();
+	
+	SourceObj->RemoveAllFieldValueChangedDelegates(&CallbackInstances[1]);
+	TestBroadcastAll();
+	TestAllExpectedCountValues(0, 0, 0, 0);
+	ResetAllValues();
+
+	HandleA_0 = SourceObj->AddFieldValueChangedDelegate(UMVVMFieldValueChangedTest::FFieldNotificationClassDescriptor::PropertyInt, UMVVMFieldValueChangedTest::FFieldValueChangedDelegate::CreateRaw(&CallbackInstances[0], &FCallbackHandler::OnPropertyIntChanged));
+	HandleB_0 = SourceObj->AddFieldValueChangedDelegate(UMVVMFieldValueChangedTest::FFieldNotificationClassDescriptor::PropertyFloat, UMVVMFieldValueChangedTest::FFieldValueChangedDelegate::CreateRaw(&CallbackInstances[0], &FCallbackHandler::OnPropertyFloatChanged));
+	HandleC_0 = SourceObj->AddFieldValueChangedDelegate(UMVVMFieldValueChangedTest::FFieldNotificationClassDescriptor::FunctionInt, UMVVMFieldValueChangedTest::FFieldValueChangedDelegate::CreateRaw(&CallbackInstances[0], &FCallbackHandler::OnFunctionIntChanged));
+	HandleD_0 = SourceObj->AddFieldValueChangedDelegate(UMVVMFieldValueChangedTest::FFieldNotificationClassDescriptor::FunctionFloat, UMVVMFieldValueChangedTest::FFieldValueChangedDelegate::CreateRaw(&CallbackInstances[0], &FCallbackHandler::OnFunctionFloatChanged));
+
+	HandleA_1 = SourceObj->AddFieldValueChangedDelegate(UMVVMFieldValueChangedTest::FFieldNotificationClassDescriptor::PropertyInt, UMVVMFieldValueChangedTest::FFieldValueChangedDelegate::CreateRaw(&CallbackInstances[0], &FCallbackHandler::OnPropertyIntChanged));
+	HandleB_1 = SourceObj->AddFieldValueChangedDelegate(UMVVMFieldValueChangedTest::FFieldNotificationClassDescriptor::PropertyFloat, UMVVMFieldValueChangedTest::FFieldValueChangedDelegate::CreateRaw(&CallbackInstances[0], &FCallbackHandler::OnPropertyFloatChanged));
+	HandleC_1 = SourceObj->AddFieldValueChangedDelegate(UMVVMFieldValueChangedTest::FFieldNotificationClassDescriptor::FunctionInt, UMVVMFieldValueChangedTest::FFieldValueChangedDelegate::CreateRaw(&CallbackInstances[0], &FCallbackHandler::OnFunctionIntChanged));
+	HandleD_1 = SourceObj->AddFieldValueChangedDelegate(UMVVMFieldValueChangedTest::FFieldNotificationClassDescriptor::FunctionFloat, UMVVMFieldValueChangedTest::FFieldValueChangedDelegate::CreateRaw(&CallbackInstances[0], &FCallbackHandler::OnFunctionFloatChanged));
+
+	HandleA_2 = SourceObj->AddFieldValueChangedDelegate(UMVVMFieldValueChangedTest::FFieldNotificationClassDescriptor::PropertyInt, UMVVMFieldValueChangedTest::FFieldValueChangedDelegate::CreateRaw(&CallbackInstances[1], &FCallbackHandler::OnPropertyIntChanged));
+	HandleB_2 = SourceObj->AddFieldValueChangedDelegate(UMVVMFieldValueChangedTest::FFieldNotificationClassDescriptor::PropertyFloat, UMVVMFieldValueChangedTest::FFieldValueChangedDelegate::CreateRaw(&CallbackInstances[1], &FCallbackHandler::OnPropertyFloatChanged));
+	HandleC_2 = SourceObj->AddFieldValueChangedDelegate(UMVVMFieldValueChangedTest::FFieldNotificationClassDescriptor::FunctionInt, UMVVMFieldValueChangedTest::FFieldValueChangedDelegate::CreateRaw(&CallbackInstances[1], &FCallbackHandler::OnFunctionIntChanged));
+	HandleD_2 = SourceObj->AddFieldValueChangedDelegate(UMVVMFieldValueChangedTest::FFieldNotificationClassDescriptor::FunctionFloat, UMVVMFieldValueChangedTest::FFieldValueChangedDelegate::CreateRaw(&CallbackInstances[1], &FCallbackHandler::OnFunctionFloatChanged));
+
+	HandleA_3 = SourceObj->AddFieldValueChangedDelegate(UMVVMFieldValueChangedTest::FFieldNotificationClassDescriptor::PropertyInt, UMVVMFieldValueChangedTest::FFieldValueChangedDelegate::CreateRaw(&CallbackInstances[2], &FCallbackHandler::OnPropertyIntChanged));
+	HandleB_3 = SourceObj->AddFieldValueChangedDelegate(UMVVMFieldValueChangedTest::FFieldNotificationClassDescriptor::PropertyFloat, UMVVMFieldValueChangedTest::FFieldValueChangedDelegate::CreateRaw(&CallbackInstances[2], &FCallbackHandler::OnPropertyFloatChanged));
+	HandleC_3 = SourceObj->AddFieldValueChangedDelegate(UMVVMFieldValueChangedTest::FFieldNotificationClassDescriptor::FunctionInt, UMVVMFieldValueChangedTest::FFieldValueChangedDelegate::CreateRaw(&CallbackInstances[2], &FCallbackHandler::OnFunctionIntChanged));
+	HandleD_3 = SourceObj->AddFieldValueChangedDelegate(UMVVMFieldValueChangedTest::FFieldNotificationClassDescriptor::FunctionFloat, UMVVMFieldValueChangedTest::FFieldValueChangedDelegate::CreateRaw(&CallbackInstances[2], &FCallbackHandler::OnFunctionFloatChanged));
+
+	SourceObj->RemoveAllFieldValueChangedDelegates(UMVVMFieldValueChangedTest::FFieldNotificationClassDescriptor::PropertyInt, &CallbackInstances[0]);
+	TestBroadcastAll();
+	TestAllExpectedCountValues(2, 4, 4, 4);
+	TestExpectedCountValues(0, 0, 2, 2, 2);
+	TestExpectedCountValues(1, 1, 1, 1, 1);
+	ResetAllValues();
+
+	return true;
+}
+
+#undef LOCTEXT_NAMESPACE 
+#endif //WITH_AUTOMATION_WORKER
diff --git a/Engine/Plugins/Runtime/ModelViewViewModel/Source/ModelViewViewModelEditor/Private/Tests/MVVMFieldValueChangedTest.h b/Engine/Plugins/Runtime/ModelViewViewModel/Source/ModelViewViewModelEditor/Private/Tests/MVVMFieldValueChangedTest.h
new file mode 100644
--- /dev/null
+++ b/Engine/Plugins/Runtime/ModelViewViewModel/Source/ModelViewViewModelEditor/Private/Tests/MVVMFieldValueChangedTest.h
@@ -0,0 +1,33 @@
+// Copyright Epic Games, Inc. All Rights Reserved.
+
+#pragma once
+
+#include "MVVMViewModelBase.h"
+
+
+#include "MVVMFieldValueChangedTest.generated.h"
+
+UCLASS(MinimalAPI, Transient, hidedropdown, Hidden)
+class UMVVMFieldValueChangedTest : public UMVVMViewModelBase
+{
+	GENERATED_BODY()
+
+public:
+	UPROPERTY(BlueprintReadWrite, FieldNotify, Category = "MVVM")
+	int32 PropertyInt;
+
+	UPROPERTY(BlueprintReadOnly, FieldNotify, Category = "MVVM")
+	float PropertyFloat;
+
+	UFUNCTION(BlueprintCallable, FieldNotify, Category = "MVVM")
+	int32 FunctionInt() const
+	{
+		return 0;
+	}
+
+	UFUNCTION(BlueprintCallable, FieldNotify, Category = "MVVM")
+	float FunctionFloat() const
+	{
+		return 1.0f;
+	}
+};
diff --git a/Engine/Source/Runtime/FieldNotification/Private/FieldNotificationDelegate.cpp b/Engine/Source/Runtime/FieldNotification/Private/FieldNotificationDelegate.cpp
--- a/Engine/Source/Runtime/FieldNotification/Private/FieldNotificationDelegate.cpp
+++ b/Engine/Source/Runtime/FieldNotification/Private/FieldNotificationDelegate.cpp
@@ -386,9 +386,9 @@ FFieldMulticastDelegate::FRemoveAllResult FFieldMulticastDelegate::RemoveAll(con
 			{
 				Delegates.RemoveAt(Index);
 			}
-			else if (Element.Key.Id == InFieldId && Element.Key.Object == InObject)
+			else if (Element.Key.Object == InObject)
 			{
-				if (Element.Delegate.IsBoundToObject(InUserObject))
+				if (Element.Key.Id == InFieldId && Element.Delegate.IsBoundToObject(InUserObject))
 				{
 					Delegates.RemoveAt(Index);
 					++Result.RemoveCount;
