MVVM: Remove the same binding that was added. Used to remove all the bindings from that has the same FieldId. That was wrong. The same widget may have a new viewmodel of the same type and it would removed the bindings from all viewmodel.

diff --git a/Engine/Plugins/Runtime/ModelViewViewModel/Source/ModelViewViewModel/Private/View/MVVMView.cpp b/Engine/Plugins/Runtime/ModelViewViewModel/Source/ModelViewViewModel/Private/View/MVVMView.cpp
--- a/Engine/Plugins/Runtime/ModelViewViewModel/Source/ModelViewViewModel/Private/View/MVVMView.cpp
+++ b/Engine/Plugins/Runtime/ModelViewViewModel/Source/ModelViewViewModel/Private/View/MVVMView.cpp
@@ -94,6 +94,8 @@ void UMVVMView::InitializeBindings()
 	bHasEveryTickBinding = false;
 
 	const TArrayView<const FMVVMViewClass_CompiledBinding> CompiledBindings = ClassExtension->GetCompiledBindings();
+	RegisteredLibraryBindings.Reset(CompiledBindings.Num());
+	RegisteredLibraryBindings.AddDefaulted(CompiledBindings.Num());
 	for (int32 Index = 0; Index < CompiledBindings.Num(); ++Index)
 	{
 		const FMVVMViewClass_CompiledBinding& Binding = CompiledBindings[Index];
@@ -599,21 +601,21 @@ void UMVVMView::ExecuteLibraryBinding(const FMVVMViewClass_CompiledBinding& Bind
 
 bool UMVVMView::IsLibraryBindingEnabled(int32 InBindindIndex) const
 {
-	return RegisteredLibraryBindings.IsValidIndex(InBindindIndex) && RegisteredLibraryBindings[InBindindIndex];
+	return RegisteredLibraryBindings[InBindindIndex].IsValid();
 }
 
 
 void UMVVMView::EnableLibraryBinding(const FMVVMViewClass_CompiledBinding& Binding, int32 BindingIndex)
 {
-	RegisteredLibraryBindings.PadToNum(BindingIndex + 1, false);
-	check(RegisteredLibraryBindings[BindingIndex] == false);
+	check(!RegisteredLibraryBindings[BindingIndex].IsValid());
 
 	EMVVMExecutionMode ExecutionMode = Binding.GetExecuteMode();
 	bool bCanExecute = true;
 	bool bRegistered = false;
 	if (!Binding.IsOneTime() && ExecutionMode != EMVVMExecutionMode::Tick)
 	{
-		bRegistered = RegisterLibraryBinding(Binding, BindingIndex);
+		RegisteredLibraryBindings[BindingIndex] = RegisterLibraryBinding(Binding, BindingIndex);
+		bRegistered = RegisteredLibraryBindings[BindingIndex].IsValid();
 		bCanExecute = bRegistered;
 	}
 
@@ -629,8 +631,6 @@ void UMVVMView::EnableLibraryBinding(const FMVVMViewClass_CompiledBinding& Bindi
 		));
 	}
 
-	RegisteredLibraryBindings[BindingIndex] = bRegistered;
-
 	if (bCanExecute && Binding.NeedsExecutionAtInitialization())
 	{
 		ExecuteLibraryBinding(Binding, BindingIndex);
@@ -644,13 +644,8 @@ void UMVVMView::DisableLibraryBinding(const FMVVMViewClass_CompiledBinding& Bind
 {
 	check(IsLibraryBindingEnabled(BindingIndex));
 
-
-	EMVVMExecutionMode ExecutionMode = Binding.GetExecuteMode();
-	RegisteredLibraryBindings[BindingIndex] = false;
-	if (!Binding.IsOneTime() && ExecutionMode != EMVVMExecutionMode::Tick)
-	{
-		UnregisterLibraryBinding(Binding, BindingIndex);
-	}
+	UnregisterLibraryBinding(Binding, RegisteredLibraryBindings[BindingIndex], BindingIndex);
+	RegisteredLibraryBindings[BindingIndex] = FDelegateHandle();
 
 	if (bLogBinding)
 	{
@@ -666,7 +661,7 @@ void UMVVMView::DisableLibraryBinding(const FMVVMViewClass_CompiledBinding& Bind
 }
 
 
-bool UMVVMView::RegisterLibraryBinding(const FMVVMViewClass_CompiledBinding& Binding, int32 BindingIndex)
+FDelegateHandle UMVVMView::RegisterLibraryBinding(const FMVVMViewClass_CompiledBinding& Binding, int32 BindingIndex)
 {
 	check(ClassExtension);
 
@@ -692,7 +687,7 @@ bool UMVVMView::RegisterLibraryBinding(const FMVVMViewClass_CompiledBinding& Bin
 #if UE_WITH_MVVM_DEBUGGING
 		UE::MVVM::FDebugging::BroadcastLibraryBindingRegistered(this, Binding, UE::MVVM::FDebugging::ERegisterLibraryBindingResult::Failed_InvalidFieldId);
 #endif
-		return false;
+		return FDelegateHandle();
 	}
 
 	UE::FieldNotification::FFieldId FieldId = FieldIdResult.StealValue();
@@ -703,7 +698,7 @@ bool UMVVMView::RegisterLibraryBinding(const FMVVMViewClass_CompiledBinding& Bin
 #if UE_WITH_MVVM_DEBUGGING
 		UE::MVVM::FDebugging::BroadcastLibraryBindingRegistered(this, Binding, UE::MVVM::FDebugging::ERegisterLibraryBindingResult::Failed_FieldIdNotFound);
 #endif
-		return false;
+		return FDelegateHandle();
 	}
 
 	// The source may not have been created because the property path was wrong.
@@ -718,11 +713,12 @@ bool UMVVMView::RegisterLibraryBinding(const FMVVMViewClass_CompiledBinding& Bin
 #if UE_WITH_MVVM_DEBUGGING
 		UE::MVVM::FDebugging::BroadcastLibraryBindingRegistered(this, Binding, UE::MVVM::FDebugging::ERegisterLibraryBindingResult::Failed_InvalidSource);
 #endif
-		return false;
+		return FDelegateHandle();
 	}
 
 	UE::FieldNotification::FFieldMulticastDelegate::FDelegate Delegate = UE::FieldNotification::FFieldMulticastDelegate::FDelegate::CreateUObject(this, &UMVVMView::HandledLibraryBindingValueChanged, BindingIndex);
-	if (Source->AddFieldValueChangedDelegate(FieldId, MoveTemp(Delegate)).IsValid())
+	FDelegateHandle ResultHandle = Source->AddFieldValueChangedDelegate(FieldId, MoveTemp(Delegate));
+	if (ResultHandle.IsValid())
 	{
 		if (FMVVMViewSource* RegisteredSource = FindViewSource(Binding.GetSourceName()))
 		{
@@ -744,15 +740,15 @@ bool UMVVMView::RegisterLibraryBinding(const FMVVMViewClass_CompiledBinding& Bin
 	else
 	{
 #if UE_WITH_MVVM_DEBUGGING
-		UE::MVVM::FDebugging::BroadcastLibraryBindingRegistered(this, Binding, UE::MVVM::FDebugging::ERegisterLibraryBindingResult::Success);
+		UE::MVVM::FDebugging::BroadcastLibraryBindingRegistered(this, Binding, UE::MVVM::FDebugging::ERegisterLibraryBindingResult::Failed_InvalidSourceField);
 #endif
 	}
 
-	return true;
+	return ResultHandle;
 }
 
 
-void UMVVMView::UnregisterLibraryBinding(const FMVVMViewClass_CompiledBinding& Binding, int32 BindingIndex)
+void UMVVMView::UnregisterLibraryBinding(const FMVVMViewClass_CompiledBinding& Binding, FDelegateHandle Handle, int32 BindingIndex)
 {
 	TScriptInterface<INotifyFieldValueChanged> Source = FindSource(Binding, BindingIndex, true);
 	if (Source.GetInterface() && ClassExtension)
@@ -760,7 +756,7 @@ void UMVVMView::UnregisterLibraryBinding(const FMVVMViewClass_CompiledBinding& B
 		TValueOrError<UE::FieldNotification::FFieldId, void> FieldIdResult = ClassExtension->GetBindingLibrary().GetFieldId(Binding.GetSourceFieldId());
 		if (ensureMsgf(FieldIdResult.HasValue() && FieldIdResult.GetValue().IsValid(), TEXT("If the binding was enabled then the FieldId should exist.")))
 		{
-			Source->RemoveAllFieldValueChangedDelegates(FieldIdResult.GetValue(), this);
+			Source->RemoveFieldValueChangedDelegate(FieldIdResult.GetValue(), Handle);
 		}
 
 		FMVVMViewSource* RegisteredSource = FindViewSource(Binding.GetSourceName());
diff --git a/Engine/Plugins/Runtime/ModelViewViewModel/Source/ModelViewViewModel/Public/View/MVVMView.h b/Engine/Plugins/Runtime/ModelViewViewModel/Source/ModelViewViewModel/Public/View/MVVMView.h
--- a/Engine/Plugins/Runtime/ModelViewViewModel/Source/ModelViewViewModel/Public/View/MVVMView.h
+++ b/Engine/Plugins/Runtime/ModelViewViewModel/Source/ModelViewViewModel/Public/View/MVVMView.h
@@ -2,7 +2,7 @@
 
 #pragma once
 
-
+#include "CoreTypes.h"
 #include "Extensions/UserWidgetExtension.h"
 
 #include "MVVMView.generated.h"
@@ -116,8 +116,8 @@ private:
 	void DisableLibraryBinding(const FMVVMViewClass_CompiledBinding& Item, int32 BindingIndex);
 	bool IsLibraryBindingEnabled(int32 InBindindIndex) const;
 
-	bool RegisterLibraryBinding(const FMVVMViewClass_CompiledBinding& Binding, int32 BindingIndex);
-	void UnregisterLibraryBinding(const FMVVMViewClass_CompiledBinding& Binding, int32 BindingIndex);
+	FDelegateHandle RegisterLibraryBinding(const FMVVMViewClass_CompiledBinding& Binding, int32 BindingIndex);
+	void UnregisterLibraryBinding(const FMVVMViewClass_CompiledBinding& Binding, FDelegateHandle Handle, int32 BindingIndex);
 
 	TScriptInterface<INotifyFieldValueChanged> FindSource(const FMVVMViewClass_CompiledBinding& Binding, int32 BindingIndex, bool bAllowNull) const;
 	FMVVMViewSource* FindViewSource(const FName SourceName);
@@ -131,7 +131,7 @@ private:
 	TArray<FMVVMViewSource> Sources;
 
 	/** The binding that are registered by the view to the sources. */
-	TBitArray<> RegisteredLibraryBindings;
+	TArray<FDelegateHandle> RegisteredLibraryBindings;
 
 	/** Should log when a binding is executed. */
 	UPROPERTY(EditAnywhere, Transient, Category = "Viewmodel")
