MVVM: Test if the view is initialized before accessing the IsLibraryBindingEnabled function. Fix comments from previous review.

diff --git a/Engine/Plugins/Runtime/ModelViewViewModel/Source/ModelViewViewModel/Private/View/MVVMView.cpp b/Engine/Plugins/Runtime/ModelViewViewModel/Source/ModelViewViewModel/Private/View/MVVMView.cpp
--- a/Engine/Plugins/Runtime/ModelViewViewModel/Source/ModelViewViewModel/Private/View/MVVMView.cpp
+++ b/Engine/Plugins/Runtime/ModelViewViewModel/Source/ModelViewViewModel/Private/View/MVVMView.cpp
@@ -94,6 +94,7 @@ void UMVVMView::InitializeBindings()
 	bHasEveryTickBinding = false;
 
 	const TArrayView<const FMVVMViewClass_CompiledBinding> CompiledBindings = ClassExtension->GetCompiledBindings();
+	ensure(RegisteredLibraryBindings.IsEmpty());
 	RegisteredLibraryBindings.Reset(CompiledBindings.Num());
 	RegisteredLibraryBindings.AddDefaulted(CompiledBindings.Num());
 	for (int32 Index = 0; Index < CompiledBindings.Num(); ++Index)
@@ -125,11 +126,11 @@ void UMVVMView::Destruct()
 
 	if (bInitialized)
 	{
-		DeintializeBindings();
+		DeinitializeBindings();
 	}
 }
 
-void UMVVMView::DeintializeBindings()
+void UMVVMView::DeinitializeBindings()
 {
 	check(bInitialized == true);
 	bInitialized = false;
@@ -316,12 +317,15 @@ bool UMVVMView::SetSourceInternal(FName ViewModelName, TScriptInterface<INotifyF
 		const TArrayView<const FMVVMViewClass_CompiledBinding> CompiledBindings = ClassExtension->GetCompiledBindings();
 
 		// Unregister any bindings from that source
-		for (int32 Index = 0; Index < CompiledBindings.Num(); ++Index)
+		if (bInitialized)
 		{
-			const FMVVMViewClass_CompiledBinding& Binding = CompiledBindings[Index];
-			if (IsLibraryBindingEnabled(Index) && Binding.GetSourceName() == ViewModelName)
+			for (int32 Index = 0; Index < CompiledBindings.Num(); ++Index)
 			{
-				DisableLibraryBinding(Binding, Index);
+				const FMVVMViewClass_CompiledBinding& Binding = CompiledBindings[Index];
+				if (IsLibraryBindingEnabled(Index) && Binding.GetSourceName() == ViewModelName)
+				{
+					DisableLibraryBinding(Binding, Index);
+				}
 			}
 		}
 
@@ -340,7 +344,7 @@ bool UMVVMView::SetSourceInternal(FName ViewModelName, TScriptInterface<INotifyF
 		bool bPreviousEveryTickBinding = bHasEveryTickBinding;
 		bHasEveryTickBinding = false;
 		// Register back any binding that was previously enabled
-		if (NewValue.GetObject() && bInitialized)
+		if (bInitialized && NewValue.GetObject())
 		{
 			// Enabled the default bindings
 			for (int32 Index = 0; Index < CompiledBindings.Num(); ++Index)
@@ -370,7 +374,7 @@ bool UMVVMView::SetSourceInternal(FName ViewModelName, TScriptInterface<INotifyF
 			}
 		}
 
-		if (bPreviousEveryTickBinding != bHasEveryTickBinding)
+		if (bInitialized && bPreviousEveryTickBinding != bHasEveryTickBinding)
 		{
 			if (bHasEveryTickBinding)
 			{
@@ -441,7 +445,7 @@ void UMVVMView::HandledLibraryBindingValueChanged(UObject* InViewModelOrWidget,
 
 void UMVVMView::ExecuteDelayedBinding(const FMVVMViewDelayedBinding& DelayedBinding) const
 {
-	if (ensure(ClassExtension))
+	if (ensure(ClassExtension) && bInitialized)
 	{
 		if (ensure(ClassExtension->GetCompiledBindings().IsValidIndex(DelayedBinding.GetCompiledBindingIndex())))
 		{
@@ -480,8 +484,9 @@ void UMVVMView::ExecuteDelayedBinding(const FMVVMViewDelayedBinding& DelayedBind
 void UMVVMView::ExecuteEveryTickBindings() const
 {
 	ensure(bHasEveryTickBinding);
+	ensure(bInitialized);
 
-	if (ClassExtension)
+	if (ClassExtension && bInitialized)
 	{
 		const TArrayView<const FMVVMViewClass_CompiledBinding> CompiledBindings = ClassExtension->GetCompiledBindings();
 
@@ -601,7 +606,7 @@ void UMVVMView::ExecuteLibraryBinding(const FMVVMViewClass_CompiledBinding& Bind
 
 bool UMVVMView::IsLibraryBindingEnabled(int32 InBindindIndex) const
 {
-	return RegisteredLibraryBindings[InBindindIndex].IsValid();
+	return ensure(RegisteredLibraryBindings.IsValidIndex(InBindindIndex)) && RegisteredLibraryBindings[InBindindIndex].IsValid();
 }
 
 
diff --git a/Engine/Plugins/Runtime/ModelViewViewModel/Source/ModelViewViewModel/Public/View/MVVMView.h b/Engine/Plugins/Runtime/ModelViewViewModel/Source/ModelViewViewModel/Public/View/MVVMView.h
--- a/Engine/Plugins/Runtime/ModelViewViewModel/Source/ModelViewViewModel/Public/View/MVVMView.h
+++ b/Engine/Plugins/Runtime/ModelViewViewModel/Source/ModelViewViewModel/Public/View/MVVMView.h
@@ -69,7 +69,7 @@ public:
 	//~ End UUserWidgetExtension implementation
 
 	void InitializeBindings();
-	void DeintializeBindings();
+	void DeinitializeBindings();
 
 	const UMVVMViewClass* GetViewClass() const
 	{
