MVVM: Create a library for common MVVM conversion function.

diff --git a/Engine/Plugins/Runtime/ModelViewViewModel/Source/ModelViewViewModel/Private/Bindings/MVVMConversionLibrary.cpp b/Engine/Plugins/Runtime/ModelViewViewModel/Source/ModelViewViewModel/Private/Bindings/MVVMConversionLibrary.cpp
new file mode 100644
--- /dev/null
+++ b/Engine/Plugins/Runtime/ModelViewViewModel/Source/ModelViewViewModel/Private/Bindings/MVVMConversionLibrary.cpp
@@ -0,0 +1,8 @@
+// Copyright Epic Games, Inc. All Rights Reserved.
+
+#include "Bindings/MVVMConversionLibrary.h"
+
+ESlateVisibility UMVVMConversionLibrary::Conv_BoolToSlateVisibility(bool bIsVisible, ESlateVisibility TrueCaseVisibility, ESlateVisibility FalseCaseVisibility)
+{
+	return bIsVisible ? TrueCaseVisibility : FalseCaseVisibility;
+}
diff --git a/Engine/Plugins/Runtime/ModelViewViewModel/Source/ModelViewViewModel/Public/Bindings/MVVMConversionLibrary.h b/Engine/Plugins/Runtime/ModelViewViewModel/Source/ModelViewViewModel/Public/Bindings/MVVMConversionLibrary.h
new file mode 100644
--- /dev/null
+++ b/Engine/Plugins/Runtime/ModelViewViewModel/Source/ModelViewViewModel/Public/Bindings/MVVMConversionLibrary.h
@@ -0,0 +1,20 @@
+// Copyright Epic Games, Inc. All Rights Reserved.
+
+#pragma once
+
+#include "CoreMinimal.h"
+#include "Components/SlateWrapperTypes.h"
+#include "Kismet/BlueprintFunctionLibrary.h"
+#include "MVVMConversionLibrary.generated.h"
+
+UCLASS()
+class MODELVIEWVIEWMODEL_API UMVVMConversionLibrary : public UBlueprintFunctionLibrary
+{
+	GENERATED_BODY()
+	
+	/**
+	 * Converts a bool to a slate visibility.
+	 */
+	UFUNCTION(BlueprintPure, Category = "Widget", meta = (DisplayName = "To Bool (Visibility)"))
+	static ESlateVisibility Conv_BoolToSlateVisibility(bool bIsVisible, ESlateVisibility TrueVisibility = ESlateVisibility::Visible, ESlateVisibility FalseVisibility = ESlateVisibility::Collapsed);
+};
\ No newline at end of file
