MVVM: Add a blueprint view to the AssetSearch. The module is not added because it depends on the AssetSearch plugin. The code to enable it is ... but the "Optional" doesn't work at the moment and it would always enabled the AssetSearch.

 {
    "Name": "ModelViewViewModelAssetSearch",
    "Type": "Editor",
    "LoadingPhase": "Default"
  }
],
"Plugins": [
{
  "Name": "AssetSearch",
  "Enabled": true,
  "Optional": true
  ]

diff --git a/Engine/Plugins/Editor/AssetSearch/Source/Private/Indexers/BlueprintIndexer.cpp b/Engine/Plugins/Editor/AssetSearch/Source/Private/Indexers/BlueprintIndexer.cpp
--- a/Engine/Plugins/Editor/AssetSearch/Source/Private/Indexers/BlueprintIndexer.cpp
+++ b/Engine/Plugins/Editor/AssetSearch/Source/Private/Indexers/BlueprintIndexer.cpp
@@ -1,6 +1,7 @@
 // Copyright Epic Games, Inc. All Rights Reserved.
 
 #include "BlueprintIndexer.h"
+#include "Blueprint/BlueprintExtension.h"
 #include "EdGraph/EdGraph.h"
 #include "K2Node_CallFunction.h"
 #include "K2Node_Variable.h"
@@ -22,6 +23,7 @@ enum class EBlueprintIndexerVersion
 	IndexingPublicEditableFieldsOnNodes,
 	DontIndexPinsUnlessItsInputWithNoConnections,
 	BetterSupportForIndexingEventNodes,
+	BlueprintExtension,
 
 	// -----<new versions can be added above this line>-------------------------------------------------
 	VersionPlusOne,
@@ -46,6 +48,7 @@ void FBlueprintIndexer::IndexAsset(const UObject* InAssetObject, FSearchSerializ
 	IndexClassDefaultObject(BP, Serializer);
 	IndexComponents(BP, Serializer);
 	IndexGraphs(BP, Serializer);
+	IndexExtensions(BP, Serializer);
 }
 
 void FBlueprintIndexer::IndexClassDefaultObject(const UBlueprint* InBlueprint, FSearchSerializer& Serializer) const
@@ -175,4 +178,12 @@ void FBlueprintIndexer::IndexMemberReference(FSearchSerializer& Serializer, cons
 	}
 }
 
+void FBlueprintIndexer::IndexExtensions(const UBlueprint* InBlueprint, FSearchSerializer& Serializer) const
+{
+	for (UBlueprintExtension* Extension : InBlueprint->GetExtensions())
+	{
+		Serializer.IndexNestedAsset(Extension);
+	}
+}
+
 #undef LOCTEXT_NAMESPACE
diff --git a/Engine/Plugins/Editor/AssetSearch/Source/Private/Indexers/BlueprintIndexer.h b/Engine/Plugins/Editor/AssetSearch/Source/Private/Indexers/BlueprintIndexer.h
--- a/Engine/Plugins/Editor/AssetSearch/Source/Private/Indexers/BlueprintIndexer.h
+++ b/Engine/Plugins/Editor/AssetSearch/Source/Private/Indexers/BlueprintIndexer.h
@@ -24,4 +24,5 @@ private:
 	void IndexComponents(const UBlueprint* InBlueprint, FSearchSerializer& Serializer) const;
 	void IndexGraphs(const UBlueprint* InBlueprint, FSearchSerializer& Serializer) const;
 	void IndexMemberReference(FSearchSerializer& Serializer, const FMemberReference& MemberReference, const FString& MemberType) const;
+	void IndexExtensions(const UBlueprint* InBlueprint, FSearchSerializer& Serializer) const;
 };
diff --git a/Engine/Plugins/Runtime/ModelViewViewModel/Source/ModelViewViewModelAssetSearch/ModelViewViewModelAssetSearch.Build.cs b/Engine/Plugins/Runtime/ModelViewViewModel/Source/ModelViewViewModelAssetSearch/ModelViewViewModelAssetSearch.Build.cs
new file mode 100644
--- /dev/null
+++ b/Engine/Plugins/Runtime/ModelViewViewModel/Source/ModelViewViewModelAssetSearch/ModelViewViewModelAssetSearch.Build.cs
@@ -0,0 +1,28 @@
+// Copyright Epic Games, Inc. All Rights Reserved.
+
+using System.IO;
+using System.Collections.Generic;
+using UnrealBuildTool;
+
+public class ModelViewViewModelAssetSearch : ModuleRules 
+{
+	public ModelViewViewModelAssetSearch(ReadOnlyTargetRules Target) : base(Target)
+	{
+		PublicDependencyModuleNames.AddRange(
+			new string[]
+			{
+				"AssetSearch",
+				"Core",
+				"CoreUObject",
+				"ModelViewViewModel",
+				"ModelViewViewModelBlueprint",
+			});
+
+		PrivateDependencyModuleNames.AddRange(
+			new string[]
+			{
+				"InputCore",
+				"UMGEditor",
+			});
+	}
+}
diff --git a/Engine/Plugins/Runtime/ModelViewViewModel/Source/ModelViewViewModelAssetSearch/Private/MVVMAssetSearchIndexer.h b/Engine/Plugins/Runtime/ModelViewViewModel/Source/ModelViewViewModelAssetSearch/Private/MVVMAssetSearchIndexer.h
new file mode 100644
--- /dev/null
+++ b/Engine/Plugins/Runtime/ModelViewViewModel/Source/ModelViewViewModelAssetSearch/Private/MVVMAssetSearchIndexer.h
@@ -0,0 +1,54 @@
+// Copyright Epic Games, Inc. All Rights Reserved.
+
+#pragma once
+
+#include "Containers/UnrealString.h"
+#include "IAssetIndexer.h"
+#include "SearchSerializer.h"
+#include "Utility/IndexerUtilities.h"
+
+#include "MVVMBlueprintView.h"
+#include "MVVMWidgetBlueprintExtension_View.h"
+
+
+namespace UE::MVVM::Private
+{
+class FAssetSearchIndexer : public IAssetIndexer
+{
+private:
+	enum class EIndexerVersion
+	{
+		Initial = 0,
+
+		// -----<new versions can be added above this line>-------------------------------------------------
+		VersionPlusOne,
+		LatestVersion = VersionPlusOne - 1
+	};
+
+public:
+	virtual FString GetName() const override
+	{
+		return TEXT("MVVM");
+	}
+
+	virtual int32 GetVersion() const override
+	{
+		return (int32)EIndexerVersion::LatestVersion;
+	}
+
+	virtual void IndexAsset(const UObject* InAssetObject, FSearchSerializer& Serializer) const override
+	{
+		const UMVVMWidgetBlueprintExtension_View* AssetObject = CastChecked<UMVVMWidgetBlueprintExtension_View>(InAssetObject);
+
+		if (const UMVVMBlueprintView* View = AssetObject->GetBlueprintView())
+		{
+			Serializer.BeginIndexingObject(View, TEXT("$self"));
+			FIndexerUtilities::IterateIndexableProperties(View, [&Serializer](const FProperty* Property, const FString& Value) {
+				Serializer.IndexProperty(Property, Value);
+				UE_LOG(LogTemp, Warning, TEXT("--%s [Serialized]"), *Property->GetName());
+				});
+			Serializer.EndIndexingObject();
+		}
+	}
+};
+} //namespace
diff --git a/Engine/Plugins/Runtime/ModelViewViewModel/Source/ModelViewViewModelAssetSearch/Private/ModelViewViewModelAssetSearchModule.cpp b/Engine/Plugins/Runtime/ModelViewViewModel/Source/ModelViewViewModelAssetSearch/Private/ModelViewViewModelAssetSearchModule.cpp
new file mode 100644
--- /dev/null
+++ b/Engine/Plugins/Runtime/ModelViewViewModel/Source/ModelViewViewModelAssetSearch/Private/ModelViewViewModelAssetSearchModule.cpp
@@ -0,0 +1,61 @@
+// Copyright Epic Games, Inc. All Rights Reserved.
+
+
+#include "Modules/ModuleInterface.h"
+#include "Modules/ModuleManager.h"
+
+#include "IAssetSearchModule.h"
+#include "MVVMWidgetBlueprintExtension_View.h"
+
+#include "MVVMAssetSearchIndexer.h"
+
+#define LOCTEXT_NAMESPACE "ModelViewViewModelAssetSearch"
+
+class FMVVMAssetSearchModule : public IModuleInterface
+{
+public:
+	virtual void StartupModule() override
+	{
+		if (IAssetSearchModule* Ptr = FModuleManager::GetModulePtr<IAssetSearchModule>(ModuleName))
+		{
+			RegisterAssetIndexer();
+		}
+		else
+		{
+			ModuleChangedHandle = FModuleManager::Get().OnModulesChanged().AddRaw(this, &FMVVMAssetSearchModule::HandleModulesChanged);
+		}
+	}
+
+	virtual void ShutdownModule() override
+	{
+		if (ModuleChangedHandle.IsValid())
+		{
+			FModuleManager::Get().OnModulesChanged().Remove(ModuleChangedHandle);
+			ModuleChangedHandle.Reset();
+		}
+	}
+
+private:
+	void HandleModulesChanged(FName InModuleName, EModuleChangeReason InReason)
+	{
+		if (InModuleName == ModuleName && InReason == EModuleChangeReason::ModuleLoaded)
+		{
+			RegisterAssetIndexer();
+			FModuleManager::Get().OnModulesChanged().Remove(ModuleChangedHandle);
+		}
+	}
+
+	void RegisterAssetIndexer()
+	{
+		IAssetSearchModule::Get().RegisterAssetIndexer(UMVVMWidgetBlueprintExtension_View::StaticClass(), MakeUnique<UE::MVVM::Private::FAssetSearchIndexer>());
+	}
+
+private:
+	const FName ModuleName = "AssetSearch";
+	FDelegateHandle ModuleChangedHandle;
+};
+
+
+IMPLEMENT_MODULE(FMVVMAssetSearchModule, ModelViewViewModelAssetSearch);
+
+#undef LOCTEXT_NAMESPACE
