MVVM: Reused the same tree and child list if they are available. The key (the sharedptr) is what is kept internally by the TreeView to know where scrolling and which widget are extended.

diff --git a/Engine/Plugins/Runtime/ModelViewViewModel/Source/ModelViewViewModelEditor/Private/Widgets/SMVVMViewBindingListView.cpp b/Engine/Plugins/Runtime/ModelViewViewModel/Source/ModelViewViewModelEditor/Private/Widgets/SMVVMViewBindingListView.cpp
--- a/Engine/Plugins/Runtime/ModelViewViewModel/Source/ModelViewViewModelEditor/Private/Widgets/SMVVMViewBindingListView.cpp
+++ b/Engine/Plugins/Runtime/ModelViewViewModel/Source/ModelViewViewModelEditor/Private/Widgets/SMVVMViewBindingListView.cpp
@@ -40,19 +40,19 @@ struct FBindingEntry
 	enum class ERowType
 	{
 		None,
-		Widget,
+		Group,
 		Binding,
 		Parameter
 	};
 
 	FMVVMBlueprintViewBinding* GetBinding(UMVVMBlueprintView* View) const
 	{
-		return View->GetBindingAt(BindingIndex);
+		return View->GetBinding(BindingId);
 	}
 
 	const FMVVMBlueprintViewBinding* GetBinding(const UMVVMBlueprintView* View) const
 	{
-		return View->GetBindingAt(BindingIndex);
+		return View->GetBinding(BindingId);
 	}
 
 	ERowType GetRowType() const
@@ -60,30 +60,30 @@ struct FBindingEntry
 		return RowType;
 	}
 
-	int32 GetBindingIndex() const
+	FGuid GetBindingId() const
 	{
-		return BindingIndex;
+		return BindingId;
 	}
 
-	void SetBindingIndex(int32 Index)
+	void SetBindingId(FGuid Id)
 	{
 		check(RowType == ERowType::None);
 		RowType = ERowType::Binding;
-		BindingIndex = Index;
+		BindingId = Id;
 	}
 
-	void SetWidgetName(FName WidgetName)
+	void SetGroupName(FName GroupName)
 	{
 		check(RowType == ERowType::None);
-		RowType = ERowType::Widget;
-		Name = WidgetName;
+		RowType = ERowType::Group;
+		Name = GroupName;
 	}
 
-	void SetParameterName(int32 Index, FName ParameterName)
+	void SetParameterName(FGuid Id, FName ParameterName)
 	{
 		check(RowType == ERowType::None);
 		RowType = ERowType::Parameter;
-		BindingIndex = Index;
+		BindingId = Id;
 		Name = ParameterName;
 	}
 
@@ -102,17 +102,22 @@ struct FBindingEntry
 		Children.Add(Child);
 	}
 
+	void ResetChildren()
+	{
+		Children.Reset();
+	}
+
 	bool operator==(const FBindingEntry& Other) const
 	{
 		return RowType == Other.RowType &&
 			Name == Other.Name &&
-			BindingIndex == Other.BindingIndex;
+			BindingId == Other.BindingId;
 	}
 
 private:
 	ERowType RowType = ERowType::None;
 	FName Name;
-	int32 BindingIndex = INDEX_NONE;
+	FGuid BindingId;
 	TArray<TSharedPtr<FBindingEntry>> Children;
 };
 
@@ -1149,7 +1154,7 @@ void SBindingsList::Construct(const FArguments& InArgs, TSharedPtr<SBindingsPane
 	ChildSlot
 	[
 		SAssignNew(TreeView, STreeView<TSharedPtr<FBindingEntry>>)
-		.TreeItemsSource(&RootWidgets)
+		.TreeItemsSource(&RootGroups)
 		.SelectionMode(ESelectionMode::Single)
 		.OnGenerateRow(this, &SBindingsList::GenerateEntryRow)
 		.OnGetChildren(this, &SBindingsList::GetChildrenOfEntry)
@@ -1182,80 +1187,134 @@ void SBindingsList::Refresh()
 	UMVVMWidgetBlueprintExtension_View* MVVMExtensionPtr = MVVMExtension.Get();
 	UMVVMBlueprintView* BlueprintView = MVVMExtensionPtr ? MVVMExtensionPtr->GetBlueprintView() : nullptr;
 
-	// store the current binding index
-	TArray<TSharedPtr<FBindingEntry>> PreviousSelectedEntries;
-	if (TreeView.IsValid() && BlueprintView)
+	struct FPreviousGroup
 	{
-		PreviousSelectedEntries = TreeView->GetSelectedItems();
-	}
+		TSharedPtr<FBindingEntry> Group;
+		TArray<TSharedPtr<FBindingEntry>> Children;
+	};
+	TArray<FPreviousGroup> PreviousRootGroups;
+	for (TSharedPtr<FBindingEntry> PreviousEntry : RootGroups)
+	{
+		ensure(PreviousEntry->GetRowType() == FBindingEntry::ERowType::Group);
+		FPreviousGroup& NewItem = PreviousRootGroups.AddDefaulted_GetRef();
+		NewItem.Group = PreviousEntry;
 
-	RootWidgets.Reset();
+		// Add bindings
+		for (TSharedPtr<FBindingEntry> PreviousChildA : PreviousEntry->GetChildren())
+		{
+			NewItem.Children.Add(PreviousChildA);
+			// Add function arguments
+			for (TSharedPtr<FBindingEntry> PreviousChildB : PreviousChildA->GetChildren())
+			{
+				NewItem.Children.Add(PreviousChildB);
+			}
+			PreviousChildA->ResetChildren();
+		}
+		PreviousEntry->ResetChildren();
+	}
+	RootGroups.Reset();
 
-	TArray<TSharedPtr<FBindingEntry>> NewSelectedEntries;
+	TArray<TSharedPtr<FBindingEntry>> NewEntries;
 
 	// generate our entries
 	// for each widget with bindings, create an entry at the root level
 	// then add all bindings that reference that widget as its children
 	if (BlueprintView)
 	{
-		auto ReselectIfRequired = [&PreviousSelectedEntries, &NewSelectedEntries](const TSharedPtr<FBindingEntry>& ToFind)
-		{
-			if (PreviousSelectedEntries.ContainsByPredicate([ToFind](const TSharedPtr<FBindingEntry>& Entry) -> bool
-				{
-					return Entry.Get() == ToFind.Get();
-				}))
-			{
-				NewSelectedEntries.Add(ToFind);
-			}
-		};
-
 		TArrayView<FMVVMBlueprintViewBinding> Bindings = BlueprintView->GetBindings();
 		for (int32 BindingIndex = 0; BindingIndex < Bindings.Num(); ++BindingIndex)
 		{
 			const FMVVMBlueprintViewBinding& Binding = Bindings[BindingIndex];
 			
-			FName WidgetName = Binding.DestinationPath.GetWidgetName();
-
-			TSharedPtr<FBindingEntry> ExistingWidget;
-			for (TSharedPtr<FBindingEntry> Widget : RootWidgets)
+			FName GroupName;
+			if (Binding.DestinationPath.IsFromWidget())
+			{
+				GroupName = Binding.DestinationPath.GetWidgetName();
+			}
+			else
 			{
-				if (WidgetName == Widget->GetName())
+				if (const FMVVMBlueprintViewModelContext* ViewModelContext = BlueprintView->FindViewModel(Binding.DestinationPath.GetViewModelId()))
 				{
-					ExistingWidget = Widget;
-					break;
+					GroupName = ViewModelContext->GetViewModelName();
 				}
 			}
+			FPreviousGroup* PreviousGroupEntry = PreviousRootGroups.FindByPredicate([GroupName](const FPreviousGroup& Other) { return Other.Group->GetName() == GroupName; });
 
-			if (!ExistingWidget.IsValid())
+			// Find the group entry
+			TSharedPtr<FBindingEntry> GroupEntry;
 			{
-				TSharedRef<FBindingEntry> NewWidget = MakeShared<FBindingEntry>();
-				NewWidget->SetWidgetName(WidgetName);
-				ExistingWidget = NewWidget;
+				if (PreviousGroupEntry)
+				{
+					GroupEntry = PreviousGroupEntry->Group;
+				}
+				else if (TSharedPtr<FBindingEntry>* FoundGroup = NewEntries.FindByPredicate([GroupName](const TSharedPtr<FBindingEntry>& Other)
+					{ return Other->GetName() == GroupName && Other->GetRowType() == FBindingEntry::ERowType::Group; }))
+				{
+					GroupEntry = *FoundGroup;
+				}
 
-				RootWidgets.Add(NewWidget);
+				if (!GroupEntry.IsValid())
+				{
+					GroupEntry = MakeShared<FBindingEntry>();
+					GroupEntry->SetGroupName(GroupName);
 
-				ReselectIfRequired(NewWidget);
+					NewEntries.Add(GroupEntry);
+				}
+				RootGroups.AddUnique(GroupEntry);
 			}
 
-			TSharedRef<FBindingEntry> NewBindingEntry = MakeShared<FBindingEntry>();
-			NewBindingEntry->SetBindingIndex(BindingIndex);
-			ExistingWidget->AddChild(NewBindingEntry);
-			ReselectIfRequired(NewBindingEntry);
+			// Create/Find the child entry
+			TSharedPtr<FBindingEntry> BindingEntry;
+			FGuid BindingId = Binding.BindingId;
+			{
+				if (PreviousGroupEntry)
+				{
+					if (TSharedPtr<FBindingEntry>* FoundBinding = PreviousGroupEntry->Children.FindByPredicate([BindingId](const TSharedPtr<FBindingEntry>& Other)
+						{ return Other->GetBindingId() == BindingId && Other->GetRowType() == FBindingEntry::ERowType::Binding; }))
+					{
+						BindingEntry = *FoundBinding;
+					}
+				}
 
-			// create entries for conversion function parameters
-			UMVVMEditorSubsystem* EditorSubsystem = GEditor->GetEditorSubsystem<UMVVMEditorSubsystem>();
-			const UFunction* Function = EditorSubsystem->GetConversionFunction(MVVMExtensionPtr->GetWidgetBlueprint(), Binding, UE::MVVM::IsForwardBinding(Binding.BindingType));
-			if (Function != nullptr)
+				if (!BindingEntry.IsValid())
+				{
+					BindingEntry = MakeShared<FBindingEntry>();
+					BindingEntry->SetBindingId(BindingId);
+
+					NewEntries.Add(BindingEntry);
+				}
+				GroupEntry->AddChild(BindingEntry.ToSharedRef());
+			}
+
+			// Create/Find entries for conversion function parameters
 			{
-				TValueOrError<TArray<const FProperty*>, FText> ArgumentsResult = BindingHelper::TryGetArgumentsForConversionFunction(Function);
-				if (ArgumentsResult.HasValue())
+				const UFunction* Function = GEditor->GetEditorSubsystem<UMVVMEditorSubsystem>()->GetConversionFunction(MVVMExtensionPtr->GetWidgetBlueprint(), Binding, UE::MVVM::IsForwardBinding(Binding.BindingType));
+				if (Function != nullptr)
 				{
-					for (const FProperty* Argument : ArgumentsResult.GetValue())
+					TValueOrError<TArray<const FProperty*>, FText> ArgumentsResult = BindingHelper::TryGetArgumentsForConversionFunction(Function);
+					if (ArgumentsResult.HasValue())
 					{
-						TSharedRef<FBindingEntry> NewArgumentEntry = MakeShared<FBindingEntry>();
-						NewArgumentEntry->SetParameterName(BindingIndex, Argument->GetFName());
-						NewBindingEntry->AddChild(NewArgumentEntry);
-						ReselectIfRequired(NewArgumentEntry);
+						for (const FProperty* Argument : ArgumentsResult.GetValue())
+						{
+							TSharedPtr<FBindingEntry> ArgumentEntry;
+							if (PreviousGroupEntry)
+							{
+								if (TSharedPtr<FBindingEntry>* FoundParameter = PreviousGroupEntry->Children.FindByPredicate([BindingId, ArgumentName = Argument->GetFName()](const TSharedPtr<FBindingEntry>& Other)
+									{ return Other->GetBindingId() == BindingId && Other->GetRowType() == FBindingEntry::ERowType::Parameter && Other->GetName() == ArgumentName; }))
+								{
+									ArgumentEntry = *FoundParameter;
+								}
+							}
+
+							if (!ArgumentEntry.IsValid())
+							{
+								ArgumentEntry = MakeShared<FBindingEntry>();
+								ArgumentEntry->SetParameterName(Binding.BindingId, Argument->GetFName());
+
+								NewEntries.Add(ArgumentEntry);
+							}
+							BindingEntry->AddChild(ArgumentEntry.ToSharedRef());
+						}
 					}
 				}
 			}
@@ -1265,10 +1324,7 @@ void SBindingsList::Refresh()
 	if (TreeView.IsValid())
 	{
 		TreeView->RequestTreeRefresh();
-
-		TreeView->SetItemSelection(NewSelectedEntries, true);
-
-		for (const TSharedPtr<FBindingEntry>& Entry : RootWidgets)
+		for (const TSharedPtr<FBindingEntry>& Entry : NewEntries)
 		{
 			Private::ExpandAll(TreeView, Entry);
 		}
@@ -1283,7 +1339,7 @@ TSharedRef<ITableRow> SBindingsList::GenerateEntryRow(TSharedPtr<FBindingEntry>
 	{
 		switch (Entry->GetRowType())
 		{
-			case FBindingEntry::ERowType::Widget:
+			case FBindingEntry::ERowType::Group:
 			{
 				Row = SNew(SWidgetRow, OwnerTable, Entry, MVVMExtensionPtr->GetWidgetBlueprint());
 				break;
diff --git a/Engine/Plugins/Runtime/ModelViewViewModel/Source/ModelViewViewModelEditor/Private/Widgets/SMVVMViewBindingListView.h b/Engine/Plugins/Runtime/ModelViewViewModel/Source/ModelViewViewModelEditor/Private/Widgets/SMVVMViewBindingListView.h
--- a/Engine/Plugins/Runtime/ModelViewViewModel/Source/ModelViewViewModelEditor/Private/Widgets/SMVVMViewBindingListView.h
+++ b/Engine/Plugins/Runtime/ModelViewViewModel/Source/ModelViewViewModelEditor/Private/Widgets/SMVVMViewBindingListView.h
@@ -43,7 +43,7 @@ private:
 private:
 	TWeakPtr<SBindingsPanel> BindingPanel;
 	TSharedPtr<STreeView<TSharedPtr<FBindingEntry>>> TreeView;
-	TArray<TSharedPtr<FBindingEntry>> RootWidgets;
+	TArray<TSharedPtr<FBindingEntry>> RootGroups;
 	TWeakObjectPtr<UMVVMWidgetBlueprintExtension_View> MVVMExtension;
 	mutable bool bSelectionChangedGuard = false;
 };
