MVVM: Prevent crash with undo or when a binding is removed. The DetailView used to point directly to the binding in the view. WHen the view is modified, the DetailView was pointing to a dangling pointer. Also remove the customizaer because they don't respect the selection rules.

diff --git a/Engine/Plugins/Runtime/ModelViewViewModel/Source/ModelViewViewModelBlueprint/Public/MVVMBlueprintViewBinding.h b/Engine/Plugins/Runtime/ModelViewViewModel/Source/ModelViewViewModelBlueprint/Public/MVVMBlueprintViewBinding.h
--- a/Engine/Plugins/Runtime/ModelViewViewModel/Source/ModelViewViewModelBlueprint/Public/MVVMBlueprintViewBinding.h
+++ b/Engine/Plugins/Runtime/ModelViewViewModel/Source/ModelViewViewModelBlueprint/Public/MVVMBlueprintViewBinding.h
@@ -20,14 +20,14 @@ struct MODELVIEWVIEWMODELBLUEPRINT_API FMVVMBlueprintViewConversionPath
 	GENERATED_BODY()
 
 	/** The Conversion function when converting the value from the destination to the source. */
-	UPROPERTY(EditAnywhere, Category = "Viewmodel", AdvancedDisplay)
+	UPROPERTY(VisibleAnywhere, Category = "Viewmodel", AdvancedDisplay)
 	FMemberReference DestinationToSourceFunction;
 	
 	UPROPERTY()
 	FName DestinationToSourceWrapper;
 
 	/** The Conversion function when converting the value from the source to the destination. */
-	UPROPERTY(EditAnywhere, Category = "Viewmodel", AdvancedDisplay)
+	UPROPERTY(VisibleAnywhere, Category = "Viewmodel", AdvancedDisplay)
 	FMemberReference SourceToDestinationFunction;
 
 	UPROPERTY()
@@ -42,10 +42,10 @@ struct MODELVIEWVIEWMODELBLUEPRINT_API FMVVMBlueprintViewBinding
 {
 	GENERATED_BODY()
 
-	UPROPERTY(EditAnywhere, Category = "Viewmodel")
+	UPROPERTY(VisibleAnywhere, Category = "Viewmodel")
 	FMVVMBlueprintPropertyPath SourcePath;
 
-	UPROPERTY(EditAnywhere, Category = "Viewmodel")
+	UPROPERTY(VisibleAnywhere, Category = "Viewmodel")
 	FMVVMBlueprintPropertyPath DestinationPath;
 
 	/** */
@@ -59,7 +59,7 @@ struct MODELVIEWVIEWMODELBLUEPRINT_API FMVVMBlueprintViewBinding
 	EMVVMExecutionMode OverrideExecutionMode = EMVVMExecutionMode::Immediate;
 
 	/** */
-	UPROPERTY(EditAnywhere, Category = "Viewmodel")
+	UPROPERTY(VisibleAnywhere, Category = "Viewmodel")
 	FMVVMBlueprintViewConversionPath Conversion;
 
 	/** The unique ID of this binding. */
diff --git a/Engine/Plugins/Runtime/ModelViewViewModel/Source/ModelViewViewModelBlueprint/Public/MVVMPropertyPath.h b/Engine/Plugins/Runtime/ModelViewViewModel/Source/ModelViewViewModelBlueprint/Public/MVVMPropertyPath.h
--- a/Engine/Plugins/Runtime/ModelViewViewModel/Source/ModelViewViewModelBlueprint/Public/MVVMPropertyPath.h
+++ b/Engine/Plugins/Runtime/ModelViewViewModel/Source/ModelViewViewModelBlueprint/Public/MVVMPropertyPath.h
@@ -210,21 +210,21 @@ struct FMVVMBlueprintPropertyPath
 
 private:
 	/** Reference to property for this binding. */
-	UPROPERTY(EditAnywhere, Category = "MVVM")
+	UPROPERTY(VisibleAnywhere, Category = "MVVM")
 	TArray<FMVVMBlueprintFieldPath> Paths;
 
-	UPROPERTY(EditAnywhere, Category = "MVVM")
+	UPROPERTY(VisibleAnywhere, Category = "MVVM")
 	FName WidgetName;
 
-	UPROPERTY(EditAnywhere, Category = "MVVM")
+	UPROPERTY(VisibleAnywhere, Category = "MVVM")
 	FGuid ContextId;
 
 #if WITH_EDITORONLY_DATA
 	// Use the Paths. BindingReference and BindingKind are deprecated.
 	UPROPERTY()
-	FMemberReference BindingReference;
+	FMemberReference BindingReference_DEPRECATED;
 	UPROPERTY()
-	EBindingKind BindingKind = EBindingKind::Function;
+	EBindingKind BindingKind_DEPRECATED = EBindingKind::Function;
 #endif
 
 public:
@@ -343,7 +343,7 @@ public:
 
 	bool IsEmpty() const
 	{
-		return !IsFromWidget() && !IsFromViewModel() && BindingReference.GetMemberName() == FName();
+		return !IsFromWidget() && !IsFromViewModel() && BindingReference_DEPRECATED.GetMemberName() == FName();
 	}
 
 	/**
@@ -380,10 +380,10 @@ public:
 	{
 		if (Ar.IsLoading())
 		{
-			if (!BindingReference.GetMemberName().IsNone())
+			if (!BindingReference_DEPRECATED.GetMemberName().IsNone())
 			{
-				Paths.AddDefaulted_GetRef().SetDeprecatedBindingReference(BindingReference, BindingKind);
-				BindingReference = FMemberReference();
+				Paths.AddDefaulted_GetRef().SetDeprecatedBindingReference(BindingReference_DEPRECATED, BindingKind_DEPRECATED);
+				BindingReference_DEPRECATED = FMemberReference();
 			}
 		}
 	}
diff --git a/Engine/Plugins/Runtime/ModelViewViewModel/Source/ModelViewViewModelEditor/Private/Customizations/MVVMConversionPathCustomization.cpp b/Engine/Plugins/Runtime/ModelViewViewModel/Source/ModelViewViewModelEditor/Private/Customizations/MVVMConversionPathCustomization.cpp
deleted file mode 100644
diff --git a/Engine/Plugins/Runtime/ModelViewViewModel/Source/ModelViewViewModelEditor/Private/Customizations/MVVMConversionPathCustomization.h b/Engine/Plugins/Runtime/ModelViewViewModel/Source/ModelViewViewModelEditor/Private/Customizations/MVVMConversionPathCustomization.h
deleted file mode 100644
diff --git a/Engine/Plugins/Runtime/ModelViewViewModel/Source/ModelViewViewModelEditor/Private/Customizations/MVVMPropertyPathCustomization.cpp b/Engine/Plugins/Runtime/ModelViewViewModel/Source/ModelViewViewModelEditor/Private/Customizations/MVVMPropertyPathCustomization.cpp
deleted file mode 100644
diff --git a/Engine/Plugins/Runtime/ModelViewViewModel/Source/ModelViewViewModelEditor/Private/Customizations/MVVMPropertyPathCustomization.h b/Engine/Plugins/Runtime/ModelViewViewModel/Source/ModelViewViewModelEditor/Private/Customizations/MVVMPropertyPathCustomization.h
deleted file mode 100644
diff --git a/Engine/Plugins/Runtime/ModelViewViewModel/Source/ModelViewViewModelEditor/Private/Widgets/SMVVMViewBindingPanel.cpp b/Engine/Plugins/Runtime/ModelViewViewModel/Source/ModelViewViewModelEditor/Private/Widgets/SMVVMViewBindingPanel.cpp
--- a/Engine/Plugins/Runtime/ModelViewViewModel/Source/ModelViewViewModelEditor/Private/Widgets/SMVVMViewBindingPanel.cpp
+++ b/Engine/Plugins/Runtime/ModelViewViewModel/Source/ModelViewViewModelEditor/Private/Widgets/SMVVMViewBindingPanel.cpp
@@ -9,8 +9,6 @@
 #include "MVVMBlueprintView.h"
 #include "MVVMDeveloperProjectSettings.h"
 #include "MVVMWidgetBlueprintExtension_View.h"
-#include "Customizations/MVVMConversionPathCustomization.h"
-#include "Customizations/MVVMPropertyPathCustomization.h"
 
 #include "IStructureDetailsView.h"
 #include "PropertyEditorModule.h"
@@ -36,6 +34,31 @@
 namespace UE::MVVM
 {
 
+namespace Private
+{
+struct FStructDetailNotifyHook : FNotifyHook
+{
+	virtual ~FStructDetailNotifyHook() = default;
+
+	virtual void NotifyPostChange(const FPropertyChangedEvent& PropertyChangedEvent, FProperty* PropertyThatChanged) override
+	{
+		if (Binding.BindingId.IsValid())
+		{
+			if (UMVVMWidgetBlueprintExtension_View* Extension = MVVMExtension.Get())
+			{
+				if (FMVVMBlueprintViewBinding* CurrentBinding = Extension->GetBlueprintView()->GetBinding(Binding.BindingId))
+				{
+					*CurrentBinding = Binding;
+				}
+			}
+		}
+	}
+
+	FMVVMBlueprintViewBinding Binding;
+	TWeakObjectPtr<UMVVMWidgetBlueprintExtension_View> MVVMExtension;
+};
+}
+
 /** */
 void SBindingsPanel::Construct(const FArguments& InArgs, TSharedPtr<FWidgetBlueprintEditor> WidgetBlueprintEditor, bool bInIsDrawerTab)
 {
@@ -57,6 +80,9 @@ void SBindingsPanel::Construct(const FArguments& InArgs, TSharedPtr<FWidgetBluep
 	}
 
 	{
+		NotifyHook = MakePimpl<Private::FStructDetailNotifyHook>();
+		NotifyHook->MVVMExtension = MVVMExtensionPtr;
+
 		// Connection Settings
 		FPropertyEditorModule& PropertyEditorModule = FModuleManager::GetModuleChecked<FPropertyEditorModule>("PropertyEditor");
 
@@ -67,17 +93,15 @@ void SBindingsPanel::Construct(const FArguments& InArgs, TSharedPtr<FWidgetBluep
 		DetailsViewArgs.NameAreaSettings = FDetailsViewArgs::HideNameArea;
 		DetailsViewArgs.ViewIdentifier = NAME_None;
 		DetailsView = PropertyEditorModule.CreateDetailView(DetailsViewArgs);
-		
+		DetailsViewArgs.NotifyHook = NotifyHook.Get();
+
 		FStructureDetailsViewArgs StructureDetailsViewArgs;
-		StructDetailsView = PropertyEditorModule.CreateStructureDetailView(DetailsViewArgs, StructureDetailsViewArgs, TSharedPtr<FStructOnScope>());
-		StructDetailsView->GetDetailsView()->RegisterInstancedCustomPropertyTypeLayout(FMVVMBlueprintPropertyPath::StaticStruct()->GetFName(), FOnGetPropertyTypeCustomizationInstance::CreateStatic(&UE::MVVM::FPropertyPathCustomization::MakeInstance, WidgetBlueprint));
-		StructDetailsView->GetDetailsView()->RegisterInstancedCustomPropertyTypeLayout(FMVVMBlueprintViewConversionPath::StaticStruct()->GetFName(), FOnGetPropertyTypeCustomizationInstance::CreateStatic(&UE::MVVM::FConversionPathCustomization::MakeInstance, WidgetBlueprint));
+		StructDetailsView = PropertyEditorModule.CreateStructureDetailView(DetailsViewArgs, StructureDetailsViewArgs, MakeShared<FStructOnScope>(FMVVMBlueprintViewBinding::StaticStruct(), reinterpret_cast<uint8*>(&NotifyHook->Binding)));
 	}
 
 	HandleBlueprintViewChangedDelegate();
 }
 
-
 SBindingsPanel::~SBindingsPanel()
 {
 	if (UMVVMWidgetBlueprintExtension_View* Extension = MVVMExtension.Get())
@@ -121,13 +145,12 @@ void SBindingsPanel::HandleBlueprintViewChangedDelegate()
 	];
 }
 
-
 void SBindingsPanel::OnBindingListSelectionChanged(TConstArrayView<FMVVMBlueprintViewBinding*> Selection)
 {
+	NotifyHook->Binding = FMVVMBlueprintViewBinding();
 	if (FMVVMBlueprintViewBinding* Binding = Selection.Num() > 0 ? Selection.Last() : nullptr)
 	{
-		TSharedRef<FStructOnScope> StructScope = MakeShared<FStructOnScope>(FMVVMBlueprintViewBinding::StaticStruct(), reinterpret_cast<uint8*>(Binding));
-		StructDetailsView->SetStructureData(StructScope);
+		NotifyHook->Binding = *Binding;
 		DetailContainer->SetContent(StructDetailsView->GetWidget().ToSharedRef());
 	}
 	else
@@ -138,6 +161,21 @@ void SBindingsPanel::OnBindingListSelectionChanged(TConstArrayView<FMVVMBlueprin
 	}
 }
 
+void SBindingsPanel::Tick(const FGeometry& AllottedGeometry, const double InCurrentTime, const float InDeltaTime)
+{
+	if (NotifyHook->Binding.BindingId.IsValid())
+	{
+		if (UMVVMWidgetBlueprintExtension_View* Extension = MVVMExtension.Get())
+		{
+			if (const FMVVMBlueprintViewBinding* Binding = Extension->GetBlueprintView()->GetBinding(NotifyHook->Binding.BindingId))
+			{
+				NotifyHook->Binding = *Binding;
+			}
+		}
+	}
+
+	Super::Tick(AllottedGeometry, InCurrentTime, InDeltaTime);
+}
 
 FReply SBindingsPanel::AddDefaultBinding()
 {
@@ -259,6 +297,7 @@ TSharedRef<SWidget> SBindingsPanel::GenerateSettingsMenu()
 	UMVVMWidgetBlueprintExtension_View* MVVMExtensionPtr = MVVMExtension.Get();
 	DetailsView->SetObject(MVVMExtensionPtr && MVVMExtensionPtr->GetBlueprintView() ? MVVMExtensionPtr->GetBlueprintView() : nullptr);
 	DetailContainer->SetContent(DetailsView.ToSharedRef());
+	NotifyHook->Binding = FMVVMBlueprintViewBinding();
 
 	UWidgetBlueprintToolMenuContext* WidgetBlueprintMenuContext = NewObject<UWidgetBlueprintToolMenuContext>();
 	WidgetBlueprintMenuContext->WidgetBlueprintEditor = WeakBlueprintEditor;
@@ -436,7 +475,7 @@ void SBindingsPanel::ToggleDetailsVisibility()
 	if (DetailContainer.IsValid())
 	{
 		EVisibility NewVisibility = GetDetailsVisibleCheckState() == ECheckBoxState::Checked ? EVisibility::Collapsed : EVisibility::Visible;
-		return DetailContainer->SetVisibility(NewVisibility);
+		DetailContainer->SetVisibility(NewVisibility);
 	}
 }
 
diff --git a/Engine/Plugins/Runtime/ModelViewViewModel/Source/ModelViewViewModelEditor/Private/Widgets/SMVVMViewBindingPanel.h b/Engine/Plugins/Runtime/ModelViewViewModel/Source/ModelViewViewModelEditor/Private/Widgets/SMVVMViewBindingPanel.h
--- a/Engine/Plugins/Runtime/ModelViewViewModel/Source/ModelViewViewModelEditor/Private/Widgets/SMVVMViewBindingPanel.h
+++ b/Engine/Plugins/Runtime/ModelViewViewModel/Source/ModelViewViewModelEditor/Private/Widgets/SMVVMViewBindingPanel.h
@@ -19,9 +19,13 @@ enum class ECheckBoxState : uint8;
 namespace UE::MVVM
 {
 class SBindingsList;
+namespace Private { struct FStructDetailNotifyHook; }
 
 class SBindingsPanel : public SCompoundWidget
 {
+private:
+	using Super = SCompoundWidget;
+
 public:
 	SLATE_BEGIN_ARGS(SBindingsPanel) {}
 	SLATE_END_ARGS()
@@ -32,6 +36,7 @@ public:
 	//~ Begin SWidget Interface
 	virtual FReply OnKeyDown(const FGeometry& MyGeometry, const FKeyEvent& InKeyEvent) override;
 	virtual bool SupportsKeyboardFocus() const override;
+	virtual void Tick(const FGeometry& AllottedGeometry, const double InCurrentTime, const float InDeltaTime) override;
 	//~ End SWidget Interface
 
 	void OnBindingListSelectionChanged(TConstArrayView<FMVVMBlueprintViewBinding*> Selection);
@@ -68,6 +73,7 @@ private:
 	TSharedPtr<IDetailsView> DetailsView;
 	TSharedPtr<IStructureDetailsView> StructDetailsView;
 	TWeakObjectPtr<UMVVMWidgetBlueprintExtension_View> MVVMExtension;
+	TPimplPtr<Private::FStructDetailNotifyHook> NotifyHook;
 	FDelegateHandle BlueprintViewChangedDelegateHandle;
 	bool bIsDrawerTab;
 };
