MVVM: Refactor the field picker to update when the value changes. The widget creates the widget and only recreate it when the binding changes.

diff --git a/Engine/Plugins/Runtime/ModelViewViewModel/Source/ModelViewViewModel/Public/Types/MVVMFieldVariant.h b/Engine/Plugins/Runtime/ModelViewViewModel/Source/ModelViewViewModel/Public/Types/MVVMFieldVariant.h
--- a/Engine/Plugins/Runtime/ModelViewViewModel/Source/ModelViewViewModel/Public/Types/MVVMFieldVariant.h
+++ b/Engine/Plugins/Runtime/ModelViewViewModel/Source/ModelViewViewModel/Public/Types/MVVMFieldVariant.h
@@ -121,6 +121,32 @@ namespace UE::MVVM
 			return nullptr;
 		}
 
+		UE_NODISCARD bool HasAnyFlags(EObjectFlags InFlags) const
+		{
+			if (Type == EVariantType::Property)
+			{
+				return Binding.Property ? Binding.Property->HasAnyFlags(InFlags) : false;
+			}
+			else if (Type == EVariantType::Function)
+			{
+				return Binding.Function ? Binding.Function->HasAnyFlags(InFlags) : false;
+			}
+			return false;
+		}
+
+		UE_NODISCARD bool IsValid() const
+		{
+			if (Type == EVariantType::Property)
+			{
+				return Binding.Property != nullptr;
+			}
+			else if (Type == EVariantType::Function)
+			{
+				return Binding.Function != nullptr;
+			}
+			return false;
+		}
+
 		UE_NODISCARD bool IsEmpty() const
 		{
 			return Type == EVariantType::Empty || Binding.Function == nullptr;
diff --git a/Engine/Plugins/Runtime/ModelViewViewModel/Source/ModelViewViewModelEditor/Private/MVVMEditorSubsystem.cpp b/Engine/Plugins/Runtime/ModelViewViewModel/Source/ModelViewViewModelEditor/Private/MVVMEditorSubsystem.cpp
--- a/Engine/Plugins/Runtime/ModelViewViewModel/Source/ModelViewViewModelEditor/Private/MVVMEditorSubsystem.cpp
+++ b/Engine/Plugins/Runtime/ModelViewViewModel/Source/ModelViewViewModelEditor/Private/MVVMEditorSubsystem.cpp
@@ -175,9 +175,12 @@ UMVVMBlueprintView* UMVVMEditorSubsystem::RequestView(UWidgetBlueprint* WidgetBl
 
 UMVVMBlueprintView* UMVVMEditorSubsystem::GetView(const UWidgetBlueprint* WidgetBlueprint) const
 {
-	if (UMVVMWidgetBlueprintExtension_View* ExtensionView = UMVVMWidgetBlueprintExtension_View::GetExtension<UMVVMWidgetBlueprintExtension_View>(WidgetBlueprint))
+	if (WidgetBlueprint)
 	{
-		return ExtensionView->GetBlueprintView();
+		if (UMVVMWidgetBlueprintExtension_View* ExtensionView = UMVVMWidgetBlueprintExtension_View::GetExtension<UMVVMWidgetBlueprintExtension_View>(WidgetBlueprint))
+		{
+			return ExtensionView->GetBlueprintView();
+		}
 	}
 	return nullptr;
 }
diff --git a/Engine/Plugins/Runtime/ModelViewViewModel/Source/ModelViewViewModelEditor/Private/Types/MVVMBindingSource.cpp b/Engine/Plugins/Runtime/ModelViewViewModel/Source/ModelViewViewModelEditor/Private/Types/MVVMBindingSource.cpp
--- a/Engine/Plugins/Runtime/ModelViewViewModel/Source/ModelViewViewModelEditor/Private/Types/MVVMBindingSource.cpp
+++ b/Engine/Plugins/Runtime/ModelViewViewModel/Source/ModelViewViewModelEditor/Private/Types/MVVMBindingSource.cpp
@@ -4,6 +4,7 @@
 #include "Blueprint/WidgetTree.h"
 #include "Editor.h"
 #include "MVVMBlueprintView.h"
+#include "MVVMPropertyPath.h"
 #include "MVVMEditorSubsystem.h"
 #include "Types/MVVMBindingName.h"
 #include "WidgetBlueprint.h"
@@ -83,3 +84,16 @@ UE::MVVM::FBindingSource UE::MVVM::FBindingSource::CreateForViewModel(const UWid
 	return Source;
 }
 
+
+UE::MVVM::FBindingSource UE::MVVM::FBindingSource::CreateFromPropertyPath(const UWidgetBlueprint* WidgetBlueprint, const FMVVMBlueprintPropertyPath& Path)
+{
+	if (Path.IsFromViewModel())
+	{
+		return FBindingSource::CreateForViewModel(WidgetBlueprint, Path.GetViewModelId());
+	}
+	else if (Path.IsFromWidget())
+	{
+		return FBindingSource::CreateForWidget(WidgetBlueprint, Path.GetWidgetName());
+	}
+	return FBindingSource();
+}
diff --git a/Engine/Plugins/Runtime/ModelViewViewModel/Source/ModelViewViewModelEditor/Private/Types/MVVMBindingSource.h b/Engine/Plugins/Runtime/ModelViewViewModel/Source/ModelViewViewModelEditor/Private/Types/MVVMBindingSource.h
--- a/Engine/Plugins/Runtime/ModelViewViewModel/Source/ModelViewViewModelEditor/Private/Types/MVVMBindingSource.h
+++ b/Engine/Plugins/Runtime/ModelViewViewModel/Source/ModelViewViewModelEditor/Private/Types/MVVMBindingSource.h
@@ -6,54 +6,55 @@
 #include "UObject/Class.h"
 
 struct FMVVMBindingName;
-
+struct FMVVMBlueprintPropertyPath;
 class UWidgetBlueprint;
 
 namespace UE::MVVM
 {
-	struct FBindingSource
+struct FBindingSource
+{
+	FGuid ViewModelId;
+	FName Name;
+	TWeakObjectPtr<UClass> Class;
+	FText DisplayName;
+
+	bool operator==(const FBindingSource& Other) const
 	{
-		FGuid ViewModelId;
-		FName Name;
-		TWeakObjectPtr<UClass> Class;
-		FText DisplayName;
-
-		bool operator==(const FBindingSource& Other) const
-		{
-			return ViewModelId == Other.ViewModelId && Name == Other.Name && Class == Other.Class;
-		}
-
-		bool operator!=(const FBindingSource& Other) const
-		{
-			return !(operator==(Other));
-		}
-
-		friend int32 GetTypeHash(const FBindingSource& Source)
-		{
-			uint32 Hash = HashCombine(GetTypeHash(Source.ViewModelId), GetTypeHash(Source.Name));
-			Hash = HashCombine(Hash, GetTypeHash(Source.Class));
-			return Hash;
-		}
-
-		bool IsValid() const
-		{
-			return Class != nullptr && (ViewModelId.IsValid() || !Name.IsNone());
-		}
-
-		void Reset()
-		{
-			ViewModelId = FGuid();
-			Name = FName();
-			Class = nullptr;
-			DisplayName = FText::GetEmpty();
-		}
-
-		FMVVMBindingName ToBindingName(const UWidgetBlueprint* WidgetBlueprint) const;
-
-		static FBindingSource CreateForWidget(const UWidgetBlueprint* WidgetBlueprint, FName WidgetName);
-		static FBindingSource CreateForViewModel(const UWidgetBlueprint* WidgetBlueprint, FGuid ViewModelId);
-		static FBindingSource CreateForViewModel(const UWidgetBlueprint* WidgetBlueprint, FName ViewModelName);
-	};
+		return ViewModelId == Other.ViewModelId && Name == Other.Name && Class == Other.Class;
+	}
+
+	bool operator!=(const FBindingSource& Other) const
+	{
+		return !(operator==(Other));
+	}
+
+	friend int32 GetTypeHash(const FBindingSource& Source)
+	{
+		uint32 Hash = HashCombine(GetTypeHash(Source.ViewModelId), GetTypeHash(Source.Name));
+		Hash = HashCombine(Hash, GetTypeHash(Source.Class));
+		return Hash;
+	}
+
+	bool IsValid() const
+	{
+		return Class != nullptr && (ViewModelId.IsValid() || !Name.IsNone());
+	}
+
+	void Reset()
+	{
+		ViewModelId = FGuid();
+		Name = FName();
+		Class = nullptr;
+		DisplayName = FText::GetEmpty();
+	}
+
+	FMVVMBindingName ToBindingName(const UWidgetBlueprint* WidgetBlueprint) const;
+
+	static FBindingSource CreateForWidget(const UWidgetBlueprint* WidgetBlueprint, FName WidgetName);
+	static FBindingSource CreateForViewModel(const UWidgetBlueprint* WidgetBlueprint, FGuid ViewModelId);
+	static FBindingSource CreateForViewModel(const UWidgetBlueprint* WidgetBlueprint, FName ViewModelName);
+	static FBindingSource CreateFromPropertyPath(const UWidgetBlueprint* WidgetBlueprint, const FMVVMBlueprintPropertyPath& Path);
+};
 }
 
 template <>
diff --git a/Engine/Plugins/Runtime/ModelViewViewModel/Source/ModelViewViewModelEditor/Private/Widgets/SMVVMCachedViewBindingConversionFunction.cpp b/Engine/Plugins/Runtime/ModelViewViewModel/Source/ModelViewViewModelEditor/Private/Widgets/SMVVMCachedViewBindingConversionFunction.cpp
new file mode 100644
--- /dev/null
+++ b/Engine/Plugins/Runtime/ModelViewViewModel/Source/ModelViewViewModelEditor/Private/Widgets/SMVVMCachedViewBindingConversionFunction.cpp
@@ -0,0 +1,59 @@
+// Copyright Epic Games, Inc. All Rights Reserved.
+
+#include "Widgets/SMVVMCachedViewBindingConversionFunction.h"
+
+#include "Editor.h"
+#include "MVVMBlueprintView.h"
+#include "MVVMEditorSubsystem.h"
+#include "Styling/MVVMEditorStyle.h"
+#include "WidgetBlueprint.h"
+#include "Widgets/SMVVMPropertyPath.h"
+#include "Widgets/Images/SImage.h"
+#include "Widgets/Text/STextBlock.h"
+
+#define LOCTEXT_NAMESPACE "SCachedViewBindingConverisonFunction"
+
+namespace UE::MVVM
+{
+
+void SCachedViewBindingConversionFunction::Construct(const FArguments& InArgs, const UWidgetBlueprint* InWidgetBlueprint)
+{
+	WidgetBlueprint = InWidgetBlueprint;
+	OnGetConversionFunction = InArgs._OnGetConversionFunction;
+
+	ChildSlot
+	[
+		SNew(SHorizontalBox)
+		.ToolTipText_Lambda([this]() 
+		{
+			const UFunction* FoundFunction = OnGetConversionFunction.Execute();
+			return FoundFunction != nullptr ? FoundFunction->GetToolTipText() : FText::GetEmpty();
+		})
+		+ SHorizontalBox::Slot()
+		.HAlign(HAlign_Left)
+		.VAlign(VAlign_Center)
+		.Padding(0, 0, 4.0f, 0)
+		.AutoWidth()
+		[
+			SNew(SImage)
+			.DesiredSizeOverride(FVector2D(16.0f, 16.0f))
+			.Image(FAppStyle::Get().GetBrush("GraphEditor.Function_16x"))
+		]
+		+ SHorizontalBox::Slot()
+		.HAlign(HAlign_Fill)
+		.VAlign(VAlign_Center)
+		[
+			SNew(STextBlock)
+			.Text_Lambda([this]() 
+			{ 
+				const UFunction* FoundFunction = OnGetConversionFunction.Execute();
+				return FoundFunction != nullptr ? FoundFunction->GetDisplayNameText() : FText::GetEmpty();
+			})
+		]
+	];
+}
+
+
+} // namespace UE::MVVM
+
+#undef LOCTEXT_NAMESPACE
diff --git a/Engine/Plugins/Runtime/ModelViewViewModel/Source/ModelViewViewModelEditor/Private/Widgets/SMVVMCachedViewBindingConversionFunction.h b/Engine/Plugins/Runtime/ModelViewViewModel/Source/ModelViewViewModelEditor/Private/Widgets/SMVVMCachedViewBindingConversionFunction.h
new file mode 100644
--- /dev/null
+++ b/Engine/Plugins/Runtime/ModelViewViewModel/Source/ModelViewViewModelEditor/Private/Widgets/SMVVMCachedViewBindingConversionFunction.h
@@ -0,0 +1,35 @@
+// Copyright Epic Games, Inc. All Rights Reserved.
+
+#pragma once
+
+#include "Widgets/SMVVMCachedViewBindingPropertyPath.h"
+
+class SHorizontalBox;
+class UWidgetBlueprint;
+
+namespace UE::MVVM
+{
+DECLARE_DELEGATE_RetVal(const UFunction*, FOnGetConversionFunction);
+
+class SCachedViewBindingConversionFunction : public SCompoundWidget
+{
+private:
+	using Super = SCompoundWidget;
+
+public:
+
+	SLATE_BEGIN_ARGS(SCachedViewBindingConversionFunction)
+		: _TextStyle(&FCoreStyle::Get().GetWidgetStyle<FTextBlockStyle>("NormalText"))
+	{}
+		SLATE_STYLE_ARGUMENT(FTextBlockStyle, TextStyle)
+		SLATE_EVENT(FOnGetConversionFunction, OnGetConversionFunction)
+	SLATE_END_ARGS()
+
+	void Construct(const FArguments& InArgs, const UWidgetBlueprint* WidgetBlueprint);
+
+private:
+	TWeakObjectPtr<const UWidgetBlueprint> WidgetBlueprint;
+	FOnGetConversionFunction OnGetConversionFunction;
+};
+
+} // namespace UE::MVVM
diff --git a/Engine/Plugins/Runtime/ModelViewViewModel/Source/ModelViewViewModelEditor/Private/Widgets/SMVVMCachedViewBindingPropertyPath.cpp b/Engine/Plugins/Runtime/ModelViewViewModel/Source/ModelViewViewModelEditor/Private/Widgets/SMVVMCachedViewBindingPropertyPath.cpp
new file mode 100644
--- /dev/null
+++ b/Engine/Plugins/Runtime/ModelViewViewModel/Source/ModelViewViewModelEditor/Private/Widgets/SMVVMCachedViewBindingPropertyPath.cpp
@@ -0,0 +1,53 @@
+// Copyright Epic Games, Inc. All Rights Reserved.
+
+#include "Widgets/SMVVMCachedViewBindingPropertyPath.h"
+
+#include "Editor.h"
+#include "MVVMBlueprintView.h"
+#include "MVVMEditorSubsystem.h"
+#include "Styling/MVVMEditorStyle.h"
+#include "WidgetBlueprint.h"
+#include "Widgets/SMVVMPropertyPath.h"
+
+#define LOCTEXT_NAMESPACE "SCachedViewBindingPropertyPath"
+
+namespace UE::MVVM
+{
+
+void SCachedViewBindingPropertyPath::Construct(const FArguments& InArgs, const UWidgetBlueprint* InWidgetBlueprint)
+{
+	SetCanTick(true);
+
+	WidgetBlueprint = InWidgetBlueprint;
+	OnGetPropertyPath = InArgs._OnGetPropertyPath;
+
+	const FMVVMBlueprintPropertyPath PropertyPath = OnGetPropertyPath.Execute();
+	CachedPropertyPath = PropertyPath.GetFields();
+
+	ChildSlot
+	[
+		SAssignNew(PropertyPathWidget, SPropertyPath, InWidgetBlueprint)
+		.TextStyle(InArgs._TextStyle)
+		.PropertyPath(PropertyPath)
+		.ShowContext(InArgs._ShowContext)
+		.ShowOnlyLastPath(InArgs._ShowOnlyLastPath)
+	];
+}
+
+void SCachedViewBindingPropertyPath::Tick(const FGeometry& AllottedGeometry, const double InCurrentTime, const float InDeltaTime)
+{
+	Super::Tick(AllottedGeometry, InCurrentTime, InDeltaTime);
+
+	const FMVVMBlueprintPropertyPath PropertyPath = OnGetPropertyPath.Execute();
+	TArray<FMVVMConstFieldVariant> NewCachedPropertyPath = PropertyPath.GetFields();
+	if (NewCachedPropertyPath != CachedPropertyPath)
+	{
+		CachedPropertyPath = MoveTemp(NewCachedPropertyPath);
+		PropertyPathWidget->SetPropertyPath(PropertyPath);
+	}
+}
+
+
+} // namespace UE::MVVM
+
+#undef LOCTEXT_NAMESPACE
diff --git a/Engine/Plugins/Runtime/ModelViewViewModel/Source/ModelViewViewModelEditor/Private/Widgets/SMVVMCachedViewBindingPropertyPath.h b/Engine/Plugins/Runtime/ModelViewViewModel/Source/ModelViewViewModelEditor/Private/Widgets/SMVVMCachedViewBindingPropertyPath.h
new file mode 100644
--- /dev/null
+++ b/Engine/Plugins/Runtime/ModelViewViewModel/Source/ModelViewViewModelEditor/Private/Widgets/SMVVMCachedViewBindingPropertyPath.h
@@ -0,0 +1,45 @@
+// Copyright Epic Games, Inc. All Rights Reserved.
+
+#pragma once
+
+#include "MVVMPropertyPath.h"
+#include "Styling/CoreStyle.h"
+#include "Styling/SlateTypes.h"
+#include "Styling/SlateWidgetStyleAsset.h"
+#include "Widgets/SCompoundWidget.h"
+
+class SHorizontalBox;
+class UWidgetBlueprint;
+
+namespace UE::MVVM
+{
+class SPropertyPath;
+
+DECLARE_DELEGATE_RetVal(FMVVMBlueprintPropertyPath, FOnGetPropertyPath);
+
+class SCachedViewBindingPropertyPath : public SCompoundWidget
+{
+private:
+	using Super = SCompoundWidget;
+
+public:
+	SLATE_BEGIN_ARGS(SCachedViewBindingPropertyPath)
+		: _TextStyle(&FCoreStyle::Get().GetWidgetStyle<FTextBlockStyle>("NormalText"))
+	{}
+		SLATE_STYLE_ARGUMENT(FTextBlockStyle, TextStyle)
+		SLATE_EVENT(FOnGetPropertyPath, OnGetPropertyPath)
+		SLATE_ARGUMENT_DEFAULT(bool, ShowContext) = true;
+		SLATE_ARGUMENT_DEFAULT(bool, ShowOnlyLastPath) = false;
+	SLATE_END_ARGS()
+
+	void Construct(const FArguments& InArgs, const UWidgetBlueprint* WidgetBlueprint);
+	virtual void Tick(const FGeometry& AllottedGeometry, const double InCurrentTime, const float InDeltaTime) override;
+
+private:
+	TWeakObjectPtr<const UWidgetBlueprint> WidgetBlueprint;
+	TSharedPtr<SPropertyPath> PropertyPathWidget;
+	FOnGetPropertyPath OnGetPropertyPath;
+	TArray<FMVVMConstFieldVariant> CachedPropertyPath;
+};
+
+} // namespace UE::MVVM
diff --git a/Engine/Plugins/Runtime/ModelViewViewModel/Source/ModelViewViewModelEditor/Private/Widgets/SMVVMFieldEntry.cpp b/Engine/Plugins/Runtime/ModelViewViewModel/Source/ModelViewViewModelEditor/Private/Widgets/SMVVMFieldEntry.cpp
--- a/Engine/Plugins/Runtime/ModelViewViewModel/Source/ModelViewViewModelEditor/Private/Widgets/SMVVMFieldEntry.cpp
+++ b/Engine/Plugins/Runtime/ModelViewViewModel/Source/ModelViewViewModelEditor/Private/Widgets/SMVVMFieldEntry.cpp
@@ -1,6 +1,6 @@
 // Copyright Epic Games, Inc. All Rights Reserved.
 
-#include "SMVVMFieldEntry.h"
+#include "Widgets/SMVVMFieldEntry.h"
 
 #include "Bindings/MVVMBindingHelper.h"
 #include "SMVVMFieldIcon.h"
@@ -123,9 +123,8 @@ FText GetPathToolTip(TConstArrayView<FMVVMConstFieldVariant> Path)
 
 } // namespace Private
 
-void SFieldEntry::Construct(const FArguments& InArgs)
+void SFieldPaths::Construct(const FArguments& InArgs)
 {
-	Field = InArgs._Field;
 	TextStyle = InArgs._TextStyle;
 
 	ChildSlot
@@ -133,67 +132,79 @@ void SFieldEntry::Construct(const FArguments& InArgs)
 		SAssignNew(FieldBox, SHorizontalBox)
 	];
 
-	Refresh();
+	SetFieldPaths(InArgs._FieldPaths);
 }
 
-void SFieldEntry::Refresh()
+void SFieldPaths::SetFieldPaths(TArrayView<UE::MVVM::FMVVMConstFieldVariant> InPropertyPath)
 {
 	FieldBox->ClearChildren();
 
-	TArray<FMVVMConstFieldVariant> Fields = Field.GetFields();
-
 	int32 Index = 0;
-	if (bShowOnlyLast && Fields.Num() > 0)
+	if (bShowOnlyLast && InPropertyPath.Num() > 0)
 	{
-		Index = Fields.Num() - 1;
+		Index = InPropertyPath.Num() - 1;
 	}
 
-	for (; Index < Fields.Num(); ++Index)
+	for (; Index < InPropertyPath.Num(); ++Index)
 	{
-		FieldBox->AddSlot()
-		.HAlign(HAlign_Left)
-		.VAlign(VAlign_Center)
-		.AutoWidth()
-		[
-			SNew(SFieldIcon)
-			.Field(Fields[Index])
-		];
-
-		FieldBox->AddSlot()
-		.HAlign(HAlign_Fill)
-		.VAlign(VAlign_Center)
-		.Padding(4, 0, 0, 0)
-		.AutoWidth()
-		[
-			SNew(STextBlock)
-			.TextStyle(TextStyle)
-			.Clipping(EWidgetClipping::OnDemand)
-			.Text(Private::GetFieldDisplayName(Fields[Index]))
-		];
+		bool bIsValid = InPropertyPath[Index].IsValid();
+		bool bIsFieldValid = bIsValid && !InPropertyPath[Index].HasAnyFlags(RF_NewerVersionExists);
+		UStruct* Struct = InPropertyPath[Index].GetOwner();
+		bool bIsOwnerValid = bIsValid && Struct && !Struct->HasAnyFlags(RF_NewerVersionExists);
+
+		if (bIsFieldValid && bIsOwnerValid)
+		{
+			FieldBox->AddSlot()
+				.HAlign(HAlign_Left)
+				.VAlign(VAlign_Center)
+				.AutoWidth()
+				[
+					SNew(SFieldIcon)
+					.Field(InPropertyPath[Index])
+				];
+
+			FieldBox->AddSlot()
+				.HAlign(HAlign_Fill)
+				.VAlign(VAlign_Center)
+				.Padding(4, 0, 0, 0)
+				.AutoWidth()
+				[
+					SNew(STextBlock)
+					.TextStyle(TextStyle)
+					.Clipping(EWidgetClipping::OnDemand)
+					.Text(Private::GetFieldDisplayName(InPropertyPath[Index]))
+				];
+		}
+		else
+		{
+			FieldBox->AddSlot()
+				.HAlign(HAlign_Left)
+				.VAlign(VAlign_Center)
+				.AutoWidth()
+				[
+					SNew(STextBlock)
+					.TextStyle(TextStyle)
+					.Clipping(EWidgetClipping::OnDemand)
+					.Text(LOCTEXT("Invalid", "Invalid Field"))
+				];
+		}
 
 		// if not the last, then we need to add a chevron separator
-		if (Index != Fields.Num() - 1)
+		if (Index != InPropertyPath.Num() - 1)
 		{
 			FieldBox->AddSlot()
-			.HAlign(HAlign_Left)
-			.VAlign(VAlign_Center)
-			.Padding(6, 0)
-			.AutoWidth()
-			[
-				SNew(SImage)
-				.Image(FAppStyle::Get().GetBrush("Icons.ChevronRight"))
-			];
+				.HAlign(HAlign_Left)
+				.VAlign(VAlign_Center)
+				.Padding(6, 0)
+				.AutoWidth()
+				[
+					SNew(SImage)
+					.Image(FAppStyle::Get().GetBrush("Icons.ChevronRight"))
+				];
 		}
 	}
 
-	SetToolTipText(Private::GetPathToolTip(Fields));
-}
-
-void SFieldEntry::SetField(const FMVVMBlueprintPropertyPath& InField)
-{
-	Field = InField;
-
-	Refresh();
+	SetToolTipText(Private::GetPathToolTip(InPropertyPath));
 } 
 
 } // namespace UE::MVVM
diff --git a/Engine/Plugins/Runtime/ModelViewViewModel/Source/ModelViewViewModelEditor/Private/Widgets/SMVVMFieldEntry.h b/Engine/Plugins/Runtime/ModelViewViewModel/Source/ModelViewViewModelEditor/Private/Widgets/SMVVMFieldEntry.h
--- a/Engine/Plugins/Runtime/ModelViewViewModel/Source/ModelViewViewModelEditor/Private/Widgets/SMVVMFieldEntry.h
+++ b/Engine/Plugins/Runtime/ModelViewViewModel/Source/ModelViewViewModelEditor/Private/Widgets/SMVVMFieldEntry.h
@@ -17,27 +17,23 @@ namespace UE::MVVM
 
 class SFieldIcon;
 
-using FIsFieldValidResult = TValueOrError<bool, FString>;
-DECLARE_DELEGATE_RetVal_OneParam(FIsFieldValidResult, FIsFieldValid, FMVVMBlueprintPropertyPath);
-
-class SFieldEntry : public SCompoundWidget
+/** Show the TArray<FMVVMBlueprintFieldPath>. */
+class SFieldPaths : public SCompoundWidget
 {
 public:
-	SLATE_BEGIN_ARGS(SFieldEntry) :
-		_TextStyle(&FCoreStyle::Get().GetWidgetStyle<FTextBlockStyle>("NormalText"))
+	SLATE_BEGIN_ARGS(SFieldPaths)
+		: _TextStyle(&FCoreStyle::Get().GetWidgetStyle<FTextBlockStyle>("NormalText"))
 	{}
 		SLATE_STYLE_ARGUMENT(FTextBlockStyle, TextStyle)
-		SLATE_ARGUMENT(FMVVMBlueprintPropertyPath, Field)
+		SLATE_ARGUMENT(TArrayView<UE::MVVM::FMVVMConstFieldVariant>, FieldPaths)
 		SLATE_ARGUMENT(bool, ShowOnlyLast)
 	SLATE_END_ARGS()
 
 	void Construct(const FArguments& InArgs);
-	void Refresh();
-	void SetField(const FMVVMBlueprintPropertyPath& InField);
+	void SetFieldPaths(TArrayView<UE::MVVM::FMVVMConstFieldVariant> InField);
 
 private:
 	const FTextBlockStyle* TextStyle = nullptr;
-	FMVVMBlueprintPropertyPath Field;
 	TSharedPtr<SHorizontalBox> FieldBox;
 	bool bShowOnlyLast = false;
 };
diff --git a/Engine/Plugins/Runtime/ModelViewViewModel/Source/ModelViewViewModelEditor/Private/Widgets/SMVVMFieldSelector.cpp b/Engine/Plugins/Runtime/ModelViewViewModel/Source/ModelViewViewModelEditor/Private/Widgets/SMVVMFieldSelector.cpp
--- a/Engine/Plugins/Runtime/ModelViewViewModel/Source/ModelViewViewModelEditor/Private/Widgets/SMVVMFieldSelector.cpp
+++ b/Engine/Plugins/Runtime/ModelViewViewModel/Source/ModelViewViewModelEditor/Private/Widgets/SMVVMFieldSelector.cpp
@@ -2,25 +2,14 @@
 
 #include "Widgets/SMVVMFieldSelector.h"
 
-#include "Bindings/MVVMBindingHelper.h"
-#include "Editor.h"
-#include "Framework/Views/TableViewMetadata.h"
-#include "Hierarchy/SReadOnlyHierarchyView.h"
-#include "MVVMEditorSubsystem.h"
-#include "MVVMDeveloperProjectSettings.h"
-#include "MVVMWidgetBlueprintExtension_View.h"
-#include "SPrimaryButton.h"
+#include "WidgetBlueprint.h"
+
 #include "Styling/MVVMEditorStyle.h"
-#include "Styling/SlateIconFinder.h"
-#include "Types/MVVMBindingMode.h"
-#include "Widgets/SMVVMFieldEntry.h"
-#include "Widgets/SMVVMSourceEntry.h"
-#include "Widgets/Input/SCheckBox.h"
+
 #include "Widgets/Input/SComboButton.h"
-#include "Widgets/Input/SSearchBox.h"
-#include "Widgets/Layout/SScrollBox.h"
-#include "Widgets/SMVVMViewModelBindingListWidget.h"
-#include "Widgets/Views/STreeView.h"
+#include "Widgets/Layout/SBox.h"
+#include "Widgets/Layout/SWidgetSwitcher.h"
+#include "Widgets/Text/STextBlock.h"
 
 #define LOCTEXT_NAMESPACE "MVVMFieldSelector"
 
@@ -32,1289 +21,167 @@ namespace Private
 
 FBindingSource GetSourceFromPath(const UWidgetBlueprint* WidgetBlueprint, const FMVVMBlueprintPropertyPath& Path)
 {
-	if (Path.IsFromViewModel())
-	{
-		return FBindingSource::CreateForViewModel(WidgetBlueprint, Path.GetViewModelId());
-	}
-	else if (Path.IsFromWidget())
-	{
-		return FBindingSource::CreateForWidget(WidgetBlueprint, Path.GetWidgetName());
-	}
-	return FBindingSource();
+	return FBindingSource::CreateFromPropertyPath(WidgetBlueprint, Path);
 }
 
 } // namespace Private
 
-void SFieldSelector::Construct(const FArguments& InArgs, const UWidgetBlueprint* InWidgetBlueprint, bool bInViewModelProperty)
+void SFieldSelector::Construct(const FArguments& InArgs, const UWidgetBlueprint* InWidgetBlueprint)
 {
 	WidgetBlueprint = InWidgetBlueprint;
 	check(InWidgetBlueprint);
 	
-	BindingMode = InArgs._BindingMode;
-	check(BindingMode.IsSet());
-
-	bViewModelProperty = bInViewModelProperty;
-	AssignableTo = InArgs._AssignableTo;
 	TextStyle = InArgs._TextStyle;
-
-	OnFieldSelectionChangedDelegate = InArgs._OnFieldSelectionChanged;
-	check(OnFieldSelectionChangedDelegate.IsBound());
-
-	SelectedField = InArgs._SelectedField;
-	check(SelectedField.IsSet());
-
-	CachedSelectedField = SelectedField.Get();
-
-	ShowConversionFunctions = InArgs._ShowConversionFunctions;
-	SelectedConversionFunction = InArgs._SelectedConversionFunction;
-	CachedSelectedConversionFunction = SelectedConversionFunction.Get(nullptr);
-	OnConversionFunctionSelectionChangedDelegate = InArgs._OnConversionFunctionSelectionChanged;
-
-	TSharedPtr<SHorizontalBox> SourceEntryBox;
+	OnGetPropertyPath = InArgs._OnGetPropertyPath;
+	OnGetConversionFunction = InArgs._OnGetConversionFunction;
+	OnGetSelectionContext = InArgs._OnGetSelectionContext;
+	OnFieldSelectionChanged = InArgs._OnFieldSelectionChanged;
 
 	ChildSlot
 	[
-		SAssignNew(ComboButton, SComboButton)
-		.ComboButtonStyle(FMVVMEditorStyle::Get(), "FieldSelector.ComboButton")
-		.OnGetMenuContent(this, &SFieldSelector::OnGetMenuContent)
-		.ContentPadding(FMargin(4, 2))
-		.ButtonContent()
+		SNew(SBox)
+		.MinDesiredWidth(200)
 		[
-			SNew(SOverlay)
-
-			// is a property set?
-			+ SOverlay::Slot()
-			[
-				SAssignNew(SourceEntryBox, SHorizontalBox)
-				.Visibility_Lambda([this]() { return CachedSelectedConversionFunction == nullptr && !CachedSelectedField.IsEmpty() ? EVisibility::Visible : EVisibility::Collapsed; })
-			]
-
-			// is a conversion function set?
-			+ SOverlay::Slot()
-			[
-				SNew(SHorizontalBox)
-				.Visibility_Lambda([this]() 
-				{
-					return CachedSelectedConversionFunction != nullptr ? EVisibility::Visible : EVisibility::Collapsed; 
-				})
-				.ToolTipText_Lambda([this]() 
-				{
-					return CachedSelectedConversionFunction != nullptr ? CachedSelectedConversionFunction->GetToolTipText() : FText::GetEmpty();
-				})
-				+ SHorizontalBox::Slot()
+			SAssignNew(ComboButton, SComboButton)
+			.ComboButtonStyle(FMVVMEditorStyle::Get(), "FieldSelector.ComboButton")
+			.OnGetMenuContent(this, &SFieldSelector::HandleGetMenuContent)
+			.ContentPadding(FMargin(4.0f, 2.0f))
+			.ButtonContent()
+			[
+				SNew(SWidgetSwitcher)
+				.WidgetIndex(this, &SFieldSelector::GetCurrentDisplayIndex)
+				//0-Property/Function (from Widget or Viewmodel).
+				//0-Property/Function argument of conversion function
+				+ SWidgetSwitcher::Slot()
+				.Padding(0.0f, 0.0f, 8.0f, 0.0f)
 				.HAlign(HAlign_Left)
 				.VAlign(VAlign_Center)
-				.Padding(0, 0, 4.0f, 0)
-				.AutoWidth()
 				[
-					SNew(SImage)
-					.DesiredSizeOverride(FVector2D(16.0f, 16.0f))
-					.Image(FAppStyle::Get().GetBrush("GraphEditor.Function_16x"))
+					SNew(SCachedViewBindingPropertyPath, WidgetBlueprint.Get())
+					.TextStyle(TextStyle)
+					.ShowContext(InArgs._ShowContext)
+					.OnGetPropertyPath(InArgs._OnGetPropertyPath)
 				]
-				+ SHorizontalBox::Slot()
-				.HAlign(HAlign_Fill)
+
+				//1-Conversion Function
+				+ SWidgetSwitcher::Slot()
+				.Padding(0.0f, 0.0f, 8.0f, 0.0f)
+				.HAlign(HAlign_Left)
 				.VAlign(VAlign_Center)
 				[
-					SNew(STextBlock)
-					.Text_Lambda([this]() 
-					{ 
-						return CachedSelectedConversionFunction != nullptr ? CachedSelectedConversionFunction->GetDisplayNameText() : FText::GetEmpty(); 
-					})
+					SNew(SCachedViewBindingConversionFunction, WidgetBlueprint.Get())
+					.TextStyle(TextStyle)
+					.OnGetConversionFunction(InArgs._OnGetConversionFunction)
 				]
-			]
 
-			// nothing selected
-			+ SOverlay::Slot()
-			[
-				SNew(SBox)
-				.Padding(FMargin(8, 0, 8, 0))
+				//2-Nothing selected.
+				+ SWidgetSwitcher::Slot()
 				[
-					SNew(STextBlock)
-					.Visibility_Lambda([this]() 
-					{ 
-						return CachedSelectedField.IsEmpty() && CachedSelectedConversionFunction == nullptr ? EVisibility::Visible : EVisibility::Collapsed;
-					})
-					.TextStyle(FAppStyle::Get(), "HintText")
-					.Text(LOCTEXT("None", "No field selected"))
+					SNew(SBox)
+					.Padding(FMargin(8.0f, 0.0f, 8.0f, 0.0f))
+					[
+						SNew(STextBlock)
+						.TextStyle(FAppStyle::Get(), "HintText")
+						.Text(LOCTEXT("None", "No field selected"))
+					]
 				]
 			]
 		]
 	];
-	
-	if (InArgs._Source.IsValid())
-	{
-		FixedSource = InArgs._Source;
-	}
-
-	if (InArgs._ShowSource)
-	{
-		SourceEntryBox->AddSlot()
-			.Padding(8, 0, 0, 0)
-			.HAlign(HAlign_Left)
-			.VAlign(VAlign_Center)
-			.AutoWidth()
-			[
-				SAssignNew(SelectedSourceWidget, SSourceEntry)
-			];
-
-		SourceEntryBox->AddSlot()
-			.Padding(6, 0)
-			.HAlign(HAlign_Left)
-			.VAlign(VAlign_Center)
-			.AutoWidth()
-			[
-				SNew(SImage)
-				.Image(FAppStyle::Get().GetBrush("Icons.ChevronRight"))
-				.Visibility_Lambda([this]() { return CachedSelectedField.HasPaths() ? EVisibility::Visible : EVisibility::Collapsed; })
-			];
-	}
-
-	SourceEntryBox->AddSlot()
-		.Padding(0, 0, 8, 0)
-		.HAlign(HAlign_Left)
-		.VAlign(VAlign_Center)
-		.AutoWidth()
-		[
-			SAssignNew(SelectedEntryWidget, SFieldEntry)
-			.TextStyle(TextStyle)
-		];
-
-	Refresh();
-}
-
-void SFieldSelector::Refresh()
-{
-	CachedSelectedField = SelectedField.Get();
-	CachedSelectedConversionFunction = SelectedConversionFunction.Get(nullptr);
-
-	if (SelectedEntryWidget.IsValid())
-	{
-		SelectedEntryWidget->SetField(CachedSelectedField);
-	}
-
-	if (SelectedSourceWidget.IsValid() && WidgetBlueprint.IsValid())
-	{
-		SelectedSourceWidget->RefreshSource(Private::GetSourceFromPath(WidgetBlueprint.Get(), CachedSelectedField));
-	}
-}
-
-TSharedRef<SWidget> SFieldSelector::OnGenerateFieldWidget(FMVVMBlueprintPropertyPath Path) const
-{
-	return SNew(SFieldEntry)
-		.TextStyle(TextStyle)
-		.Field(Path);
-}
-
-bool SFieldSelector::IsSelectEnabled() const
-{
-	if (BindingList.IsValid())
-	{
-		FMVVMBlueprintPropertyPath Path = BindingList->GetSelectedProperty();
-		if (Path.IsFromViewModel() || Path.IsFromWidget())
-		{
-			return true;
-		}
-	}
-
-	if (ConversionFunctionList.IsValid())
-	{
-		if (ConversionFunctionList->GetSelectedItems().Num() > 0)
-		{
-			return true;
-		}
-	}
-
-	return false;
-}
-
-FReply SFieldSelector::OnSelectClicked()
-{
-	if (BindingList.IsValid())
-	{
-		SetPropertySelection(BindingList->GetSelectedProperty());
-	}
-
-	if (ConversionFunctionList.IsValid())
-	{
-		TArray<const UFunction*> Selection = ConversionFunctionList->GetSelectedItems();
-		SetConversionFunctionSelection(Selection.Num() > 0 ? Selection[0] : nullptr);
-	}
-
-	return FReply::Handled();
-}
-
-bool SFieldSelector::IsClearEnabled() const
-{
-	return !CachedSelectedField.IsEmpty() || CachedSelectedConversionFunction != nullptr;
-}
-
-FReply SFieldSelector::OnClearClicked()
-{
-	FMVVMBlueprintPropertyPath NewProperty;
-	
-	if (FixedSource.IsSet())
-	{
-		FBindingSource Source = FixedSource.GetValue();
-		if (Source.ViewModelId.IsValid())
-		{
-			NewProperty.SetViewModelId(Source.ViewModelId);
-		}
-		else 
-		{
-			NewProperty.SetWidgetName(Source.Name);
-		}
-	}
-
-	SetPropertySelection(NewProperty);
-	SetConversionFunctionSelection(nullptr);
-	return FReply::Handled();
-}
-
-FReply SFieldSelector::OnCancelClicked()
-{
-	if (ComboButton.IsValid())
-	{
-		ComboButton->SetIsOpen(false);
-	}
-
-	return FReply::Handled();
-}
-
-void SFieldSelector::SetPropertySelection(const FMVVMBlueprintPropertyPath& SelectedPath)
-{
-	if (ComboButton.IsValid())
-	{
-		ComboButton->SetIsOpen(false);
-	}
-
-	OnFieldSelectionChangedDelegate.Execute(SelectedPath);
-	Refresh();
-}
-
-TSharedPtr<SFieldSelector::FConversionFunctionItem> SFieldSelector::FindConversionFunctionCategory(const TArray<TSharedPtr<FConversionFunctionItem>>& Items, TArrayView<FString> CategoryNameParts) const
-{
-	if (CategoryNameParts.Num() > 0)
-	{
-		for (const TSharedPtr<FConversionFunctionItem>& Item : Items)
-		{
-			if (Item->GetCategoryName() == CategoryNameParts[0])
-			{
-				TArrayView<FString> RemainingParts = CategoryNameParts.RightChop(1);
-
-				// last category part, this is what we're looking for
-				if (RemainingParts.Num() == 0)
-				{
-					return Item;
-				}
-
-				// recurse into children
-				return FindConversionFunctionCategory(Item->Children, RemainingParts);
-			}
-		}
-	}
-
-	return TSharedPtr<FConversionFunctionItem>();
-}
-
-void SFieldSelector::OnSearchBoxTextChanged(const FText& NewText)
-{
-	if (BindingList.IsValid())
-	{
-		BindingList->SetRawFilterText(NewText);
-	}
-
-	if (WidgetList.IsValid())
-	{
-		WidgetList->SetRawFilterText(NewText);
-	}
-
-	if (ConversionFunctionCategoryTree.IsValid())
-	{
-		TArray<TSharedPtr<FConversionFunctionItem>> OldSelectedCategories;
-		ConversionFunctionCategoryTree->GetSelectedItems(OldSelectedCategories);
-
-		FilterConversionFunctionCategories();
-		ConversionFunctionCategoryTree->RequestTreeRefresh();
-
-		TArray<TSharedPtr<FConversionFunctionItem>> NewSelectedCategories;
-
-		// reselect old selection
-		for (const TSharedPtr<FConversionFunctionItem>& OldItem : OldSelectedCategories)
-		{
-			if (TSharedPtr<FConversionFunctionItem> NewItem = FindConversionFunctionCategory(FilteredConversionFunctionRoot, OldItem->CategoryPath))
-			{
-				NewSelectedCategories.Add(NewItem);
-			}
-		}
-		
-		if (NewSelectedCategories.Num() > 0)
-		{
-			ConversionFunctionCategoryTree->SetItemSelection(NewSelectedCategories, true);
-		}
-	}
-
-	if (ConversionFunctionList.IsValid())
-	{
-		TArray<const UFunction*> OldSelectedFunctions;
-		ConversionFunctionList->GetSelectedItems(OldSelectedFunctions);
-
-		FilterConversionFunctions();
-		ConversionFunctionList->RequestListRefresh();
-
-		ConversionFunctionList->SetItemSelection(OldSelectedFunctions, true);
-	}
-}
-
-void SFieldSelector::OnViewModelSelected(FBindingSource Source, ESelectInfo::Type)
-{
-	if (!bIsMenuInitialized)
-	{
-		return;
-	}
-
-	if (BindingList.IsValid())
-	{
-		BindingList->ClearSources();
-	}
-
-	TArray<FBindingSource> Selection = ViewModelList->GetSelectedItems();
-	if (Selection.Num() == 0)
-	{
-		return;
-	}
-
-	if (ConversionFunctionCategoryTree.IsValid())
-	{
-		ConversionFunctionCategoryTree->ClearSelection();
-	}
-
-	if (BindingList.IsValid())
-	{
-		BindingList->AddSources(Selection);
-	}
-}
-
-void SFieldSelector::OnWidgetSelected(FName WidgetName, ESelectInfo::Type)
-{
-	if (!bIsMenuInitialized)
-	{
-		return;
-	}
-
-	if (BindingList.IsValid())
-	{
-		BindingList->ClearSources();
-	}
-
-	TArray<FName> Selection = WidgetList->GetSelectedWidgets();
-	if (Selection.Num() == 0)
-	{
-		return;
-	}
-
-	if (ConversionFunctionCategoryTree.IsValid())
-	{
-		ConversionFunctionCategoryTree->ClearSelection();
-	}
-
-	if (BindingList.IsValid())
-	{
-		const UWidgetBlueprint* WidgetBP = WidgetBlueprint.Get();
-
-		TArray<FBindingSource> Sources;
-		Algo::Transform(Selection, Sources, [WidgetBP](const FName& WidgetName)
-		{
-			return FBindingSource::CreateForWidget(WidgetBP, WidgetName);
-		});
-
-		BindingList->AddSources(Sources);
-	}
-}
-
-TSharedRef<ITableRow> SFieldSelector::GenerateViewModelRow(FBindingSource ViewModel, const TSharedRef<STableViewBase>& OwnerTable) const
-{
-	return SNew(STableRow<FBindingSource>, OwnerTable)
-		[
-			SNew(SHorizontalBox)
-			+ SHorizontalBox::Slot()
-			.Padding(4)
-			.HAlign(HAlign_Left)
-			.VAlign(VAlign_Center)
-			.AutoWidth()
-			[
-				SNew(SImage)
-				.Image(FSlateIconFinder::FindIconBrushForClass(ViewModel.Class.Get()))
-				.ColorAndOpacity(FSlateColor::UseForeground())
-			]
-			+ SHorizontalBox::Slot()
-			.HAlign(HAlign_Left)
-			.VAlign(VAlign_Center)
-			.AutoWidth()
-			[
-				SNew(STextBlock)
-				.Text(ViewModel.DisplayName)
-				.ColorAndOpacity(FSlateColor::UseForeground())
-				.HighlightText_Lambda([this]() { return SearchBox.IsValid() ? SearchBox->GetText() : FText::GetEmpty(); })
-			]
-		];
-}
-
-int32 SFieldSelector::FilterConversionFunctionCategoryChildren(const TArray<FString>& FilterStrings, const TArray<TSharedPtr<FConversionFunctionItem>>& SourceArray, TArray<TSharedPtr<FConversionFunctionItem>>& OutDestArray)
-{
-	int32 NumFunctions = 0;
-
-	for (const TSharedPtr<FConversionFunctionItem>& SourceItem : SourceArray)
-	{
-		TArray<TSharedPtr<FConversionFunctionItem>> FilteredChildren;
-
-		// check if our name matches the filters
-		bool bMatchesFilters = true;
-
-		const FString ItemName = SourceItem->Function != nullptr ? SourceItem->Function->GetName() : FString();
-		
-		for (const FString& Filter : FilterStrings)
-		{
-			if (!ItemName.Contains(Filter))
-			{
-				bMatchesFilters = false;
-				break;
-			}
-		}
-
-		int32 NumChildren = 0;
-		if (bMatchesFilters)
-		{
-			ensureAlways(SourceItem->Function != nullptr);
-			NumChildren = 1;
-		}
-		else
-		{
-			// if we don't match, then we still want to check all our children
-			NumChildren = FilterConversionFunctionCategoryChildren(FilterStrings, SourceItem->Children, FilteredChildren);
-		}
-
-		// then add this item to the destination array
-		if (NumChildren > 0)
-		{
-			TSharedPtr<FConversionFunctionItem>& NewItem = OutDestArray.Add_GetRef(MakeShared<FConversionFunctionItem>());
-			NewItem->CategoryPath = SourceItem->CategoryPath;
-			NewItem->Function = SourceItem->Function; 
-			NewItem->Children = FilteredChildren;
-			NewItem->NumFunctions = NumChildren;
-
-			NumFunctions += NewItem->NumFunctions;
-		}
-	}
-
-	return NumFunctions;
-}
-
-void SFieldSelector::FilterConversionFunctionCategories()
-{
-	FilteredConversionFunctionRoot.Reset();
-
-	TArray<FString> FilterStrings;
-	if (SearchBox.IsValid())
-	{
-		SearchBox->GetText().ToString().ParseIntoArrayWS(FilterStrings);
-	}
-
-	if (FilterStrings.IsEmpty())
-	{
-		// don't bother filtering if we don't have any search terms
-		FilteredConversionFunctionRoot = ConversionFunctionRoot;
-		return;
-	}
-
-	TSharedPtr<FConversionFunctionItem> RootItem = FilteredConversionFunctionRoot.Add_GetRef(MakeShared<FConversionFunctionItem>());
-	RootItem->CategoryPath = { TEXT("Conversion Functions") };
-	
-	RootItem->NumFunctions = FilterConversionFunctionCategoryChildren(FilterStrings, ConversionFunctionRoot[0]->Children, FilteredConversionFunctionRoot[0]->Children);
-	
-	ExpandAll(FilteredConversionFunctionRoot, true);
-}
-
-int32 SFieldSelector::SortConversionFunctionItemsRecursive(TArray<TSharedPtr<FConversionFunctionItem>>& Items)
-{
-	int32 NumFound = 0;
-
-	Items.Sort([](const TSharedPtr<FConversionFunctionItem>& A, const TSharedPtr<FConversionFunctionItem>& B)
-	{
-		if (!A->GetCategoryName().IsEmpty() && !B->GetCategoryName().IsEmpty())
-		{
-			return A->GetCategoryName() < B->GetCategoryName();
-		}
-		if (!A->GetCategoryName().IsEmpty() && B->GetCategoryName().IsEmpty())
-		{
-			return true;
-		}
-		if (A->GetCategoryName().IsEmpty() && !B->GetCategoryName().IsEmpty())
-		{
-			return false;
-		}
-		if (A->Function != nullptr && B->Function != nullptr)
-		{
-			return A->Function->GetDisplayNameText().CompareTo(B->Function->GetDisplayNameText()) <= 0;
-		}
-		return true;
-	});
-
-	for (const TSharedPtr<FConversionFunctionItem>& Item : Items)
-	{
-		NumFound += SortConversionFunctionItemsRecursive(Item->Children);
-
-		if (Item->Function != nullptr)
-		{
-			NumFound += 1;
-		}
-	}
-
-	return NumFound;
-}
-
-void SFieldSelector::GenerateConversionFunctionItems()
-{
-	UMVVMEditorSubsystem* Subsystem = GEditor->GetEditorSubsystem<UMVVMEditorSubsystem>();
-	TArray<UFunction*> AllConversionFunctions = Subsystem->GetAvailableConversionFunctions(WidgetBlueprint.Get(), FMVVMBlueprintPropertyPath(), FMVVMBlueprintPropertyPath());
-
-	// remove all incompatible conversion functions
-	for (int32 Idx = 0; Idx < AllConversionFunctions.Num(); ++Idx)
-	{
-		const UFunction* Function = AllConversionFunctions[Idx];
-		const FProperty* ReturnProperty = BindingHelper::GetReturnProperty(Function);
-		const FProperty* AssignToProperty = AssignableTo.Get(nullptr);
-
-		if (AssignToProperty != nullptr && !BindingHelper::ArePropertiesCompatible(ReturnProperty, AssignToProperty))
-		{
-			AllConversionFunctions.RemoveAt(Idx);
-			--Idx;
-		}
-	}
-
-	auto AddFunctionToItem = [](const UFunction* Function, const TSharedPtr<FConversionFunctionItem>& Parent)
-	{
-		TSharedPtr<FConversionFunctionItem>& Item = Parent->Children.Add_GetRef(MakeShared<FConversionFunctionItem>());
-		Item->Function = Function;
-		Item->NumFunctions = 1;
-		Parent->NumFunctions += 1;
-	};
-
-	TSharedPtr<FConversionFunctionItem> CurrentSelectedItem;
-
-	TArray<FString> CategoryPath;
-
-	ConversionFunctionRoot.Reset();
-	TSharedPtr<FConversionFunctionItem>& RootItem = ConversionFunctionRoot.Add_GetRef(MakeShared<FConversionFunctionItem>());
-	RootItem->CategoryPath = { TEXT("Conversion Functions") };
-
-	for (const UFunction* Function : AllConversionFunctions)
-	{
-		FText CategoryName = Function->GetMetaDataText("Category");
-		if (CategoryName.IsEmpty())
-		{
-			AddFunctionToItem(Function, RootItem);
-			continue;
-		}
-
-		// split into subcategories and trim
-		CategoryPath.Reset();
-		CategoryName.ToString().ParseIntoArray(CategoryPath, TEXT("|"));
-		for (FString& SubCategory : CategoryPath)
-		{
-			SubCategory.TrimStartAndEndInline();
-		}
-
-		TSharedPtr<FConversionFunctionItem> ParentItem = RootItem;
-
-		// create items for the entire category path
-		// eg. "Math|Boolean|AND" 
-		// Math 
-		//   > Boolean
-		//     > AND
-		for (int32 PathIndex = 0; PathIndex < CategoryPath.Num(); ++PathIndex)
-		{
-			ParentItem->NumFunctions += 1;
-
-			ParentItem = FindOrCreateItemForCategory(ParentItem->Children, MakeArrayView(CategoryPath.GetData(), PathIndex + 1));
-		}
-
-		AddFunctionToItem(Function, ParentItem);
-	}
-
-	int32 NumItems = SortConversionFunctionItemsRecursive(ConversionFunctionRoot);
-	ensure(NumItems == RootItem->NumFunctions);
 }
 
-void SFieldSelector::ExpandAllToItem(const UFunction* Function)
+int32 SFieldSelector::GetCurrentDisplayIndex() const
 {
-	TArray<TSharedPtr<FConversionFunctionItem>> Path;
-
-	FText FullCategoryName = Function->GetMetaDataText("Category");
-	if (FullCategoryName.IsEmpty())
-	{
-		Path.Add(FilteredConversionFunctionRoot[0]);
-	}
-	else
+	if (OnGetConversionFunction.IsBound())
 	{
-		TArray<FString> CategoryPath;
-		FullCategoryName.ToString().ParseIntoArray(CategoryPath, TEXT("|"));
-
-		TSharedPtr<FConversionFunctionItem> CurrentParent = FilteredConversionFunctionRoot[0];
-
-		for (const FString& SubCategory : CategoryPath)
+		if (OnGetConversionFunction.Execute() != nullptr)
 		{
-			const FString Trimmed = SubCategory.TrimStartAndEnd();
-
-			TSharedPtr<FConversionFunctionItem>* FoundItem = 
-				CurrentParent->Children.FindByPredicate([Trimmed, Function](const TSharedPtr<FConversionFunctionItem>& Item)
-				{
-					return Item->GetCategoryName() == Trimmed || Item->Function == Function;
-				});
-
-			if (FoundItem != nullptr)
-			{
-				Path.Add(*FoundItem);
-				CurrentParent = *FoundItem;
-			}
+			return 1;
 		}
 	}
-
-	if (Path.Num() > 0)
+	if (OnGetPropertyPath.IsBound())
 	{
-		ConversionFunctionCategoryTree->SetItemExpansion(FilteredConversionFunctionRoot[0], true);
-
-		ExpandAll(Path, false);
-	}
-}
-
-void SFieldSelector::ExpandAll(const TArray<TSharedPtr<FConversionFunctionItem>>& Items, bool bRecursive)
-{
-	for (const TSharedPtr<FConversionFunctionItem>& Item : Items)
-	{
-		ConversionFunctionCategoryTree->SetItemExpansion(Item, true);
-
-		if (bRecursive)
+		if (!OnGetPropertyPath.Execute().IsEmpty())
 		{
-			ExpandAll(Item->Children, bRecursive);
+			return 0;
 		}
 	}
+	return 2;
 }
 
-void SFieldSelector::FilterViewModels()
-{
-	FilteredViewModelSources.Reset();
-
-	TArray<FString> FilterStrings;
-	if (SearchBox.IsValid())
-	{
-		SearchBox->GetText().ToString().ParseIntoArrayWS(FilterStrings);
-	}
-
-	if (FilterStrings.IsEmpty())
-	{
-		FilteredViewModelSources = ViewModelSources;
-		return;
-	}
-
-	for (const FBindingSource& ViewModel : ViewModelSources)
-	{
-		const FString DisplayName = ViewModel.DisplayName.ToString();
-		const FString ClassName = ViewModel.Class != nullptr ? ViewModel.Class->GetName() : FString();
-		const FString Name = ViewModel.Name.ToString();
-
-		bool bMatchesFilters = true;
-		
-		for (const FString& Filter : FilterStrings)
-		{
-			if (!DisplayName.Contains(Filter) && 
-				!ClassName.Contains(Filter) &&
-				!Name.Contains(Filter))
-			{
-				bMatchesFilters = false;
-				break;
-			}
-		}
-
-		if (bMatchesFilters)
-		{
-			FBindingSource& NewSource = FilteredViewModelSources.Add_GetRef(FBindingSource());
-			NewSource.Class = ViewModel.Class;
-			NewSource.Name = ViewModel.Name;
-			NewSource.DisplayName = ViewModel.DisplayName;
-			NewSource.ViewModelId = ViewModel.ViewModelId;
-		}
-	}
-}
+//void SFieldSelector::Refresh()
+//{
+//	CachedSelectedField = SelectedField.Get();
+//	CachedSelectedConversionFunction = SelectedConversionFunction.Get(nullptr);
+//
+//	//if (PropertyPathWidget.IsValid())
+//	//{
+//	//	PropertyPathWidget->SetPropertyPath(CachedSelectedField);
+//	//}
+//
+//	//if (SelectedSourceWidget.IsValid() && WidgetBlueprint.IsValid())
+//	//{
+//	//	SelectedSourceWidget->RefreshSource(Private::GetSourceFromPath(WidgetBlueprint.Get(), CachedSelectedField));
+//	//}
+//}
 
-TSharedRef<SWidget> SFieldSelector::CreateSourcePanel()
+TSharedRef<SWidget> SFieldSelector::HandleGetMenuContent()
 {
-	if (!WidgetBlueprint.IsValid())
+	const UWidgetBlueprint* WidgetBlueprintPtr = WidgetBlueprint.Get();
+	if (!WidgetBlueprintPtr)
 	{
 		return SNullWidget::NullWidget;
 	}
 
-	// show source picker
-	TSharedPtr<SWidget> SourcePicker;
-	if (bViewModelProperty)
-	{
-		if (FixedSource.IsSet())
-		{
-			ViewModelSources.Add(FixedSource.GetValue());
-		}
-		else
-		{
-			UMVVMEditorSubsystem* EditorSubsystem = GEditor->GetEditorSubsystem<UMVVMEditorSubsystem>();
-			ViewModelSources = EditorSubsystem->GetAllViewModels(WidgetBlueprint.Get());
-		}
-
-		FilterViewModels();
-
-		ViewModelList = SNew(SListView<FBindingSource>)
-			.ListItemsSource(&FilteredViewModelSources)
-			.OnGenerateRow(this, &SFieldSelector::GenerateViewModelRow)
-			.SelectionMode(ESelectionMode::Multi)
-			.OnSelectionChanged(this, &SFieldSelector::OnViewModelSelected);
-
-		FBindingSource SelectedSource;
-		for (const FBindingSource& Source : FilteredViewModelSources)
-		{
-			if (Source.ViewModelId == CachedSelectedField.GetViewModelId())
-			{
-				SelectedSource = Source;
-			}
-		}
-
-		if (SelectedSource.IsValid())
-		{
-			ViewModelList->SetItemSelection(SelectedSource, true);
-		}
-
-		SourcePicker = ViewModelList;
-	}
-	else
-	{
-		TArray<FName> ShowOnly;
-		if (FixedSource.IsSet())
-		{
-			ShowOnly.Add(FixedSource.GetValue().Name);
-		}
-
-		WidgetList = SNew(SReadOnlyHierarchyView, WidgetBlueprint.Get())
-			.OnSelectionChanged(this, &SFieldSelector::OnWidgetSelected)
-			.SelectionMode(ESelectionMode::Multi)
-			.ShowSearch(false)
-			.ShowOnly(ShowOnly);
-
-		if (CachedSelectedField.IsFromWidget())
-		{
-			WidgetList->SetSelectedWidget(CachedSelectedField.GetWidgetName());
-		}
-
-		SourcePicker = WidgetList;
-	}
-
-	if (ShowConversionFunctions.Get(false))
-	{
-		GenerateConversionFunctionItems();
-
-		FilterConversionFunctionCategories();
-		
-		TSharedRef<SWidget> StackedSourcePicker =
-				SNew(SVerticalBox)
-				+ SVerticalBox::Slot()
-				.AutoHeight()
-				[
-					SourcePicker.ToSharedRef()
-				]
-				+ SVerticalBox::Slot()
-				.AutoHeight()
-				[
-					SAssignNew(ConversionFunctionCategoryTree, STreeView<TSharedPtr<FConversionFunctionItem>>)
-					.SelectionMode(ESelectionMode::Multi)
-					.TreeItemsSource(&FilteredConversionFunctionRoot)
-					.OnGenerateRow(this, &SFieldSelector::GenerateConversionFunctionCategoryRow)
-					.OnSelectionChanged(this, &SFieldSelector::OnConversionFunctionCategorySelected)
-					.OnGetChildren(this, &SFieldSelector::GetConversionFunctionCategoryChildren)
-				];
-
-		SourcePicker = StackedSourcePicker;
-	}
-
-	return SourcePicker.ToSharedRef();
-}
-
-TSharedRef<SWidget> SFieldSelector::OnGetMenuContent()
-{
-	if (!WidgetBlueprint.IsValid())
-	{
-		return SNullWidget::NullWidget;
-	}
-
-	bIsMenuInitialized = false;
-
-	ConversionFunctionCategoryTree.Reset();
-	FilteredConversionFunctionRoot.Reset();
-	ConversionFunctions.Reset();
-	FilteredConversionFunctions.Reset();
-	ConversionFunctionRoot.Reset();
-	ViewModelSources.Reset();
-	ViewModelList.Reset();
-	WidgetList.Reset();
-
-	TSharedRef<SVerticalBox> VBox = SNew(SVerticalBox)
-		+ SVerticalBox::Slot()
-		.HAlign(HAlign_Fill)
-		.VAlign(VAlign_Top)
-		.Padding(0, 4, 0, 4)
-		.AutoHeight()
-		[
-			SNew(SHorizontalBox)
-			+ SHorizontalBox::Slot()
-			.FillWidth(1.f)
-			.Padding(3.f, 0.f, 3.f, 0.f)	
-			[
-				SAssignNew(SearchBox, SSearchBox)
-				.OnTextChanged(this, &SFieldSelector::OnSearchBoxTextChanged)
-			]
-			+ SHorizontalBox::Slot()
-			.HAlign(HAlign_Right)
-			.VAlign(VAlign_Center)
-			.AutoWidth()
-			[
-				SNew(SCheckBox)
-				.OnCheckStateChanged(this, &SFieldSelector::OnContextToggleChanged)
-				.IsChecked(this, &SFieldSelector::ContextToggleIsChecked)
-				.ToolTipText(LOCTEXT("ContextFlagToolTip", "Should the list be filtered to only properties and functions that make sense in the current context?"))
-				[
-					SNew(STextBlock)
-					.Text(LOCTEXT("MVVMContextSensitiveToggle", "Context Sensitive"))
-				]
-			]
-		];
-
-	// If we're showing conversion functions, we don't want to set the AssignableTo property of SSourceBindingList, because then it will only show exact matches, 
-	// and since we're also showing conversion functions we know that's not what the user wants.
-	// However, in the case we're not showing conversion functions, we want only exact matches.
-	const FProperty* AssignableToProperty = nullptr;
-	const bool bShowConversionFunctions = ShowConversionFunctions.Get(false);
-	if (!bShowConversionFunctions)
-	{
-		AssignableToProperty = AssignableTo.Get(nullptr);
-	}
-
-	BindingList = SNew(SSourceBindingList, WidgetBlueprint.Get())
-		.ShowSearchBox(false)
-		.OnDoubleClicked(this, &SFieldSelector::SetPropertySelection)
-		.FieldVisibilityFlags(GetFieldVisibilityFlags())
-		.AssignableTo(AssignableToProperty);
-
-	if (FixedSource.IsSet())
-	{
-		// Single fixed source, don't show the separate source panel.
-		BindingList->AddSource(FixedSource.GetValue());
-	}
-	else if (!CachedSelectedField.IsEmpty())
-	{
-		FBindingSource Source = CachedSelectedField.IsFromViewModel() ? 
-			FBindingSource::CreateForViewModel(WidgetBlueprint.Get(), CachedSelectedField.GetViewModelId()) : 
-			FBindingSource::CreateForWidget(WidgetBlueprint.Get(), CachedSelectedField.GetWidgetName());
-		BindingList->AddSource(Source);
-		BindingList->SetSelectedProperty(CachedSelectedField);
-	}
-	
-	TSharedPtr<SVerticalBox> BindingListVBox;
-		
-	VBox->AddSlot()
-		.FillHeight(1.0f)
-		[
-			SNew(SBorder)
-			.BorderImage(FAppStyle::Get().GetBrush("Brushes.Recessed"))
-			[
-				SNew(SSplitter)
-				.PhysicalSplitterHandleSize(4.0f)
-				+ SSplitter::Slot()
-				.Value(0.5f)
-				.MinSize(100)
-				[
-					SNew(SScrollBox)
-					+ SScrollBox::Slot()
-					[
-						CreateSourcePanel()
-					]
-				]
-				+ SSplitter::Slot()
-				.Value(0.5f)
-				.MinSize(100)
-				[
-					SNew(SScrollBox)
-					+ SScrollBox::Slot()
-					[
-						SAssignNew(BindingListVBox, SVerticalBox)
-						+ SVerticalBox::Slot()
-						.AutoHeight()
-						[
-							BindingList.ToSharedRef()
-						]
-					]
-				]
-			]
-		];
-
-	if (bShowConversionFunctions)
-	{
-		BindingListVBox->AddSlot()
-			.AutoHeight()
-			[
-				SAssignNew(ConversionFunctionList, SListView<const UFunction*>)
-				.SelectionMode(ESelectionMode::Single)
-				.ListItemsSource(&FilteredConversionFunctions)
-				.OnMouseButtonDoubleClick(this, &SFieldSelector::SetConversionFunctionSelection)
-				.OnGenerateRow(this, &SFieldSelector::GenerateConversionFunctionRow)
-			];
-
-		if (CachedSelectedConversionFunction != nullptr)
-		{
-			ExpandAllToItem(CachedSelectedConversionFunction);
-			ConversionFunctionList->SetItemSelection(CachedSelectedConversionFunction, true);
-		}
-	}
-
-	VBox->AddSlot()
-		.Padding(4, 4, 4, 0)
-		.HAlign(HAlign_Right)
-		.AutoHeight()
-		[
-			SNew(SHorizontalBox)
-			+ SHorizontalBox::Slot()
-			.AutoWidth()
-			[
-				SNew(SPrimaryButton)
-				.OnClicked(this, &SFieldSelector::OnSelectClicked)
-				.IsEnabled(this, &SFieldSelector::IsSelectEnabled)
-				.Text(LOCTEXT("Select", "Select"))
-			]
-			+ SHorizontalBox::Slot()
-			.Padding(FAppStyle::Get().GetMargin("StandardDialog.SlotPadding"))
-			.AutoWidth()
-			[
-				SNew(SButton)
-				.OnClicked(this, &SFieldSelector::OnClearClicked)
-				.IsEnabled(this, &SFieldSelector::IsClearEnabled)
-				.HAlign(HAlign_Center)
-				[
-					SNew(STextBlock)
-					.Text(LOCTEXT("Clear", "Clear"))
-				]
-			]
-			+ SHorizontalBox::Slot()
-			.Padding(FAppStyle::Get().GetMargin("StandardDialog.SlotPadding"))
-			.AutoWidth()
-			[
-				SNew(SButton)
-				.OnClicked(this, &SFieldSelector::OnCancelClicked)
-				.HAlign(HAlign_Center)
-				[
-					SNew(STextBlock)
-					.Text(LOCTEXT("Cancel", "Cancel"))
-				]
-			]
-		];
-
-	TSharedRef<SWidget> MenuWidget = 
-		SNew(SBox)
-		.MinDesiredWidth(400)
-		.MinDesiredHeight(200)
-		[
-			SNew(SBorder)
-			.BorderImage(FAppStyle::Get().GetBrush("ToolPanel.GroupBorder"))
-			.Padding(FMargin(8, 2, 8, 3))
-			[
-				VBox
-			]
-		];
-
-	ComboButton->SetMenuContentWidgetToFocus(SearchBox);
-
-	bIsMenuInitialized = true;
-	return MenuWidget;
-}
-
-EFieldVisibility SFieldSelector::GetFieldVisibilityFlags() const
-{
-	EFieldVisibility Flags = EFieldVisibility::None;
-
-	EMVVMBindingMode Mode = BindingMode.Get();
-	if (bViewModelProperty)
+	const UFunction* CurrentSelectedFunction = nullptr;
+	if (OnGetConversionFunction.IsBound())
 	{
-		if (IsForwardBinding(Mode))
-		{
-			Flags |= EFieldVisibility::Readable;
-
-			if (!IsOneTimeBinding(Mode))
-			{
-				Flags |= EFieldVisibility::Notify;
-			}
-		}
-		if (IsBackwardBinding(Mode))
-		{
-			Flags |= EFieldVisibility::Writable;
-		}
-	}
-	else
-	{
-		if (IsForwardBinding(Mode))
-		{
-			Flags |= EFieldVisibility::Writable;
-		}
-		if (IsBackwardBinding(Mode))
-		{
-			Flags |= EFieldVisibility::Readable;
-
-			if (!IsOneTimeBinding(Mode))
-			{
-				Flags |= EFieldVisibility::Notify;
-			}
-		}
+		CurrentSelectedFunction = OnGetConversionFunction.Execute();
 	}
-
-	return Flags;
-}
-
-void SFieldSelector::OnContextToggleChanged(ECheckBoxState CheckState)
-{
-	UMVVMWidgetBlueprintExtension_View* ExtensionView = UMVVMWidgetBlueprintExtension_View::GetExtension<UMVVMWidgetBlueprintExtension_View>(WidgetBlueprint.Get());
-	check(ExtensionView);
-
-	FMVVMViewBindingFilterSettings FilterSettings = ExtensionView->GetFilterSettings();
-	FilterSettings.FilterFlags = CheckState == ECheckBoxState::Checked ? EFilterFlag::All : EFilterFlag::None;
-	ExtensionView->SetFilterSettings(FilterSettings);
-
-	if (WidgetList.IsValid())
+	TOptional<FMVVMBlueprintPropertyPath> CurrentSelectedPropertyPath;
+	if (OnGetPropertyPath.IsBound())
 	{
-		OnWidgetSelected(WidgetBlueprint->GetFName(), ESelectInfo::Direct);
+		CurrentSelectedPropertyPath = OnGetPropertyPath.Execute();
 	}
 
-	if (ViewModelList.IsValid())
+	FFieldSelectionContext SelectionContext;
+	if (OnGetSelectionContext.IsBound())
 	{
-		OnViewModelSelected(FBindingSource(), ESelectInfo::Direct);
+		SelectionContext = OnGetSelectionContext.Execute();
 	}
-}
-
-ECheckBoxState SFieldSelector::ContextToggleIsChecked() const
-{
-	UMVVMWidgetBlueprintExtension_View* ExtensionView = UMVVMWidgetBlueprintExtension_View::GetExtension<UMVVMWidgetBlueprintExtension_View>(WidgetBlueprint.Get());
-	check(ExtensionView);
 
-	return ExtensionView->GetFilterSettings().FilterFlags == EFilterFlag::All ? ECheckBoxState::Checked : ECheckBoxState::Unchecked;
-}
-
-TSharedPtr<SFieldSelector::FConversionFunctionItem> SFieldSelector::FindOrCreateItemForCategory(TArray<TSharedPtr<FConversionFunctionItem>>& Items, TArrayView<FString> CategoryPath)
-{
-	check(CategoryPath.Num() > 0);
+	TSharedRef<SFieldSelectorMenu> Menu = SNew(SFieldSelectorMenu, WidgetBlueprintPtr)
+		.OnFieldSelectionChanged(this, &SFieldSelector::HandleFieldSelectionChanged)
+		.OnMenuCloseRequested(this, &SFieldSelector::HandleMenuClosed)
+		.SelectionContext(SelectionContext)
+		.CurrentPropertyPathSelected(CurrentSelectedPropertyPath)
+		.CurrentFunctionSelected(CurrentSelectedFunction)
+		;
 
-	const FString& CategoryName = CategoryPath.Last();
+	ComboButton->SetMenuContentWidgetToFocus(Menu->GetWidgetToFocus());
 
-	int32 Idx = 0;
-	for (; Idx < Items.Num(); ++Idx)
-	{
-		// found item
-		if (Items[Idx]->GetCategoryName() == CategoryName)
-		{
-			return Items[Idx];
-		}
-
-		// passed the place where it should have been, break out
-		if (Items[Idx]->GetCategoryName() > CategoryName)
-		{
-			break;
-		}
-	}
-
-	TSharedPtr<FConversionFunctionItem> NewItem = Items.Insert_GetRef(MakeShared<FConversionFunctionItem>(), Idx);
-	NewItem->CategoryPath = CategoryPath;
-	return NewItem;
+	return Menu;
 }
 
-TSharedRef<ITableRow> SFieldSelector::GenerateConversionFunctionCategoryRow(TSharedPtr<FConversionFunctionItem> Item, const TSharedRef<STableViewBase>& OwnerTable)
+void SFieldSelector::HandleFieldSelectionChanged(FMVVMBlueprintPropertyPath PropertyPath, const UFunction* Function)
 {
-	FText DisplayName = FText::FormatOrdered(FText::FromString("{0} ({1})"), FText::FromString(Item->GetCategoryName()), FText::FromString(LexToString(Item->NumFunctions)));
-
-	return SNew(STableRow<TSharedPtr<FConversionFunctionItem>>, OwnerTable)
-	[
-		SNew(SHorizontalBox)
-		+ SHorizontalBox::Slot()
-		.HAlign(HAlign_Left)
-		.VAlign(VAlign_Center)
-		.Padding(0, 2.0f, 4.0f, 2.0f)
-		.AutoWidth()
-		[
-			SNew(SImage)
-			.DesiredSizeOverride(FVector2D(16.0f, 16.0f))
-			.Image(FAppStyle::Get().GetBrush("GraphEditor.Function_16x"))
-		]
-		+ SHorizontalBox::Slot()
-		.HAlign(HAlign_Fill)
-		.VAlign(VAlign_Center)
-		[
-			SNew(STextBlock)
-			.Font(Item == FilteredConversionFunctionRoot[0] ? FAppStyle::Get().GetFontStyle("NormalText") : FAppStyle::Get().GetFontStyle("BoldFont"))
-			.Text(DisplayName)
-			.ToolTipText(FText::FromString(Item->GetCategoryName()))
-			.HighlightText_Lambda([this]() { return SearchBox.IsValid() ? SearchBox->GetText() : FText::GetEmpty(); })
-		]
-	];
-}
-
-void SFieldSelector::GetConversionFunctionCategoryChildren(TSharedPtr<FConversionFunctionItem> Item, TArray<TSharedPtr<FConversionFunctionItem>>& OutItems) const
-{
-	Algo::TransformIf(Item->Children, OutItems,
-		[](const TSharedPtr<FConversionFunctionItem>& Item)
-		{
-			return !Item->GetCategoryName().IsEmpty();
-		},
-		[](const TSharedPtr<FConversionFunctionItem>& Item)
-		{
-			return Item;
-		});
-}
-
-TSharedRef<ITableRow> SFieldSelector::GenerateConversionFunctionRow(const UFunction* Function, const TSharedRef<STableViewBase>& OwnerTable)
-{
-	return SNew(STableRow<const UFunction*>, OwnerTable)
-		[
-			SNew(SHorizontalBox)
-			+ SHorizontalBox::Slot()
-			.HAlign(HAlign_Left)
-			.VAlign(VAlign_Center)
-			.Padding(0, 0, 4.0f, 0)
-			.AutoWidth()
-			[
-				SNew(SImage)
-				.DesiredSizeOverride(FVector2D(16.0f, 16.0f))
-				.Image(FAppStyle::Get().GetBrush("GraphEditor.Function_16x"))
-			]
-			+ SHorizontalBox::Slot()
-			.HAlign(HAlign_Fill)
-			.VAlign(VAlign_Center)
-			[
-				SNew(STextBlock)
-				.Text(Function->GetDisplayNameText())
-				.ToolTipText(Function->GetToolTipText())
-				.HighlightText_Lambda([this]() { return SearchBox.IsValid() ? SearchBox->GetText() : FText::GetEmpty(); })
-			]
-		];
-}
-
-void SFieldSelector::AddConversionFunctionChildrenRecursive(const TSharedPtr<FConversionFunctionItem>& Parent, TArray<const UFunction*>& OutFunctions)
-{
-	for (const TSharedPtr<FConversionFunctionItem>& Item : Parent->Children)
-	{
-		if (Item->Function != nullptr)
-		{
-			int32 Index = 0;
-			for (; Index < OutFunctions.Num(); ++Index)
-			{
-				if (OutFunctions[Index]->GetName() > Item->Function->GetName())
-				{
-					break;
-				}
-			}
-
-			OutFunctions.Insert(Item->Function, Index);
-		}
-		else
-		{
-			AddConversionFunctionChildrenRecursive(Item, OutFunctions);
-		}
-	}
-}
-
-void SFieldSelector::FilterConversionFunctions()
-{
-	TArray<FString> FilterStrings;
-	if (SearchBox.IsValid())
-	{
-		SearchBox->GetText().ToString().ParseIntoArrayWS(FilterStrings);
-	}
-
-	if (FilterStrings.IsEmpty())
-	{
-		FilteredConversionFunctions = ConversionFunctions;
-		return;
-	}
-	
-	for (const UFunction* Function : ConversionFunctions)
-	{
-		const FString FunctionName = Function->GetName();
-
-		bool bMatches = true;
-		for (const FString& Filter : FilterStrings)
-		{
-			if (!FunctionName.Contains(Filter))
-			{
-				bMatches = false;
-				break;
-			}
-		}
-
-		if (bMatches)
-		{
-			FilteredConversionFunctions.Add(Function);
-		}
-	}
-}
-
-void SFieldSelector::OnConversionFunctionCategorySelected(TSharedPtr<FConversionFunctionItem> SelectedItem, ESelectInfo::Type)
-{
-	if (!bIsMenuInitialized)
-	{
-		return;
-	}
-
-	ConversionFunctions.Reset();
-	FilteredConversionFunctions.Reset();
-
-	TArray<TSharedPtr<FConversionFunctionItem>> SelectedItems = ConversionFunctionCategoryTree->GetSelectedItems();
-	if (SelectedItems.Num() > 0)
+	if (ComboButton.IsValid())
 	{
-		if (BindingList.IsValid())
-		{
-			BindingList->ClearSources();
-		}
-
-		if (ViewModelList.IsValid())
-		{
-			ViewModelList->ClearSelection();
-		}
-
-		if (WidgetList.IsValid())
-		{
-			WidgetList->ClearSelection();
-		}
-
-		for (const TSharedPtr<FConversionFunctionItem>& Item : SelectedItems)
-		{
-			AddConversionFunctionChildrenRecursive(Item, ConversionFunctions);
-		}
-
-		FilterConversionFunctions();
+		ComboButton->SetIsOpen(false);
 	}
 
-	if (ConversionFunctionList.IsValid())
+	if (OnFieldSelectionChanged.IsBound())
 	{
-		ConversionFunctionList->RequestListRefresh();
+		OnFieldSelectionChanged.Execute(PropertyPath, Function);
 	}
 }
 
-void SFieldSelector::SetConversionFunctionSelection(const UFunction* SelectedFunction)
+void SFieldSelector::HandleMenuClosed()
 {
 	if (ComboButton.IsValid())
 	{
 		ComboButton->SetIsOpen(false);
 	}
-
-	OnConversionFunctionSelectionChangedDelegate.ExecuteIfBound(SelectedFunction);
-	Refresh();
 }
 
 } // namespace UE::MVVM
diff --git a/Engine/Plugins/Runtime/ModelViewViewModel/Source/ModelViewViewModelEditor/Private/Widgets/SMVVMFieldSelector.h b/Engine/Plugins/Runtime/ModelViewViewModel/Source/ModelViewViewModelEditor/Private/Widgets/SMVVMFieldSelector.h
--- a/Engine/Plugins/Runtime/ModelViewViewModel/Source/ModelViewViewModelEditor/Private/Widgets/SMVVMFieldSelector.h
+++ b/Engine/Plugins/Runtime/ModelViewViewModel/Source/ModelViewViewModelEditor/Private/Widgets/SMVVMFieldSelector.h
@@ -5,183 +5,57 @@
 #include "MVVMPropertyPath.h"
 #include "Styling/CoreStyle.h"
 #include "Styling/SlateTypes.h"
-#include "Types/MVVMBindingSource.h"
 #include "Styling/SlateWidgetStyleAsset.h"
 #include "Widgets/SCompoundWidget.h"
+#include "Widgets/SMVVMCachedViewBindingPropertyPath.h"
+#include "Widgets/SMVVMCachedViewBindingConversionFunction.h"
+#include "Widgets/SMVVMFieldSelectorMenu.h"
 
-enum class EMVVMBindingMode : uint8;
-namespace ESelectInfo { enum Type : int; }
-namespace UE::MVVM { class SFieldEntry; }
-namespace UE::MVVM { enum class EFieldVisibility : uint8; }
-template <typename ItemType> class SListView;
-template <typename ItemType> class STreeView;
+namespace UE::MVVM { class SCachedViewBindingPropertyPath; }
 
 class SComboButton;
-class SSearchBox;
-class SReadOnlyHierarchyView;
-class STableViewBase;
 
 namespace UE::MVVM
 {
 
-class SSourceEntry;
-class SSourceBindingList;
-
 class SFieldSelector : public SCompoundWidget
 {
 public:
-	DECLARE_DELEGATE_OneParam(FOnFieldSelectionChanged, FMVVMBlueprintPropertyPath);
-	DECLARE_DELEGATE_OneParam(FOnConversionFunctionSelectionChanged, const UFunction*);
+	DECLARE_DELEGATE_RetVal(FFieldSelectionContext, FOnGetSelectionContext);
 
-	SLATE_BEGIN_ARGS(SFieldSelector) :
-		_TextStyle(&FCoreStyle::Get().GetWidgetStyle<FTextBlockStyle>("NormalText")),
-		_ShowConversionFunctions(false)
+	SLATE_BEGIN_ARGS(SFieldSelector)
+		: _TextStyle(&FCoreStyle::Get().GetWidgetStyle<FTextBlockStyle>("NormalText"))
 		{
 		}
 		SLATE_STYLE_ARGUMENT(FTextBlockStyle, TextStyle)
-		SLATE_ATTRIBUTE(FMVVMBlueprintPropertyPath, SelectedField)
-		
-		/** Handler for when a field is selected. Always required. */
+		SLATE_ARGUMENT_DEFAULT(bool, ShowContext) = true;
+		SLATE_EVENT(FOnGetPropertyPath, OnGetPropertyPath)
+		SLATE_EVENT(FOnGetConversionFunction, OnGetConversionFunction)
 		SLATE_EVENT(FOnFieldSelectionChanged, OnFieldSelectionChanged)
-
-		SLATE_ATTRIBUTE(const UFunction*, SelectedConversionFunction)
-		
-		/** Handler for when a conversion function is selected. Required if ShowConversionFunctions is set. */
-		SLATE_EVENT(FOnConversionFunctionSelectionChanged, OnConversionFunctionSelectionChanged)
-		
-		/**
-		 * Get the binding mode for this field. 
-		 * Required for validating the availability of functions, eg. SetFoo() can only be used as a setter.
-		 */
-		SLATE_ATTRIBUTE(EMVVMBindingMode, BindingMode)
-
-		/** 
-		 * Should we show the source as well as the field?
-		 * If this is set to false, then the Source attribute must be set.
-		 */
-		SLATE_ARGUMENT_DEFAULT(bool, ShowSource) = true;
-		
-		/** The source to use. Only used if ShowSource is false. */
-		SLATE_ARGUMENT(FBindingSource, Source)
-
-		/**
-		 * Only show properties assignable to the given property.
-		 */
-		SLATE_ATTRIBUTE(const FProperty*, AssignableTo);
-
-		/** Should we show all conversion functions with results compatible with the AssignableTo property? */
-		SLATE_ATTRIBUTE(bool, ShowConversionFunctions);
+		SLATE_EVENT(FOnGetSelectionContext, OnGetSelectionContext)
 	SLATE_END_ARGS()
 
-	void Construct(const FArguments& InArgs, const UWidgetBlueprint* InWidgetBlueprint, bool bInViewModelProperty);
-	void Refresh();
+	void Construct(const FArguments& InArgs, const UWidgetBlueprint* InWidgetBlueprint);
 
 private:
-	TSharedRef<SWidget> OnGenerateFieldWidget(FMVVMBlueprintPropertyPath Path) const;
+	int32 GetCurrentDisplayIndex() const;
 
 	TSharedRef<SWidget> CreateSourcePanel();
-	TSharedRef<SWidget> OnGetMenuContent();
-
-	bool IsClearEnabled() const;
-	FReply OnClearClicked();
-
-	bool IsSelectEnabled() const;
-	FReply OnSelectClicked();
-
-	FReply OnCancelClicked();
-
-	void SetPropertySelection(const FMVVMBlueprintPropertyPath& SelectedPath);
-	EFieldVisibility GetFieldVisibilityFlags() const;
-
-	void OnViewModelSelected(FBindingSource ViewModel, ESelectInfo::Type);
-	TSharedRef<ITableRow> GenerateViewModelRow(MVVM::FBindingSource ViewModel, const TSharedRef<STableViewBase>& OwnerTable) const;
-	void FilterViewModels();
-
-	void OnWidgetSelected(FName WidgetName, ESelectInfo::Type);
-
-	void OnSearchBoxTextChanged(const FText& NewText);
-
-	void OnContextToggleChanged(ECheckBoxState CheckState);
-	ECheckBoxState ContextToggleIsChecked() const;
-
-	struct FConversionFunctionItem
-	{
-		FString GetCategoryName() { return CategoryPath.Num() > 0 ? CategoryPath.Last() : FString(); }
-
-		TArray<FString> CategoryPath;
-		const UFunction* Function;
-		TArray<TSharedPtr<FConversionFunctionItem>> Children;
-		int32 NumFunctions = 0;
-	};
-
-	TSharedPtr<FConversionFunctionItem> FindOrCreateItemForCategory(TArray<TSharedPtr<FConversionFunctionItem>>& Items, TArrayView<FString> CategoryPath);
-	TSharedRef<ITableRow> GenerateConversionFunctionCategoryRow(TSharedPtr<FConversionFunctionItem> Item, const TSharedRef<STableViewBase>& OwnerTable);
-	void GenerateConversionFunctionItems();
-	void FilterConversionFunctionCategories();
-	/** Recursively filter the items in SourceArray and place them into DestArray. Returns true if any items were added. */
-	int32 FilterConversionFunctionCategoryChildren(const TArray<FString>& FilterStrings, const TArray<TSharedPtr<FConversionFunctionItem>>& SourceArray, TArray<TSharedPtr<FConversionFunctionItem>>& OutDestArray);
-	void FilterConversionFunctions();
-
-	void OnConversionFunctionCategorySelected(TSharedPtr<FConversionFunctionItem> Item, ESelectInfo::Type);
-	void GetConversionFunctionCategoryChildren(TSharedPtr<FConversionFunctionItem> Item, TArray<TSharedPtr<FConversionFunctionItem>>& OutItems) const;
-	TSharedPtr<FConversionFunctionItem> FindConversionFunctionCategory(const TArray<TSharedPtr<FConversionFunctionItem>>& Items, TArrayView<FString> CategoryNameParts) const;
-	
-	void AddConversionFunctionChildrenRecursive(const TSharedPtr<FConversionFunctionItem>& Parent, TArray<const UFunction*>& OutFunctions);
-
-	int32 SortConversionFunctionItemsRecursive(TArray<TSharedPtr<FConversionFunctionItem>>& Items);
+	TSharedRef<SWidget> HandleGetMenuContent();
 
-	void SetConversionFunctionSelection(const UFunction* SelectedFunction);
-	void ExpandAllToItem(const UFunction* Function);
-	void ExpandAll(const TArray<TSharedPtr<FConversionFunctionItem>>& Items, bool bRecursive);
-
-	TSharedRef<ITableRow> GenerateConversionFunctionRow(const UFunction* Function, const TSharedRef<STableViewBase>& OwnerTable);
+	void HandleFieldSelectionChanged(FMVVMBlueprintPropertyPath, const UFunction*);
+	void HandleMenuClosed();
 
 private:
-	TAttribute<EMVVMBindingMode> BindingMode;
-
-	TAttribute<FMVVMBlueprintPropertyPath> SelectedField;
-	FMVVMBlueprintPropertyPath CachedSelectedField;
-
-	TSharedPtr<SFieldEntry> SelectedEntryWidget;
-	TSharedPtr<SSourceEntry> SelectedSourceWidget;
-
-	const FTextBlockStyle* TextStyle = nullptr;
-
-	FOnFieldSelectionChanged OnFieldSelectionChangedDelegate;
-
-	TSharedPtr<SSourceBindingList> BindingList;
-	TSharedPtr<SSearchBox> SearchBox;
+	TSharedPtr<SCachedViewBindingPropertyPath> PropertyPathWidget;
 	TSharedPtr<SComboButton> ComboButton;
 
-	TOptional<FBindingSource> FixedSource;
-
 	TWeakObjectPtr<const UWidgetBlueprint> WidgetBlueprint;
-	bool bViewModelProperty = false;
-
-	TArray<FBindingSource> ViewModelSources;
-	TArray<FBindingSource> FilteredViewModelSources;
-
-	TSharedPtr<SListView<FBindingSource>> ViewModelList;
-	TSharedPtr<SReadOnlyHierarchyView> WidgetList;
-
-	TAttribute<const FProperty*> AssignableTo;
-	TAttribute<bool> ShowConversionFunctions;
-
-	TAttribute<const UFunction*> SelectedConversionFunction;
-	const UFunction* CachedSelectedConversionFunction = nullptr;
-
-	FOnConversionFunctionSelectionChanged OnConversionFunctionSelectionChangedDelegate;
-
-	TSharedPtr<STreeView<TSharedPtr<FConversionFunctionItem>>> ConversionFunctionCategoryTree;
-	TArray<TSharedPtr<FConversionFunctionItem>> FilteredConversionFunctionRoot;
-	TArray<TSharedPtr<FConversionFunctionItem>> ConversionFunctionRoot;
-
-	TSharedPtr<SListView<const UFunction*>> ConversionFunctionList;
-
-	TArray<const UFunction*> ConversionFunctions;
-	TArray<const UFunction*> FilteredConversionFunctions;
-
-	bool bIsMenuInitialized = false;
+	const FTextBlockStyle* TextStyle = nullptr;
+	FOnGetPropertyPath OnGetPropertyPath;
+	FOnGetConversionFunction OnGetConversionFunction;
+	FOnGetSelectionContext OnGetSelectionContext;
+	FOnFieldSelectionChanged OnFieldSelectionChanged;
 }; 
 
 } // namespace UE::MVVM
diff --git a/Engine/Plugins/Runtime/ModelViewViewModel/Source/ModelViewViewModelEditor/Private/Widgets/SMVVMFieldSelectorMenu.cpp b/Engine/Plugins/Runtime/ModelViewViewModel/Source/ModelViewViewModelEditor/Private/Widgets/SMVVMFieldSelectorMenu.cpp
new file mode 100644
--- /dev/null
+++ b/Engine/Plugins/Runtime/ModelViewViewModel/Source/ModelViewViewModelEditor/Private/Widgets/SMVVMFieldSelectorMenu.cpp
@@ -0,0 +1,1143 @@
+// Copyright Epic Games, Inc. All Rights Reserved.
+
+#include "Widgets/SMVVMFieldSelectorMenu.h"
+
+#include "Bindings/MVVMBindingHelper.h"
+#include "Editor.h"
+#include "Framework/Views/TableViewMetadata.h"
+#include "Hierarchy/SReadOnlyHierarchyView.h"
+#include "MVVMBlueprintView.h"
+#include "MVVMEditorSubsystem.h"
+#include "MVVMDeveloperProjectSettings.h"
+#include "MVVMWidgetBlueprintExtension_View.h"
+#include "SPrimaryButton.h"
+#include "Styling/MVVMEditorStyle.h"
+#include "Styling/SlateIconFinder.h"
+#include "Types/MVVMBindingMode.h"
+#include "Widgets/SMVVMFieldEntry.h"
+#include "Widgets/SMVVMPropertyPath.h"
+#include "Widgets/SMVVMSourceEntry.h"
+#include "Widgets/Input/SCheckBox.h"
+#include "Widgets/Input/SComboButton.h"
+#include "Widgets/Input/SSearchBox.h"
+#include "Widgets/Layout/SScrollBox.h"
+#include "Widgets/Layout/SWidgetSwitcher.h"
+#include "Widgets/SMVVMCachedViewBindingPropertyPath.h"
+#include "Widgets/SMVVMCachedViewBindingConversionFunction.h"
+#include "Widgets/SMVVMViewModelBindingListWidget.h"
+#include "Widgets/Views/STreeView.h"
+
+#define LOCTEXT_NAMESPACE "MVVMFieldSelectorMenu"
+
+namespace UE::MVVM
+{
+
+namespace Private
+{
+EFieldVisibility GetFieldVisibilityFlags(EMVVMBindingMode Mode, bool bReadable, bool bWritable)
+{
+	EFieldVisibility Flags = EFieldVisibility::None;
+
+	if (bReadable)
+	{
+		Flags |= EFieldVisibility::Readable;
+	}
+	if (bWritable)
+	{
+		Flags |= EFieldVisibility::Writable;
+	}
+	if (!IsOneTimeBinding(Mode))
+	{
+		Flags |= EFieldVisibility::Notify;
+	}
+
+	return Flags;
+}
+}
+
+void SFieldSelectorMenu::Construct(const FArguments& InArgs, const UWidgetBlueprint* InWidgetBlueprint)
+{
+	WidgetBlueprint = InWidgetBlueprint;
+	check(InWidgetBlueprint);
+
+	OnFieldSelectionChanged = InArgs._OnFieldSelectionChanged;
+	OnMenuCloseRequested = InArgs._OnMenuCloseRequested;
+	SelectionContext = InArgs._SelectionContext;
+
+	bIsMenuInitialized = false;
+	bIsClearEnabled = (InArgs._CurrentPropertyPathSelected.IsSet() && !InArgs._CurrentPropertyPathSelected.GetValue().IsEmpty()) || InArgs._CurrentFunctionSelected != nullptr;
+
+	// If we're showing conversion functions, we don't want to set the AssignableTo property of SSourceBindingList, because then it will only show exact matches, 
+	// and since we're also showing conversion functions we know that's not what the user wants.
+	// However, in the case we're not showing conversion functions, we want only exact matches.
+	const FProperty* AssignableToProperty = nullptr;
+	const bool bShowConversionFunctions = SelectionContext.bAllowConversionFunctions;
+	if (!bShowConversionFunctions)
+	{
+		AssignableToProperty = SelectionContext.AssignableTo;
+	}
+
+	TSharedRef<SVerticalBox> VBox = SNew(SVerticalBox)
+	+ SVerticalBox::Slot()
+	.HAlign(HAlign_Fill)
+	.VAlign(VAlign_Top)
+	.Padding(0, 4, 0, 4)
+	.AutoHeight()
+	[
+		SNew(SHorizontalBox)
+		+ SHorizontalBox::Slot()
+		.FillWidth(1.f)
+		.Padding(3.f, 0.f, 3.f, 0.f)	
+		[
+			SAssignNew(SearchBox, SSearchBox)
+			.OnTextChanged(this, &SFieldSelectorMenu::HandleSearchBoxTextChanged)
+		]
+		+ SHorizontalBox::Slot()
+		.HAlign(HAlign_Right)
+		.VAlign(VAlign_Center)
+		.AutoWidth()
+		[
+			SNew(SCheckBox)
+			.OnCheckStateChanged(this, &SFieldSelectorMenu::HandleEnabledContextToggleChanged)
+			.IsChecked(this, &SFieldSelectorMenu::ToggleEnabledContext)
+			.ToolTipText(LOCTEXT("ContextFlagToolTip", "Should the list be filtered to only properties and functions that make sense in the current context?"))
+			[
+				SNew(STextBlock)
+				.Text(LOCTEXT("MVVMContextSensitiveToggle", "Context Sensitive"))
+			]
+		]
+	];
+
+	VBox->AddSlot()
+		.FillHeight(1.0f)
+		[
+			SNew(SBorder)
+			.BorderImage(FAppStyle::Get().GetBrush("Brushes.Recessed"))
+			[
+				SNew(SSplitter)
+				.PhysicalSplitterHandleSize(4.0f)
+				+ SSplitter::Slot()
+				.Value(0.5f)
+				.MinSize(100.0f)
+				[
+					SNew(SScrollBox)
+					+ SScrollBox::Slot()
+					[
+						CreateBindingContextPanel(InArgs)
+					]
+				]
+				+ SSplitter::Slot()
+				.Value(0.5f)
+				.MinSize(100.0f)
+				[
+					SNew(SScrollBox)
+					+ SScrollBox::Slot()
+					[
+						CreateBindingListPanel(InArgs, AssignableToProperty)
+					]
+				]
+			]
+		];
+
+	VBox->AddSlot()
+		.Padding(4.0f, 4.0f, 4.0f, 0.0f)
+		.HAlign(HAlign_Right)
+		.AutoHeight()
+		[
+			SNew(SHorizontalBox)
+			+ SHorizontalBox::Slot()
+			.AutoWidth()
+			[
+				SNew(SPrimaryButton)
+				.OnClicked(this, &SFieldSelectorMenu::HandleSelectClicked)
+				.IsEnabled(this, &SFieldSelectorMenu::IsSelectEnabled)
+				.Text(LOCTEXT("Select", "Select"))
+			]
+			+ SHorizontalBox::Slot()
+			.Padding(FAppStyle::Get().GetMargin("StandardDialog.SlotPadding"))
+			.AutoWidth()
+			[
+				SNew(SButton)
+				.OnClicked(this, &SFieldSelectorMenu::HandleClearClicked)
+				.IsEnabled(this, &SFieldSelectorMenu::IsClearEnabled)
+				.HAlign(HAlign_Center)
+				[
+					SNew(STextBlock)
+					.Text(LOCTEXT("Clear", "Clear"))
+				]
+			]
+			+ SHorizontalBox::Slot()
+			.Padding(FAppStyle::Get().GetMargin("StandardDialog.SlotPadding"))
+			.AutoWidth()
+			[
+				SNew(SButton)
+				.OnClicked(this, &SFieldSelectorMenu::HandleCancelClicked)
+				.HAlign(HAlign_Center)
+				[
+					SNew(STextBlock)
+					.Text(LOCTEXT("Cancel", "Cancel"))
+				]
+			]
+		];
+
+	TSharedRef<SWidget> MenuWidget = SNew(SBox)
+		.MinDesiredWidth(400.0f)
+		.MinDesiredHeight(200.0f)
+		[
+			SNew(SBorder)
+			.BorderImage(FAppStyle::Get().GetBrush("ToolPanel.GroupBorder"))
+			.Padding(FMargin(8.0f, 2.0f, 8.0f, 3.0f))
+			[
+				VBox
+			]
+		];
+
+	ChildSlot
+	[
+		MenuWidget
+	];
+
+	bIsMenuInitialized = true;
+}
+
+TSharedRef<SWidget> SFieldSelectorMenu::GetWidgetToFocus() const
+{
+	return SearchBox.ToSharedRef();
+}
+
+bool SFieldSelectorMenu::IsSelectEnabled() const
+{
+	if (BindingList.IsValid())
+	{
+		FMVVMBlueprintPropertyPath Path = BindingList->GetSelectedProperty();
+		if (Path.IsFromViewModel() || Path.IsFromWidget())
+		{
+			return true;
+		}
+	}
+
+	if (ConversionFunctionList.IsValid())
+	{
+		if (ConversionFunctionList->GetSelectedItems().Num() > 0)
+		{
+			return true;
+		}
+	}
+
+	return false;
+}
+
+FReply SFieldSelectorMenu::HandleSelectClicked()
+{
+	FMVVMBlueprintPropertyPath PropertyPath;
+	if (BindingList.IsValid())
+	{
+		PropertyPath = BindingList->GetSelectedProperty();
+	}
+
+	TArray<const UFunction*> Selection;
+	if (ConversionFunctionList.IsValid())
+	{
+		Selection = ConversionFunctionList->GetSelectedItems();
+	}
+
+	OnFieldSelectionChanged.ExecuteIfBound(PropertyPath, Selection.Num() > 0 ? Selection[0] : nullptr);
+
+	return FReply::Handled();
+}
+
+bool SFieldSelectorMenu::IsClearEnabled() const
+{
+	return bIsClearEnabled;
+}
+
+FReply SFieldSelectorMenu::HandleClearClicked()
+{
+	FMVVMBlueprintPropertyPath NewProperty;
+	
+	if (SelectionContext.FixedBindingSource.IsSet())
+	{
+		FBindingSource Source = SelectionContext.FixedBindingSource.GetValue();
+		if (Source.ViewModelId.IsValid())
+		{
+			NewProperty.SetViewModelId(Source.ViewModelId);
+		}
+		else 
+		{
+			NewProperty.SetWidgetName(Source.Name);
+		}
+	}
+
+	OnFieldSelectionChanged.ExecuteIfBound(NewProperty, nullptr);
+
+	return FReply::Handled();
+}
+
+FReply SFieldSelectorMenu::HandleCancelClicked()
+{
+	OnMenuCloseRequested.ExecuteIfBound();
+
+	return FReply::Handled();
+}
+
+void SFieldSelectorMenu::SetPropertyPathSelection(const FMVVMBlueprintPropertyPath& SelectedPath)
+{
+	OnFieldSelectionChanged.ExecuteIfBound(SelectedPath, nullptr);
+}
+
+void SFieldSelectorMenu::SetConversionFunctionSelection(const UFunction* SelectedFunction)
+{
+	OnFieldSelectionChanged.ExecuteIfBound(FMVVMBlueprintPropertyPath(), SelectedFunction);
+}
+
+TSharedPtr<SFieldSelectorMenu::FConversionFunctionItem> SFieldSelectorMenu::FindConversionFunctionCategory(const TArray<TSharedPtr<FConversionFunctionItem>>& Items, TArrayView<FString> CategoryNameParts) const
+{
+	if (CategoryNameParts.Num() > 0)
+	{
+		for (const TSharedPtr<FConversionFunctionItem>& Item : Items)
+		{
+			if (Item->GetCategoryName() == CategoryNameParts[0])
+			{
+				TArrayView<FString> RemainingParts = CategoryNameParts.RightChop(1);
+
+				// last category part, this is what we're looking for
+				if (RemainingParts.Num() == 0)
+				{
+					return Item;
+				}
+
+				// recurse into children
+				return FindConversionFunctionCategory(Item->Children, RemainingParts);
+			}
+		}
+	}
+
+	return TSharedPtr<FConversionFunctionItem>();
+}
+
+void SFieldSelectorMenu::HandleSearchBoxTextChanged(const FText& NewText)
+{
+	if (BindingList.IsValid())
+	{
+		BindingList->SetRawFilterText(NewText);
+	}
+
+	// Do not search for widgets, only search for properties in the widgets
+	//if (WidgetList.IsValid())
+	//{
+	//	WidgetList->SetRawFilterText(NewText);
+	//}
+
+	//FilterViewModels(NewText);
+
+	if (ConversionFunctionCategoryTree.IsValid())
+	{
+		TArray<TSharedPtr<FConversionFunctionItem>> OldSelectedCategories;
+		ConversionFunctionCategoryTree->GetSelectedItems(OldSelectedCategories);
+
+		FilterConversionFunctionCategories();
+		ConversionFunctionCategoryTree->RequestTreeRefresh();
+
+		TArray<TSharedPtr<FConversionFunctionItem>> NewSelectedCategories;
+
+		// reselect old selection
+		for (const TSharedPtr<FConversionFunctionItem>& OldItem : OldSelectedCategories)
+		{
+			if (TSharedPtr<FConversionFunctionItem> NewItem = FindConversionFunctionCategory(FilteredConversionFunctionRoot, OldItem->CategoryPath))
+			{
+				NewSelectedCategories.Add(NewItem);
+			}
+		}
+		
+		if (NewSelectedCategories.Num() > 0)
+		{
+			ConversionFunctionCategoryTree->SetItemSelection(NewSelectedCategories, true);
+		}
+	}
+
+	if (ConversionFunctionList.IsValid())
+	{
+		TArray<const UFunction*> OldSelectedFunctions;
+		ConversionFunctionList->GetSelectedItems(OldSelectedFunctions);
+
+		FilterConversionFunctions();
+		ConversionFunctionList->RequestListRefresh();
+
+		ConversionFunctionList->SetItemSelection(OldSelectedFunctions, true);
+	}
+}
+
+void SFieldSelectorMenu::HandleViewModelSelected(FBindingSource Source, ESelectInfo::Type)
+{
+	if (!bIsMenuInitialized)
+	{
+		return;
+	}
+
+	if (BindingList.IsValid())
+	{
+		BindingList->ClearSources();
+	}
+
+	TArray<FBindingSource> Selection = ViewModelList->GetSelectedItems();
+	if (Selection.Num() == 0)
+	{
+		return;
+	}
+
+	if (ConversionFunctionCategoryTree.IsValid())
+	{
+		ConversionFunctionCategoryTree->ClearSelection();
+	}
+
+	if (BindingList.IsValid())
+	{
+		BindingList->AddSources(Selection);
+	}
+}
+
+void SFieldSelectorMenu::HandleWidgetSelected(FName WidgetName, ESelectInfo::Type)
+{
+	if (!bIsMenuInitialized)
+	{
+		return;
+	}
+
+	if (BindingList.IsValid())
+	{
+		BindingList->ClearSources();
+	}
+
+	TArray<FName> Selection = WidgetList->GetSelectedWidgets();
+	if (Selection.Num() == 0)
+	{
+		return;
+	}
+
+	if (ConversionFunctionCategoryTree.IsValid())
+	{
+		ConversionFunctionCategoryTree->ClearSelection();
+	}
+
+	if (BindingList.IsValid())
+	{
+		const UWidgetBlueprint* WidgetBP = WidgetBlueprint.Get();
+
+		TArray<FBindingSource> Sources;
+		Algo::Transform(Selection, Sources, [WidgetBP](const FName& WidgetName)
+			{
+				return FBindingSource::CreateForWidget(WidgetBP, WidgetName);
+			});
+
+		BindingList->AddSources(Sources);
+	}
+}
+
+TSharedRef<ITableRow> SFieldSelectorMenu::HandleGenerateViewModelRow(FBindingSource ViewModel, const TSharedRef<STableViewBase>& OwnerTable) const
+{
+	return SNew(STableRow<FBindingSource>, OwnerTable)
+		[
+			SNew(SHorizontalBox)
+			+ SHorizontalBox::Slot()
+			.Padding(4)
+			.HAlign(HAlign_Left)
+			.VAlign(VAlign_Center)
+			.AutoWidth()
+			[
+				SNew(SImage)
+				.Image(FSlateIconFinder::FindIconBrushForClass(ViewModel.Class.Get()))
+				.ColorAndOpacity(FSlateColor::UseForeground())
+			]
+			+ SHorizontalBox::Slot()
+			.HAlign(HAlign_Left)
+			.VAlign(VAlign_Center)
+			.AutoWidth()
+			[
+				SNew(STextBlock)
+				.Text(ViewModel.DisplayName)
+				.ColorAndOpacity(FSlateColor::UseForeground())
+				.HighlightText_Lambda([this]() { return SearchBox.IsValid() ? SearchBox->GetText() : FText::GetEmpty(); })
+			]
+		];
+}
+
+int32 SFieldSelectorMenu::FilterConversionFunctionCategoryChildren(const TArray<FString>& FilterStrings, const TArray<TSharedPtr<FConversionFunctionItem>>& SourceArray, TArray<TSharedPtr<FConversionFunctionItem>>& OutDestArray)
+{
+	int32 NumFunctions = 0;
+
+	for (const TSharedPtr<FConversionFunctionItem>& SourceItem : SourceArray)
+	{
+		TArray<TSharedPtr<FConversionFunctionItem>> FilteredChildren;
+
+		// check if our name matches the filters
+		bool bMatchesFilters = true;
+
+		const FString ItemName = SourceItem->Function != nullptr ? SourceItem->Function->GetName() : FString();
+		
+		for (const FString& Filter : FilterStrings)
+		{
+			if (!ItemName.Contains(Filter))
+			{
+				bMatchesFilters = false;
+				break;
+			}
+		}
+
+		int32 NumChildren = 0;
+		if (bMatchesFilters)
+		{
+			ensureAlways(SourceItem->Function != nullptr);
+			NumChildren = 1;
+		}
+		else
+		{
+			// if we don't match, then we still want to check all our children
+			NumChildren = FilterConversionFunctionCategoryChildren(FilterStrings, SourceItem->Children, FilteredChildren);
+		}
+
+		// then add this item to the destination array
+		if (NumChildren > 0)
+		{
+			TSharedPtr<FConversionFunctionItem>& NewItem = OutDestArray.Add_GetRef(MakeShared<FConversionFunctionItem>());
+			NewItem->CategoryPath = SourceItem->CategoryPath;
+			NewItem->Function = SourceItem->Function; 
+			NewItem->Children = FilteredChildren;
+			NewItem->NumFunctions = NumChildren;
+
+			NumFunctions += NewItem->NumFunctions;
+		}
+	}
+
+	return NumFunctions;
+}
+
+void SFieldSelectorMenu::FilterConversionFunctionCategories()
+{
+	FilteredConversionFunctionRoot.Reset();
+
+	TArray<FString> FilterStrings;
+	if (SearchBox.IsValid())
+	{
+		SearchBox->GetText().ToString().ParseIntoArrayWS(FilterStrings);
+	}
+
+	if (FilterStrings.IsEmpty())
+	{
+		// don't bother filtering if we don't have any search terms
+		FilteredConversionFunctionRoot = ConversionFunctionRoot;
+		return;
+	}
+
+	TSharedPtr<FConversionFunctionItem> RootItem = FilteredConversionFunctionRoot.Add_GetRef(MakeShared<FConversionFunctionItem>());
+	RootItem->CategoryPath = { TEXT("Conversion Functions") };
+	
+	RootItem->NumFunctions = FilterConversionFunctionCategoryChildren(FilterStrings, ConversionFunctionRoot[0]->Children, FilteredConversionFunctionRoot[0]->Children);
+	
+	ExpandFunctionCategoryTree(FilteredConversionFunctionRoot, true);
+}
+
+int32 SFieldSelectorMenu::SortConversionFunctionItemsRecursive(TArray<TSharedPtr<FConversionFunctionItem>>& Items)
+{
+	int32 NumFound = 0;
+
+	Items.Sort([](const TSharedPtr<FConversionFunctionItem>& A, const TSharedPtr<FConversionFunctionItem>& B)
+		{
+			if (!A->GetCategoryName().IsEmpty() && !B->GetCategoryName().IsEmpty())
+			{
+				return A->GetCategoryName() < B->GetCategoryName();
+			}
+			if (!A->GetCategoryName().IsEmpty() && B->GetCategoryName().IsEmpty())
+			{
+				return true;
+			}
+			if (A->GetCategoryName().IsEmpty() && !B->GetCategoryName().IsEmpty())
+			{
+				return false;
+			}
+			if (A->Function != nullptr && B->Function != nullptr)
+			{
+				return A->Function->GetDisplayNameText().CompareTo(B->Function->GetDisplayNameText()) <= 0;
+			}
+			return true;
+		});
+
+	for (const TSharedPtr<FConversionFunctionItem>& Item : Items)
+	{
+		NumFound += SortConversionFunctionItemsRecursive(Item->Children);
+
+		if (Item->Function != nullptr)
+		{
+			NumFound += 1;
+		}
+	}
+
+	return NumFound;
+}
+
+void SFieldSelectorMenu::GenerateConversionFunctionItems()
+{
+	UMVVMEditorSubsystem* Subsystem = GEditor->GetEditorSubsystem<UMVVMEditorSubsystem>();
+	TArray<UFunction*> AllConversionFunctions = Subsystem->GetAvailableConversionFunctions(WidgetBlueprint.Get(), FMVVMBlueprintPropertyPath(), FMVVMBlueprintPropertyPath());
+
+	// remove all incompatible conversion functions
+	for (int32 Index = AllConversionFunctions.Num() - 1; Index >= 0; --Index)
+	{
+		const UFunction* Function = AllConversionFunctions[Index];
+		const FProperty* ReturnProperty = BindingHelper::GetReturnProperty(Function);
+		const FProperty* AssignToProperty = SelectionContext.AssignableTo;
+
+		if (AssignToProperty != nullptr && !BindingHelper::ArePropertiesCompatible(ReturnProperty, AssignToProperty))
+		{
+			AllConversionFunctions.RemoveAtSwap(Index);
+		}
+	}
+
+	auto AddFunctionToItem = [](const UFunction* Function, const TSharedPtr<FConversionFunctionItem>& Parent)
+	{
+		TSharedPtr<FConversionFunctionItem>& Item = Parent->Children.Add_GetRef(MakeShared<FConversionFunctionItem>());
+		Item->Function = Function;
+		Item->NumFunctions = 1;
+		Parent->NumFunctions += 1;
+	};
+
+	TSharedPtr<FConversionFunctionItem> CurrentSelectedItem;
+
+	TArray<FString> CategoryPath;
+
+	ConversionFunctionRoot.Reset();
+	TSharedPtr<FConversionFunctionItem>& RootItem = ConversionFunctionRoot.Add_GetRef(MakeShared<FConversionFunctionItem>());
+	RootItem->CategoryPath = { TEXT("Conversion Functions") };
+
+	FName NAME_Category = "Category";
+	for (const UFunction* Function : AllConversionFunctions)
+	{
+		const FText& CategoryName = Function->GetMetaDataText(NAME_Category);
+		if (CategoryName.IsEmpty())
+		{
+			AddFunctionToItem(Function, RootItem);
+			continue;
+		}
+
+		// split into subcategories and trim
+		CategoryPath.Reset();
+		CategoryName.ToString().ParseIntoArray(CategoryPath, TEXT("|"));
+		for (FString& SubCategory : CategoryPath)
+		{
+			SubCategory.TrimStartAndEndInline();
+		}
+
+		TSharedPtr<FConversionFunctionItem> ParentItem = RootItem;
+
+		// create items for the entire category path
+		// eg. "Math|Boolean|AND" 
+		// Math 
+		//   > Boolean
+		//     > AND
+		for (int32 PathIndex = 0; PathIndex < CategoryPath.Num(); ++PathIndex)
+		{
+			ParentItem->NumFunctions += 1;
+
+			ParentItem = FindOrCreateItemForCategory(ParentItem->Children, MakeArrayView(CategoryPath.GetData(), PathIndex + 1));
+		}
+
+		AddFunctionToItem(Function, ParentItem);
+	}
+
+	int32 NumItems = SortConversionFunctionItemsRecursive(ConversionFunctionRoot);
+	ensure(NumItems == RootItem->NumFunctions);
+}
+
+TSharedPtr<SFieldSelectorMenu::FConversionFunctionItem> SFieldSelectorMenu::ExpandFunctionCategoryTreeToItem(const UFunction* Function)
+{
+	TArray<TSharedPtr<FConversionFunctionItem>> Path;
+
+	FText FullCategoryName = Function->GetMetaDataText("Category");
+	if (FullCategoryName.IsEmpty())
+	{
+		Path.Add(FilteredConversionFunctionRoot[0]);
+	}
+	else
+	{
+		TArray<FString> CategoryPath;
+		FullCategoryName.ToString().ParseIntoArray(CategoryPath, TEXT("|"));
+
+		TSharedPtr<FConversionFunctionItem> CurrentParent = FilteredConversionFunctionRoot[0];
+
+		for (const FString& SubCategory : CategoryPath)
+		{
+			const FString Trimmed = SubCategory.TrimStartAndEnd();
+
+			TSharedPtr<FConversionFunctionItem>* FoundItem =
+				CurrentParent->Children.FindByPredicate([Trimmed, Function](const TSharedPtr<FConversionFunctionItem>& Item)
+					{
+						return Item->GetCategoryName() == Trimmed || Item->Function == Function;
+					});
+
+			if (FoundItem != nullptr)
+			{
+				Path.Add(*FoundItem);
+				CurrentParent = *FoundItem;
+			}
+		}
+	}
+
+	if (Path.Num() > 0)
+	{
+		ConversionFunctionCategoryTree->SetItemExpansion(FilteredConversionFunctionRoot[0], true);
+		ExpandFunctionCategoryTree(Path, false);
+		return Path.Last();
+	}
+	return TSharedPtr<FConversionFunctionItem>();
+}
+
+void SFieldSelectorMenu::ExpandFunctionCategoryTree(const TArray<TSharedPtr<FConversionFunctionItem>>& Items, bool bRecursive)
+{
+	for (const TSharedPtr<FConversionFunctionItem>& Item : Items)
+	{
+		ConversionFunctionCategoryTree->SetItemExpansion(Item, true);
+
+		if (bRecursive)
+		{
+			ExpandFunctionCategoryTree(Item->Children, bRecursive);
+		}
+	}
+}
+
+void SFieldSelectorMenu::FilterViewModels(const FText& NewText)
+{
+	FilteredViewModelSources.Reset();
+
+	TArray<FString> FilterStrings;
+	NewText.ToString().ParseIntoArrayWS(FilterStrings);
+
+	if (FilterStrings.IsEmpty())
+	{
+		FilteredViewModelSources = ViewModelSources;
+		return;
+	}
+
+	for (const FBindingSource& ViewModel : ViewModelSources)
+	{
+		const FString DisplayName = ViewModel.DisplayName.ToString();
+		const FString ClassName = ViewModel.Class != nullptr ? ViewModel.Class->GetName() : FString();
+		const FString Name = ViewModel.Name.ToString();
+
+		bool bMatchesFilters = true;
+		
+		for (const FString& Filter : FilterStrings)
+		{
+			if (!DisplayName.Contains(Filter) && 
+				!ClassName.Contains(Filter) &&
+				!Name.Contains(Filter))
+			{
+				bMatchesFilters = false;
+				break;
+			}
+		}
+
+		if (bMatchesFilters)
+		{
+			FBindingSource& NewSource = FilteredViewModelSources.Add_GetRef(FBindingSource());
+			NewSource.Class = ViewModel.Class;
+			NewSource.Name = ViewModel.Name;
+			NewSource.DisplayName = ViewModel.DisplayName;
+			NewSource.ViewModelId = ViewModel.ViewModelId;
+		}
+	}
+}
+
+TSharedRef<SWidget> SFieldSelectorMenu::CreateBindingContextPanel(const FArguments& InArgs)
+{
+	// show source picker
+	TSharedRef<SVerticalBox> StackedSourcePicker = SNew(SVerticalBox);
+
+	if (SelectionContext.bAllowViewModels)
+	{
+		if (SelectionContext.FixedBindingSource.IsSet()
+			&& SelectionContext.FixedBindingSource.GetValue().IsValid()
+			&& SelectionContext.FixedBindingSource.GetValue().ViewModelId.IsValid())
+		{
+			ViewModelSources.Add(SelectionContext.FixedBindingSource.GetValue());
+		}
+		else
+		{
+			ViewModelSources = GEditor->GetEditorSubsystem<UMVVMEditorSubsystem>()->GetAllViewModels(WidgetBlueprint.Get());
+		}
+
+		FilteredViewModelSources = ViewModelSources;
+
+		ViewModelList = SNew(SListView<FBindingSource>)
+			.ListItemsSource(&FilteredViewModelSources)
+			.OnGenerateRow(this, &SFieldSelectorMenu::HandleGenerateViewModelRow)
+			.SelectionMode(ESelectionMode::Multi)
+			.OnSelectionChanged(this, &SFieldSelectorMenu::HandleViewModelSelected);
+
+		FBindingSource SelectedSource;
+		if (InArgs._CurrentPropertyPathSelected.IsSet() && InArgs._CurrentPropertyPathSelected.GetValue().IsFromViewModel())
+		{
+			for (const FBindingSource& Source : FilteredViewModelSources)
+			{
+				if (Source.ViewModelId == InArgs._CurrentPropertyPathSelected.GetValue().GetViewModelId())
+				{
+					SelectedSource = Source;
+				}
+			}
+		}
+
+		if (SelectedSource.IsValid())
+		{
+			ViewModelList->SetItemSelection(SelectedSource, true);
+		}
+
+		StackedSourcePicker->AddSlot()
+			.AutoHeight()
+			[
+				ViewModelList.ToSharedRef()
+			];
+	}
+	
+	if (SelectionContext.bAllowWidgets)
+	{
+		TArray<FName> ShowOnly;
+		if (SelectionContext.FixedBindingSource.IsSet()
+			&& SelectionContext.FixedBindingSource.GetValue().IsValid()
+			&& !SelectionContext.FixedBindingSource.GetValue().Name.IsNone())
+		{
+			ShowOnly.Add(SelectionContext.FixedBindingSource.GetValue().Name);
+		}
+
+		WidgetList = SNew(SReadOnlyHierarchyView, WidgetBlueprint.Get())
+			.OnSelectionChanged(this, &SFieldSelectorMenu::HandleWidgetSelected)
+			.SelectionMode(ESelectionMode::Multi)
+			.ShowSearch(false)
+			.ShowOnly(ShowOnly)
+			.ExpandAll(false);
+
+		if (InArgs._CurrentPropertyPathSelected.IsSet() && InArgs._CurrentPropertyPathSelected.GetValue().IsFromWidget())
+		{
+			WidgetList->SetSelectedWidget(InArgs._CurrentPropertyPathSelected.GetValue().GetWidgetName());
+		}
+
+		StackedSourcePicker->AddSlot()
+			.AutoHeight()
+			[
+				WidgetList.ToSharedRef()
+			];
+	}
+
+	if (SelectionContext.bAllowConversionFunctions)
+	{
+		GenerateConversionFunctionItems();
+
+		FilterConversionFunctionCategories();
+		
+		SAssignNew(ConversionFunctionCategoryTree, STreeView<TSharedPtr<FConversionFunctionItem>>)
+			.SelectionMode(ESelectionMode::Multi)
+			.TreeItemsSource(&FilteredConversionFunctionRoot)
+			.OnGenerateRow(this, &SFieldSelectorMenu::HandleGenerateConversionFunctionCategoryRow)
+			.OnSelectionChanged(this, &SFieldSelectorMenu::HandleConversionFunctionCategorySelected)
+			.OnGetChildren(this, &SFieldSelectorMenu::HandleGetConversionFunctionCategoryChildren);
+
+		if (InArgs._CurrentFunctionSelected != nullptr)
+		{
+			TSharedPtr<FConversionFunctionItem> FunctionItem = ExpandFunctionCategoryTreeToItem(InArgs._CurrentFunctionSelected);
+			if (FunctionItem)
+			{
+				TGuardValue<bool> TmpGuard(bIsMenuInitialized, true);
+				ConversionFunctionCategoryTree->SetItemSelection(FunctionItem, true);
+			}
+		}
+
+		StackedSourcePicker->AddSlot()
+			.AutoHeight()
+			[
+				ConversionFunctionCategoryTree.ToSharedRef()
+			];
+	}
+
+	return StackedSourcePicker;
+}
+
+TSharedRef<SWidget> SFieldSelectorMenu::CreateBindingListPanel(const FArguments& InArgs, const FProperty* AssignableToProperty)
+{
+	BindingList = SNew(SSourceBindingList, WidgetBlueprint.Get())
+		.ShowSearchBox(false)
+		.OnDoubleClicked(this, &SFieldSelectorMenu::SetPropertyPathSelection)
+		.FieldVisibilityFlags(Private::GetFieldVisibilityFlags(SelectionContext.BindingMode, SelectionContext.bReadable, SelectionContext.bWritable))
+		.AssignableTo(AssignableToProperty);
+
+	if (SelectionContext.FixedBindingSource.IsSet())
+	{
+		// Single fixed source, don't show the separate source panel.
+		BindingList->AddSource(SelectionContext.FixedBindingSource.GetValue());
+	}
+	else if (InArgs._CurrentPropertyPathSelected.IsSet())
+	{
+		FBindingSource Source = InArgs._CurrentPropertyPathSelected.GetValue().IsFromViewModel() ?
+			FBindingSource::CreateForViewModel(WidgetBlueprint.Get(), InArgs._CurrentPropertyPathSelected.GetValue().GetViewModelId()) :
+			FBindingSource::CreateForWidget(WidgetBlueprint.Get(), InArgs._CurrentPropertyPathSelected.GetValue().GetWidgetName());
+		BindingList->AddSource(Source);
+	}
+
+	if (InArgs._CurrentPropertyPathSelected.IsSet())
+	{
+		BindingList->SetSelectedProperty(InArgs._CurrentPropertyPathSelected.GetValue());
+	}
+
+
+	TSharedRef<SVerticalBox> BindingListVBox = SNew(SVerticalBox);
+	BindingListVBox->AddSlot()
+		[
+			BindingList.ToSharedRef()
+		];
+
+	if (SelectionContext.bAllowConversionFunctions)
+	{
+		BindingListVBox->AddSlot()
+			.AutoHeight()
+			[
+				SAssignNew(ConversionFunctionList, SListView<const UFunction*>)
+				.SelectionMode(ESelectionMode::Single)
+				.ListItemsSource(&FilteredConversionFunctions)
+				.OnMouseButtonDoubleClick(this, &SFieldSelectorMenu::SetConversionFunctionSelection)
+				.OnGenerateRow(this, &SFieldSelectorMenu::HandleGenerateConversionFunctionRow)
+			];
+
+		if (InArgs._CurrentFunctionSelected != nullptr)
+		{
+			ConversionFunctionList->SetItemSelection(InArgs._CurrentFunctionSelected, true);
+		}
+	}
+
+	return BindingListVBox;
+}
+
+void SFieldSelectorMenu::HandleEnabledContextToggleChanged(ECheckBoxState CheckState)
+{
+	UMVVMWidgetBlueprintExtension_View* ExtensionView = UMVVMWidgetBlueprintExtension_View::GetExtension<UMVVMWidgetBlueprintExtension_View>(WidgetBlueprint.Get());
+	check(ExtensionView);
+
+	FMVVMViewBindingFilterSettings FilterSettings = ExtensionView->GetFilterSettings();
+	FilterSettings.FilterFlags = CheckState == ECheckBoxState::Checked ? EFilterFlag::All : EFilterFlag::None;
+	ExtensionView->SetFilterSettings(FilterSettings);
+
+	if (WidgetList.IsValid())
+	{
+		HandleWidgetSelected(WidgetBlueprint->GetFName(), ESelectInfo::Direct);
+	}
+
+	if (ViewModelList.IsValid())
+	{
+		HandleViewModelSelected(FBindingSource(), ESelectInfo::Direct);
+	}
+}
+
+ECheckBoxState SFieldSelectorMenu::ToggleEnabledContext() const
+{
+	UMVVMWidgetBlueprintExtension_View* ExtensionView = UMVVMWidgetBlueprintExtension_View::GetExtension<UMVVMWidgetBlueprintExtension_View>(WidgetBlueprint.Get());
+	check(ExtensionView);
+
+	return ExtensionView->GetFilterSettings().FilterFlags == EFilterFlag::All ? ECheckBoxState::Checked : ECheckBoxState::Unchecked;
+}
+
+TSharedPtr<SFieldSelectorMenu::FConversionFunctionItem> SFieldSelectorMenu::FindOrCreateItemForCategory(TArray<TSharedPtr<FConversionFunctionItem>>& Items, TArrayView<FString> CategoryPath)
+{
+	check(CategoryPath.Num() > 0);
+
+	const FString& CategoryName = CategoryPath.Last();
+
+	int32 Idx = 0;
+	for (; Idx < Items.Num(); ++Idx)
+	{
+		// found item
+		if (Items[Idx]->GetCategoryName() == CategoryName)
+		{
+			return Items[Idx];
+		}
+
+		// passed the place where it should have been, break out
+		if (Items[Idx]->GetCategoryName() > CategoryName)
+		{
+			break;
+		}
+	}
+
+	TSharedPtr<FConversionFunctionItem> NewItem = Items.Insert_GetRef(MakeShared<FConversionFunctionItem>(), Idx);
+	NewItem->CategoryPath = CategoryPath;
+	return NewItem;
+}
+
+TSharedRef<ITableRow> SFieldSelectorMenu::HandleGenerateConversionFunctionCategoryRow(TSharedPtr<FConversionFunctionItem> Item, const TSharedRef<STableViewBase>& OwnerTable)
+{
+	FText DisplayName = FText::FormatOrdered(FText::FromString("{0} ({1})"), FText::FromString(Item->GetCategoryName()), FText::FromString(LexToString(Item->NumFunctions)));
+
+	return SNew(STableRow<TSharedPtr<FConversionFunctionItem>>, OwnerTable)
+		[
+			SNew(SHorizontalBox)
+			+ SHorizontalBox::Slot()
+		.HAlign(HAlign_Left)
+		.VAlign(VAlign_Center)
+		.Padding(0, 2.0f, 4.0f, 2.0f)
+		.AutoWidth()
+		[
+			SNew(SImage)
+			.DesiredSizeOverride(FVector2D(16.0f, 16.0f))
+		.Image(FAppStyle::Get().GetBrush("GraphEditor.Function_16x"))
+		]
+	+ SHorizontalBox::Slot()
+		.HAlign(HAlign_Fill)
+		.VAlign(VAlign_Center)
+		[
+			SNew(STextBlock)
+			.Font(Item == FilteredConversionFunctionRoot[0] ? FAppStyle::Get().GetFontStyle("NormalText") : FAppStyle::Get().GetFontStyle("BoldFont"))
+		.Text(DisplayName)
+		.ToolTipText(FText::FromString(Item->GetCategoryName()))
+		.HighlightText_Lambda([this]() { return SearchBox.IsValid() ? SearchBox->GetText() : FText::GetEmpty(); })
+		]
+		];
+}
+
+void SFieldSelectorMenu::HandleGetConversionFunctionCategoryChildren(TSharedPtr<FConversionFunctionItem> Item, TArray<TSharedPtr<FConversionFunctionItem>>& OutItems) const
+{
+	Algo::TransformIf(Item->Children, OutItems,
+		[](const TSharedPtr<FConversionFunctionItem>& Item)
+		{
+			return !Item->GetCategoryName().IsEmpty();
+		},
+		[](const TSharedPtr<FConversionFunctionItem>& Item)
+		{
+			return Item;
+		});
+}
+
+TSharedRef<ITableRow> SFieldSelectorMenu::HandleGenerateConversionFunctionRow(const UFunction* Function, const TSharedRef<STableViewBase>& OwnerTable)
+{
+	return SNew(STableRow<const UFunction*>, OwnerTable)
+		[
+			SNew(SHorizontalBox)
+			+ SHorizontalBox::Slot()
+			.HAlign(HAlign_Left)
+			.VAlign(VAlign_Center)
+			.Padding(0, 0, 4.0f, 0)
+			.AutoWidth()
+			[
+				SNew(SImage)
+				.DesiredSizeOverride(FVector2D(16.0f, 16.0f))
+				.Image(FAppStyle::Get().GetBrush("GraphEditor.Function_16x"))
+			]
+			+ SHorizontalBox::Slot()
+			.HAlign(HAlign_Fill)
+			.VAlign(VAlign_Center)
+			[
+				SNew(STextBlock)
+				.Text(Function->GetDisplayNameText())
+				.ToolTipText(Function->GetToolTipText())
+				.HighlightText_Lambda([this]() { return SearchBox.IsValid() ? SearchBox->GetText() : FText::GetEmpty(); })
+			]
+		];
+}
+
+void SFieldSelectorMenu::AddConversionFunctionChildrenRecursive(const TSharedPtr<FConversionFunctionItem>& Parent, TArray<const UFunction*>& OutFunctions)
+{
+	for (const TSharedPtr<FConversionFunctionItem>& Item : Parent->Children)
+	{
+		if (Item->Function != nullptr)
+		{
+			int32 Index = 0;
+			for (; Index < OutFunctions.Num(); ++Index)
+			{
+				if (OutFunctions[Index]->GetName() > Item->Function->GetName())
+				{
+					break;
+				}
+			}
+
+			OutFunctions.Insert(Item->Function, Index);
+		}
+		else
+		{
+			AddConversionFunctionChildrenRecursive(Item, OutFunctions);
+		}
+	}
+}
+
+void SFieldSelectorMenu::FilterConversionFunctions()
+{
+	TArray<FString> FilterStrings;
+	if (SearchBox.IsValid())
+	{
+		SearchBox->GetText().ToString().ParseIntoArrayWS(FilterStrings);
+	}
+
+	if (FilterStrings.IsEmpty())
+	{
+		FilteredConversionFunctions = ConversionFunctions;
+		return;
+	}
+
+	for (const UFunction* Function : ConversionFunctions)
+	{
+		const FString FunctionName = Function->GetName();
+
+		bool bMatches = true;
+		for (const FString& Filter : FilterStrings)
+		{
+			if (!FunctionName.Contains(Filter))
+			{
+				bMatches = false;
+				break;
+			}
+		}
+
+		if (bMatches)
+		{
+			FilteredConversionFunctions.Add(Function);
+		}
+	}
+}
+
+void SFieldSelectorMenu::HandleConversionFunctionCategorySelected(TSharedPtr<FConversionFunctionItem> SelectedItem, ESelectInfo::Type)
+{
+	if (!bIsMenuInitialized)
+	{
+		return;
+	}
+
+	ConversionFunctions.Reset();
+	FilteredConversionFunctions.Reset();
+
+	TArray<TSharedPtr<FConversionFunctionItem>> SelectedItems = ConversionFunctionCategoryTree->GetSelectedItems();
+	if (SelectedItems.Num() > 0)
+	{
+		if (BindingList.IsValid())
+		{
+			BindingList->ClearSources();
+		}
+
+		if (ViewModelList.IsValid())
+		{
+			ViewModelList->ClearSelection();
+		}
+
+		if (WidgetList.IsValid())
+		{
+			WidgetList->ClearSelection();
+		}
+
+		for (const TSharedPtr<FConversionFunctionItem>& Item : SelectedItems)
+		{
+			AddConversionFunctionChildrenRecursive(Item, ConversionFunctions);
+		}
+
+		FilterConversionFunctions();
+	}
+
+	if (ConversionFunctionList.IsValid())
+	{
+		ConversionFunctionList->RequestListRefresh();
+	}
+}
+
+} // namespace UE::MVVM
+
+#undef LOCTEXT_NAMESPACE
diff --git a/Engine/Plugins/Runtime/ModelViewViewModel/Source/ModelViewViewModelEditor/Private/Widgets/SMVVMFieldSelectorMenu.h b/Engine/Plugins/Runtime/ModelViewViewModel/Source/ModelViewViewModelEditor/Private/Widgets/SMVVMFieldSelectorMenu.h
new file mode 100644
--- /dev/null
+++ b/Engine/Plugins/Runtime/ModelViewViewModel/Source/ModelViewViewModelEditor/Private/Widgets/SMVVMFieldSelectorMenu.h
@@ -0,0 +1,144 @@
+// Copyright Epic Games, Inc. All Rights Reserved.
+
+#pragma once
+
+#include "MVVMPropertyPath.h"
+#include "Styling/CoreStyle.h"
+#include "Styling/SlateTypes.h"
+#include "Types/MVVMBindingMode.h"
+#include "Types/MVVMBindingSource.h"
+#include "Widgets/SCompoundWidget.h"
+
+namespace ESelectInfo { enum Type : int; }
+class ITableRow;
+class SSearchBox;
+class SReadOnlyHierarchyView;
+class STableViewBase;
+template <typename ItemType> class SListView;
+template <typename ItemType> class STreeView;
+namespace UE::MVVM {class SSourceBindingList; }
+
+class UWidgetBlueprint;
+
+namespace UE::MVVM
+{
+
+struct FFieldSelectionContext
+{
+	EMVVMBindingMode BindingMode = EMVVMBindingMode::OneWayToDestination;
+	const FProperty* AssignableTo = nullptr;
+	TOptional<FBindingSource> FixedBindingSource;
+	bool bAllowWidgets = true;
+	bool bAllowViewModels = true;
+	bool bAllowConversionFunctions = true;
+	bool bReadable = true;
+	bool bWritable = true;
+};
+
+DECLARE_DELEGATE_TwoParams(FOnFieldSelectionChanged, FMVVMBlueprintPropertyPath, const UFunction*);
+
+class SFieldSelectorMenu : public SCompoundWidget
+{
+public:
+	DECLARE_DELEGATE_RetVal(FFieldSelectionContext, FOnGetFieldSelectionContext);
+
+	SLATE_BEGIN_ARGS(SFieldSelectorMenu)
+	{
+	}
+		SLATE_EVENT(FOnFieldSelectionChanged, OnFieldSelectionChanged)
+		SLATE_EVENT(FSimpleDelegate, OnMenuCloseRequested)
+
+		SLATE_ARGUMENT(FFieldSelectionContext, SelectionContext)
+		SLATE_ARGUMENT(TOptional<FMVVMBlueprintPropertyPath>, CurrentPropertyPathSelected)
+		SLATE_ARGUMENT_DEFAULT(const UFunction*, CurrentFunctionSelected) = nullptr;
+	SLATE_END_ARGS()
+
+	void Construct(const FArguments& InArgs, const UWidgetBlueprint* InWidgetBlueprint);
+
+	TSharedRef<SWidget> GetWidgetToFocus() const;
+
+private:
+	struct FConversionFunctionItem
+	{
+		FString GetCategoryName() { return CategoryPath.Num() > 0 ? CategoryPath.Last() : FString(); }
+
+		TArray<FString> CategoryPath;
+		const UFunction* Function = nullptr;
+		TArray<TSharedPtr<FConversionFunctionItem>> Children;
+		int32 NumFunctions = 0;
+	};
+
+private:
+	void SetPropertyPathSelection(const FMVVMBlueprintPropertyPath& SelectedPath);
+	void SetConversionFunctionSelection(const UFunction* SelectedFunction);
+
+	TSharedRef<SWidget> CreateBindingContextPanel(const FArguments& InArgs);
+	TSharedRef<SWidget> CreateBindingListPanel(const FArguments& InArgs, const FProperty* AssignableToProperty);
+
+	void HandleWidgetSelected(FName WidgetName, ESelectInfo::Type);
+
+	void HandleViewModelSelected(FBindingSource ViewModel, ESelectInfo::Type);
+	TSharedRef<ITableRow> HandleGenerateViewModelRow(MVVM::FBindingSource ViewModel, const TSharedRef<STableViewBase>& OwnerTable) const;
+
+	int32 SortConversionFunctionItemsRecursive(TArray<TSharedPtr<FConversionFunctionItem>>& Items);
+	void GenerateConversionFunctionItems();
+	void FilterConversionFunctionCategories();
+	void FilterConversionFunctions();
+	/** Recursively filter the items in SourceArray and place them into DestArray. Returns true if any items were added. */
+	int32 FilterConversionFunctionCategoryChildren(const TArray<FString>& FilterStrings, const TArray<TSharedPtr<FConversionFunctionItem>>& SourceArray, TArray<TSharedPtr<FConversionFunctionItem>>& OutDestArray);
+	void AddConversionFunctionChildrenRecursive(const TSharedPtr<FConversionFunctionItem>& Parent, TArray<const UFunction*>& OutFunctions);
+	TSharedPtr<FConversionFunctionItem> FindOrCreateItemForCategory(TArray<TSharedPtr<FConversionFunctionItem>>& Items, TArrayView<FString> CategoryPath);
+	TSharedPtr<FConversionFunctionItem> FindConversionFunctionCategory(const TArray<TSharedPtr<FConversionFunctionItem>>& Items, TArrayView<FString> CategoryNameParts) const;
+	void HandleGetConversionFunctionCategoryChildren(TSharedPtr<FConversionFunctionItem> Item, TArray<TSharedPtr<FConversionFunctionItem>>& OutItems) const;
+	void HandleConversionFunctionCategorySelected(TSharedPtr<FConversionFunctionItem> Item, ESelectInfo::Type);
+	TSharedRef<ITableRow> HandleGenerateConversionFunctionCategoryRow(TSharedPtr<FConversionFunctionItem> Item, const TSharedRef<STableViewBase>& OwnerTable);
+	TSharedRef<ITableRow> HandleGenerateConversionFunctionRow(const UFunction* Function, const TSharedRef<STableViewBase>& OwnerTable);
+
+	TSharedPtr<FConversionFunctionItem> ExpandFunctionCategoryTreeToItem(const UFunction* Function);
+	void ExpandFunctionCategoryTree(const TArray<TSharedPtr<FConversionFunctionItem>>& Items, bool bRecursive);
+
+	void FilterViewModels(const FText& NewText);
+	void HandleSearchBoxTextChanged(const FText& NewText);
+	void HandleEnabledContextToggleChanged(ECheckBoxState CheckState);
+	ECheckBoxState ToggleEnabledContext() const;
+
+	bool IsClearEnabled() const;
+	bool IsSelectEnabled() const;
+	FReply HandleClearClicked();
+	FReply HandleSelectClicked();
+	FReply HandleCancelClicked();
+
+private:
+	TWeakObjectPtr<const UWidgetBlueprint> WidgetBlueprint;
+	FOnFieldSelectionChanged OnFieldSelectionChanged;
+	FSimpleDelegate OnMenuCloseRequested;
+	FFieldSelectionContext SelectionContext;
+
+	TSharedPtr<SSearchBox> SearchBox;
+
+	//~ viewmodels (binding context panel)
+	TSharedPtr<SListView<FBindingSource>> ViewModelList;
+	TArray<FBindingSource> ViewModelSources;
+	TArray<FBindingSource> FilteredViewModelSources;
+
+	//~ widgets (binding context panel)
+	TSharedPtr<SReadOnlyHierarchyView> WidgetList;
+
+	//~ viewmodel and widgets (selection panel)
+	TSharedPtr<SSourceBindingList> BindingList;
+
+	//~ functions (binding context panel)
+	TSharedPtr<STreeView<TSharedPtr<FConversionFunctionItem>>> ConversionFunctionCategoryTree;
+	TArray<TSharedPtr<FConversionFunctionItem>> FilteredConversionFunctionRoot;
+	TArray<TSharedPtr<FConversionFunctionItem>> ConversionFunctionRoot;
+
+	//~ functions (selection panel)
+	TSharedPtr<SListView<const UFunction*>> ConversionFunctionList;
+	TArray<const UFunction*> ConversionFunctions;
+	TArray<const UFunction*> FilteredConversionFunctions;
+
+	bool bIsMenuInitialized = false;
+	bool bIsClearEnabled = false;
+}; 
+
+} // namespace UE::MVVM
diff --git a/Engine/Plugins/Runtime/ModelViewViewModel/Source/ModelViewViewModelEditor/Private/Widgets/SMVVMFunctionParameter.cpp b/Engine/Plugins/Runtime/ModelViewViewModel/Source/ModelViewViewModelEditor/Private/Widgets/SMVVMFunctionParameter.cpp
--- a/Engine/Plugins/Runtime/ModelViewViewModel/Source/ModelViewViewModelEditor/Private/Widgets/SMVVMFunctionParameter.cpp
+++ b/Engine/Plugins/Runtime/ModelViewViewModel/Source/ModelViewViewModelEditor/Private/Widgets/SMVVMFunctionParameter.cpp
@@ -15,6 +15,7 @@
 #include "Widgets/SMVVMFieldSelector.h"
 #include "Widgets/Input/SCheckBox.h"
 #include "Widgets/Layout/SBox.h"
+#include "Widgets/Layout/SWidgetSwitcher.h"
 #include "Widgets/SOverlay.h"
 #include "Widgets/Text/STextBlock.h"
 
@@ -23,13 +24,13 @@
 namespace UE::MVVM
 {
 
-void SFunctionParameter::Construct(const FArguments& InArgs)
+void SFunctionParameter::Construct(const FArguments& InArgs, UWidgetBlueprint* InWidgetBlueprint)
 {
-	WidgetBlueprint = InArgs._WidgetBlueprint;
-	check(InArgs._WidgetBlueprint != nullptr);
+	WidgetBlueprint = InWidgetBlueprint;
+	check(InWidgetBlueprint != nullptr);
 
-	Binding = InArgs._Binding;
-	check(Binding != nullptr);
+	BindingId = InArgs._BindingId;
+	check(BindingId.IsValid());
 
 	ParameterName = InArgs._ParameterName;
 	check(!ParameterName.IsNone());
@@ -37,18 +38,21 @@ void SFunctionParameter::Construct(const FArguments& InArgs)
 	bSourceToDestination = InArgs._SourceToDestination;
 	bAllowDefault = InArgs._AllowDefault;
 
-	GetBindingModeDelegate = InArgs._OnGetBindingMode;
-	check(GetBindingModeDelegate.IsBound());
-
 	UMVVMEditorSubsystem* EditorSubsystem = GEditor->GetEditorSubsystem<UMVVMEditorSubsystem>();
-	const UFunction* ConversionFunction = EditorSubsystem->GetConversionFunction(InArgs._WidgetBlueprint, *Binding, bSourceToDestination);
+	UMVVMBlueprintView* View = EditorSubsystem->GetView(InWidgetBlueprint);
+	check(View);
+
+	const FMVVMBlueprintViewBinding* Binding = View->GetBinding(BindingId);
+	check(Binding);
+
+	const UFunction* ConversionFunction = EditorSubsystem->GetConversionFunction(InWidgetBlueprint, *Binding, bSourceToDestination);
 	const FProperty* Property = ConversionFunction->FindPropertyByName(ParameterName);
 	check(Property);
 
 	TSharedRef<SWidget> ValueWidget = SNullWidget::NullWidget;
 	bool bIsBooleanPin = CastField<const FBoolProperty>(Property) != nullptr;
 
-	if (UEdGraphPin* Pin = EditorSubsystem->GetConversionFunctionArgumentPin(InArgs._WidgetBlueprint, *Binding, ParameterName, bSourceToDestination))
+	if (UEdGraphPin* Pin = EditorSubsystem->GetConversionFunctionArgumentPin(InWidgetBlueprint, *Binding, ParameterName, bSourceToDestination))
 	{
 		bIsBooleanPin = Pin->PinType.PinCategory == UEdGraphSchema_K2::PC_Boolean;
 		// create a new pin widget so that we can get the default value widget out of it
@@ -69,7 +73,7 @@ void SFunctionParameter::Construct(const FArguments& InArgs)
 	else if (!bIsBooleanPin)
 	{
 		ValueWidget = SNew(SBorder)
-			.Padding(0)
+			.Padding(0.0f)
 			.BorderImage(FMVVMEditorStyle::Get().GetBrush("FunctionParameter.Border"))
 			[
 				ValueWidget
@@ -79,8 +83,6 @@ void SFunctionParameter::Construct(const FArguments& InArgs)
 	FMVVMBlueprintPropertyPath Path = OnGetSelectedField();
 	bDefaultValueVisible = Path.IsEmpty();
 
-	ValueWidget->SetVisibility(TAttribute<EVisibility>::CreateSP(this, &SFunctionParameter::OnGetVisibility, true));
-
 	bool bFromViewModel = UE::MVVM::IsForwardBinding(Binding->BindingType);
 	TSharedPtr<SHorizontalBox> HBox;
 
@@ -91,22 +93,21 @@ void SFunctionParameter::Construct(const FArguments& InArgs)
 		[
 			SAssignNew(HBox, SHorizontalBox)
 			+ SHorizontalBox::Slot()
-			.Padding(0, 2)
+			.Padding(0.0f, 2.0f)
 			.VAlign(VAlign_Center)
 			[
-				SNew(SOverlay)
-				+ SOverlay::Slot()
+				SNew(SWidgetSwitcher)
+				.WidgetIndex(this, &SFunctionParameter::GetCurrentWidgetIndex)
+				+ SWidgetSwitcher::Slot()
 				[
 					ValueWidget
 				]
-				+ SOverlay::Slot()
+				+ SWidgetSwitcher::Slot()
 				[
-					SAssignNew(FieldSelector, UE::MVVM::SFieldSelector, WidgetBlueprint.Get(), bFromViewModel)
-					.Visibility(this, &SFunctionParameter::OnGetVisibility, false)
-					.SelectedField(this, &SFunctionParameter::OnGetSelectedField)
-					.BindingMode_Lambda([this]() { return GetBindingModeDelegate.Execute(); })
-					.OnFieldSelectionChanged(this, &SFunctionParameter::OnFieldSelectionChanged)
-					.AssignableTo(Property)
+					SNew(SFieldSelector, WidgetBlueprint.Get())
+					.OnGetPropertyPath(this, &SFunctionParameter::OnGetSelectedField)
+					.OnFieldSelectionChanged(this, &SFunctionParameter::HandleFieldSelectionChanged)
+					.OnGetSelectionContext(this, &SFunctionParameter::GetSelectedSelectionContext)
 				]
 			]
 		]
@@ -139,15 +140,9 @@ void SFunctionParameter::Construct(const FArguments& InArgs)
 	}
 }
 
-EVisibility SFunctionParameter::OnGetVisibility(bool bDefaultValue) const
+int32 SFunctionParameter::GetCurrentWidgetIndex() const
 {
-	if (!bAllowDefault)
-	{
-		return bDefaultValue ? EVisibility::Collapsed : EVisibility::Visible;
-	}
-
-	// if we're not bound then show the default value widget, otherwise show the binding widget
-	return bDefaultValue == bDefaultValueVisible ? EVisibility::Visible : EVisibility::Collapsed;
+	return bDefaultValueVisible && bAllowDefault ? 0 : 1;
 }
 
 ECheckBoxState SFunctionParameter::OnGetIsBindArgumentChecked() const 
@@ -174,20 +169,75 @@ void SFunctionParameter::OnBindArgumentChecked(ECheckBoxState Checked)
 
 FMVVMBlueprintPropertyPath SFunctionParameter::OnGetSelectedField() const
 {
-	UMVVMEditorSubsystem* EditorSubsystem = GEditor->GetEditorSubsystem<UMVVMEditorSubsystem>();
-	return EditorSubsystem->GetPathForConversionFunctionArgument(WidgetBlueprint.Get(), *Binding, ParameterName, bSourceToDestination);
+	if (const UWidgetBlueprint* WidgetBlueprintPtr = WidgetBlueprint.Get())
+	{
+		const UMVVMEditorSubsystem* Subsystem = GEditor->GetEditorSubsystem<UMVVMEditorSubsystem>();
+		if (const UMVVMBlueprintView* View = Subsystem->GetView(WidgetBlueprintPtr))
+		{
+			if (const FMVVMBlueprintViewBinding* Binding = View->GetBinding(BindingId))
+			{
+				return Subsystem->GetPathForConversionFunctionArgument(WidgetBlueprint.Get(), *Binding, ParameterName, bSourceToDestination);
+			}
+		}
+	}
+	return FMVVMBlueprintPropertyPath();
 }
 
 void SFunctionParameter::SetSelectedField(const FMVVMBlueprintPropertyPath& Path)
 {
-	UMVVMEditorSubsystem* EditorSubsystem = GEditor->GetEditorSubsystem<UMVVMEditorSubsystem>();
-	EditorSubsystem->SetPathForConversionFunctionArgument(WidgetBlueprint.Get(), *Binding, ParameterName, Path, bSourceToDestination);
+	if (const UWidgetBlueprint* WidgetBlueprintPtr = WidgetBlueprint.Get())
+	{
+		const UMVVMEditorSubsystem* Subsystem = GEditor->GetEditorSubsystem<UMVVMEditorSubsystem>();
+		if (UMVVMBlueprintView* View = Subsystem->GetView(WidgetBlueprintPtr))
+		{
+			if (FMVVMBlueprintViewBinding* Binding = View->GetBinding(BindingId))
+			{
+				Subsystem->SetPathForConversionFunctionArgument(WidgetBlueprint.Get(), *Binding, ParameterName, Path, bSourceToDestination);
+			}
+		}
+	}
 }
 
-void SFunctionParameter::OnFieldSelectionChanged(FMVVMBlueprintPropertyPath Selected)
+void SFunctionParameter::HandleFieldSelectionChanged(FMVVMBlueprintPropertyPath SelectedField, const UFunction* Function)
 {
-	SetSelectedField(Selected);
+	SetSelectedField(SelectedField);
 }
 
+FFieldSelectionContext SFunctionParameter::GetSelectedSelectionContext() const
+{
+	FFieldSelectionContext Result;
+
+	UWidgetBlueprint* WidgetBlueprintPtr = WidgetBlueprint.Get();
+	if (WidgetBlueprintPtr == nullptr)
+	{
+		return Result;
+	}
+
+	UMVVMBlueprintView* View = GEditor->GetEditorSubsystem<UMVVMEditorSubsystem>()->GetView(WidgetBlueprintPtr);
+	check(View);
+
+	const FMVVMBlueprintViewBinding* Binding = View->GetBinding(BindingId);
+	if (Binding == nullptr)
+	{
+		return Result;
+	}
+
+	const UFunction* ConversionFunction = GEditor->GetEditorSubsystem<UMVVMEditorSubsystem>()->GetConversionFunction(WidgetBlueprintPtr, *Binding, bSourceToDestination);
+	if (ConversionFunction == nullptr)
+	{
+		return Result;
+	}
+
+	Result.BindingMode = EMVVMBindingMode::OneWayToDestination;
+	Result.AssignableTo = ConversionFunction->FindPropertyByName(ParameterName);
+	Result.bAllowWidgets = true;
+	Result.bAllowViewModels = true;
+	Result.bAllowConversionFunctions = false;
+	Result.bReadable = true;
+	Result.bWritable = false;
+
+	return Result;
 }
+
+}//namespace
 #undef LOCTEXT_NAMESPACE
diff --git a/Engine/Plugins/Runtime/ModelViewViewModel/Source/ModelViewViewModelEditor/Private/Widgets/SMVVMFunctionParameter.h b/Engine/Plugins/Runtime/ModelViewViewModel/Source/ModelViewViewModelEditor/Private/Widgets/SMVVMFunctionParameter.h
--- a/Engine/Plugins/Runtime/ModelViewViewModel/Source/ModelViewViewModelEditor/Private/Widgets/SMVVMFunctionParameter.h
+++ b/Engine/Plugins/Runtime/ModelViewViewModel/Source/ModelViewViewModelEditor/Private/Widgets/SMVVMFunctionParameter.h
@@ -4,6 +4,7 @@
 
 #include "MVVMPropertyPath.h"
 #include "Widgets/SCompoundWidget.h"
+#include "Widgets/SMVVMFieldSelectorMenu.h"
 
 class SGraphPin;
 class UK2Node_CallFunction;
@@ -22,44 +23,34 @@ class SFieldSelector;
 class SFunctionParameter : public SCompoundWidget
 {
 public:
-	DECLARE_DELEGATE_RetVal(EMVVMBindingMode, FGetBindingMode)
-
-	SLATE_BEGIN_ARGS(SFunctionParameter) :
-		_WidgetBlueprint(nullptr),
-		_Binding(nullptr)
-		{}
-		SLATE_EVENT(FGetBindingMode, OnGetBindingMode)
-		SLATE_ARGUMENT(UWidgetBlueprint*, WidgetBlueprint)
-		SLATE_ARGUMENT(FMVVMBlueprintViewBinding*, Binding)
+	SLATE_BEGIN_ARGS(SFunctionParameter) {}
+		SLATE_ARGUMENT(FGuid, BindingId)
 		SLATE_ARGUMENT(FName, ParameterName)
 		SLATE_ARGUMENT_DEFAULT(bool, SourceToDestination) = true;
 		SLATE_ARGUMENT_DEFAULT(bool, AllowDefault) = true;
 	SLATE_END_ARGS()
 
-	void Construct(const FArguments& InArgs);
+	void Construct(const FArguments& InArgs, UWidgetBlueprint* WidgetBlueprint);
 
 private:
-
 	ECheckBoxState OnGetIsBindArgumentChecked() const;
 	void OnBindArgumentChecked(ECheckBoxState Checked);
-	EVisibility OnGetVisibility(bool bDefaultValue) const;
 
 	FMVVMBlueprintPropertyPath OnGetSelectedField() const;
 	void SetSelectedField(const FMVVMBlueprintPropertyPath& Path);
 
-	void OnFieldSelectionChanged(FMVVMBlueprintPropertyPath Selected);
+	void HandleFieldSelectionChanged(FMVVMBlueprintPropertyPath SelectedField, const UFunction* Function);
+	FFieldSelectionContext GetSelectedSelectionContext() const;
 
-private:
+	int32 GetCurrentWidgetIndex() const;
 
+private:
 	TWeakObjectPtr<UWidgetBlueprint> WidgetBlueprint;
-	FMVVMBlueprintViewBinding* Binding = nullptr;
+	FGuid BindingId;
 	FName ParameterName;
 	/** This reference is just to keep the default value widget alive. */
 	TSharedPtr<SGraphPin> GraphPin;
-	TSharedPtr<UE::MVVM::SFieldSelector> FieldSelector;
-	FGetBindingMode GetBindingModeDelegate;
 
-	// Previously selected field value so that toggling bind on and off doesn't clear the reference.
 	FMVVMBlueprintPropertyPath PreviousSelectedField;
 
 	bool bSourceToDestination = true;
diff --git a/Engine/Plugins/Runtime/ModelViewViewModel/Source/ModelViewViewModelEditor/Private/Widgets/SMVVMPropertyPath.cpp b/Engine/Plugins/Runtime/ModelViewViewModel/Source/ModelViewViewModelEditor/Private/Widgets/SMVVMPropertyPath.cpp
--- a/Engine/Plugins/Runtime/ModelViewViewModel/Source/ModelViewViewModelEditor/Private/Widgets/SMVVMPropertyPath.cpp
+++ b/Engine/Plugins/Runtime/ModelViewViewModel/Source/ModelViewViewModelEditor/Private/Widgets/SMVVMPropertyPath.cpp
@@ -7,9 +7,13 @@
 #include "MVVMBlueprintView.h"
 #include "MVVMEditorSubsystem.h"
 #include "Styling/MVVMEditorStyle.h"
+#include "Types/MVVMBindingSource.h"
 #include "WidgetBlueprint.h"
+#include "Widgets/SMVVMFieldEntry.h"
 #include "Widgets/SMVVMFieldIcon.h"
+#include "Widgets/SMVVMSourceEntry.h"
 #include "Widgets/SBoxPanel.h"
+#include "Widgets/Images/SImage.h"
 #include "Widgets/Text/STextBlock.h"
 
 #define LOCTEXT_NAMESPACE "SPropertyPath"
@@ -17,72 +21,69 @@
 namespace UE::MVVM
 {
 
-void SPropertyPath::Construct(const FArguments& InArgs)
+void SPropertyPath::Construct(const FArguments& InArgs, const UWidgetBlueprint* InWidgetBlueprint)
 {
-	WidgetBlueprint = InArgs._WidgetBlueprint;
-	check(InArgs._WidgetBlueprint != nullptr);
-	PropertyPath = InArgs._PropertyPath;
-	check(PropertyPath);
+	WidgetBlueprint = InWidgetBlueprint;
+	TextStyle = InArgs._TextStyle;
+	bShowContext = InArgs._ShowContext;
+	bShowOnlyLastPath = InArgs._ShowOnlyLastPath;
 
 	ChildSlot
 	[
-		SNew(SHorizontalBox)
-		+ SHorizontalBox::Slot()
-		.VAlign(VAlign_Center)
-		.AutoWidth()
-		[
-			SNew(STextBlock)
-			.TextStyle(FMVVMEditorStyle::Get(), "PropertyPath.ContextText")
-			.Text(this, &SPropertyPath::GetSourceDisplayName)
-		]
-		+ SHorizontalBox::Slot()
-		.Padding(FMargin(5, 0, 0, 0))
-		.VAlign(VAlign_Center)
-		.AutoWidth()
-		[
-			SNew(SFieldIcon)
-			.Field(GetLastField())
-		]
-		+ SHorizontalBox::Slot()
-		.Padding(FMargin(2, 0, 5, 0))
-		.VAlign(VAlign_Center)
-		.AutoWidth()
-		[
-			SNew(STextBlock)
-			.TextStyle(FMVVMEditorStyle::Get(), "PropertyPath.ContextText")
-			.Text(this, &SPropertyPath::GetFieldDisplayName)
-		]
+		SAssignNew(FieldBox, SHorizontalBox)
 	];
+	SetPropertyPath(InArgs._PropertyPath);
 }
 
-FText SPropertyPath::GetSourceDisplayName() const
+
+void SPropertyPath::SetPropertyPath(const FMVVMBlueprintPropertyPath& InPropertyPath)
 {
-	if (PropertyPath->IsFromWidget())
-	{
-		const UWidget* Widget = WidgetBlueprint->WidgetTree->FindWidget(PropertyPath->GetWidgetName());
-		return FText::FromString(Widget->GetDisplayLabel());
-	}
-	else if (PropertyPath->IsFromViewModel())
+	FieldBox->ClearChildren();
+
+	bool bHasSource = InPropertyPath.IsFromWidget() || InPropertyPath.IsFromViewModel();
+	if (bShowContext && bHasSource)
 	{
-		const UMVVMEditorSubsystem* Subsystem = GEditor->GetEditorSubsystem<UMVVMEditorSubsystem>();
-		if (const UMVVMBlueprintView* View = Subsystem->GetView(WidgetBlueprint.Get()))
+		FieldBox->AddSlot()
+			.Padding(8, 0, 0, 0)
+			.HAlign(HAlign_Left)
+			.VAlign(VAlign_Center)
+			.AutoWidth()
+			[
+				SNew(SBindingContextEntry)
+				.TextStyle(TextStyle)
+				.BindingContext(FBindingSource::CreateFromPropertyPath(WidgetBlueprint.Get(), InPropertyPath))
+			];
+
+		if (InPropertyPath.HasPaths())
 		{
-			const FMVVMBlueprintViewModelContext* ViewModel = View->FindViewModel(PropertyPath->GetViewModelId());
-			return ViewModel->GetDisplayName();
+			FieldBox->AddSlot()
+				.Padding(6, 0)
+				.HAlign(HAlign_Left)
+				.VAlign(VAlign_Center)
+				.AutoWidth()
+				[
+					SNew(SImage)
+					.Image(FAppStyle::Get().GetBrush("Icons.ChevronRight"))
+				];
 		}
 	}
-	return FText::GetEmpty();
-}
 
-FText SPropertyPath::GetFieldDisplayName() const
-{
-	return FText::FromString(PropertyPath->GetBasePropertyPath());
+	if (InPropertyPath.HasPaths())
+	{
+		FieldBox->AddSlot()
+			.Padding(0, 0, 8, 0)
+			.HAlign(HAlign_Left)
+			.VAlign(VAlign_Center)
+			.AutoWidth()
+			[
+				SNew(SFieldPaths)
+				.TextStyle(TextStyle)
+				.FieldPaths(InPropertyPath.GetFields())
+				.ShowOnlyLast(bShowOnlyLastPath)
+			];
+	}
 }
 
-FMVVMConstFieldVariant SPropertyPath::GetLastField() const
-{
-	return PropertyPath->GetFields().Last();
-}
 
 } // namespace UE::MVVM
 
diff --git a/Engine/Plugins/Runtime/ModelViewViewModel/Source/ModelViewViewModelEditor/Private/Widgets/SMVVMPropertyPath.h b/Engine/Plugins/Runtime/ModelViewViewModel/Source/ModelViewViewModelEditor/Private/Widgets/SMVVMPropertyPath.h
--- a/Engine/Plugins/Runtime/ModelViewViewModel/Source/ModelViewViewModelEditor/Private/Widgets/SMVVMPropertyPath.h
+++ b/Engine/Plugins/Runtime/ModelViewViewModel/Source/ModelViewViewModelEditor/Private/Widgets/SMVVMPropertyPath.h
@@ -2,11 +2,14 @@
 
 #pragma once
 
+#include "MVVMPropertyPath.h"
+#include "Styling/CoreStyle.h"
+#include "Styling/SlateTypes.h"
+#include "Styling/SlateWidgetStyleAsset.h"
 #include "Widgets/SCompoundWidget.h"
 
+class SHorizontalBox;
 class UWidgetBlueprint;
-namespace UE::MVVM { struct FMVVMConstFieldVariant; }
-struct FMVVMBlueprintPropertyPath;
 
 namespace UE::MVVM
 {
@@ -15,21 +18,23 @@ class SPropertyPath : public SCompoundWidget
 {
 public:
 	SLATE_BEGIN_ARGS(SPropertyPath)
+		: _TextStyle(&FCoreStyle::Get().GetWidgetStyle<FTextBlockStyle>("NormalText"))
 	{}
-		SLATE_ARGUMENT(FMVVMBlueprintPropertyPath*, PropertyPath)
-		SLATE_ARGUMENT(const UWidgetBlueprint*, WidgetBlueprint)
+		SLATE_STYLE_ARGUMENT(FTextBlockStyle, TextStyle)
+		SLATE_ARGUMENT(FMVVMBlueprintPropertyPath, PropertyPath)
+		SLATE_ARGUMENT_DEFAULT(bool, ShowContext) = true;
+		SLATE_ARGUMENT_DEFAULT(bool, ShowOnlyLastPath) = false;
 	SLATE_END_ARGS()
 
-	void Construct(const FArguments& InArgs);
+	void Construct(const FArguments& InArgs, const UWidgetBlueprint* WidgetBlueprint);
+	void SetPropertyPath(const FMVVMBlueprintPropertyPath& InPropertyPath);
 
 private:
-	FText GetSourceDisplayName() const;
-	FText GetFieldDisplayName() const;
-	FMVVMConstFieldVariant GetLastField() const;
-
-private:
-	FMVVMBlueprintPropertyPath* PropertyPath = nullptr;
-	TWeakObjectPtr<const UWidgetBlueprint> WidgetBlueprint = nullptr;
+	TWeakObjectPtr<const UWidgetBlueprint> WidgetBlueprint;
+	const FTextBlockStyle* TextStyle = nullptr;
+	TSharedPtr<SHorizontalBox> FieldBox;
+	bool bShowContext = false;
+	bool bShowOnlyLastPath = false;
 };
 
 } // namespace UE::MVVM
diff --git a/Engine/Plugins/Runtime/ModelViewViewModel/Source/ModelViewViewModelEditor/Private/Widgets/SMVVMSourceEntry.cpp b/Engine/Plugins/Runtime/ModelViewViewModel/Source/ModelViewViewModelEditor/Private/Widgets/SMVVMSourceEntry.cpp
--- a/Engine/Plugins/Runtime/ModelViewViewModel/Source/ModelViewViewModelEditor/Private/Widgets/SMVVMSourceEntry.cpp
+++ b/Engine/Plugins/Runtime/ModelViewViewModel/Source/ModelViewViewModelEditor/Private/Widgets/SMVVMSourceEntry.cpp
@@ -52,7 +52,7 @@ FText GetSourceDisplayName(const FBindingSource& Source)
 
 } // namespace Private
 
-void SSourceEntry::Construct(const FArguments& InArgs)
+void SBindingContextEntry::Construct(const FArguments& InArgs)
 {
 	ChildSlot
 	[
@@ -75,10 +75,10 @@ void SSourceEntry::Construct(const FArguments& InArgs)
 		]
 	];
 
-	RefreshSource(InArgs._Source);
+	RefreshSource(InArgs._BindingContext);
 }
 
-void SSourceEntry::RefreshSource(const FBindingSource& Source)
+void SBindingContextEntry::RefreshSource(const FBindingSource& Source)
 {
 	Image->SetImage(Private::GetSourceIcon(Source));
 	Image->SetColorAndOpacity(Private::GetSourceColor(Source));
diff --git a/Engine/Plugins/Runtime/ModelViewViewModel/Source/ModelViewViewModelEditor/Private/Widgets/SMVVMSourceEntry.h b/Engine/Plugins/Runtime/ModelViewViewModel/Source/ModelViewViewModelEditor/Private/Widgets/SMVVMSourceEntry.h
--- a/Engine/Plugins/Runtime/ModelViewViewModel/Source/ModelViewViewModelEditor/Private/Widgets/SMVVMSourceEntry.h
+++ b/Engine/Plugins/Runtime/ModelViewViewModel/Source/ModelViewViewModelEditor/Private/Widgets/SMVVMSourceEntry.h
@@ -14,15 +14,15 @@ class STextBlock;
 namespace UE::MVVM
 {
 
-class SSourceEntry : public SCompoundWidget
+class SBindingContextEntry : public SCompoundWidget
 {
 public:
-	SLATE_BEGIN_ARGS(SSourceEntry) :
-		_TextStyle(&FCoreStyle::Get().GetWidgetStyle<FTextBlockStyle>("NormalText"))
+	SLATE_BEGIN_ARGS(SBindingContextEntry)
+		: _TextStyle(&FCoreStyle::Get().GetWidgetStyle<FTextBlockStyle>("NormalText"))
 		{
 		}
 		SLATE_STYLE_ARGUMENT(FTextBlockStyle, TextStyle)
-		SLATE_ARGUMENT(UE::MVVM::FBindingSource, Source)
+		SLATE_ARGUMENT(UE::MVVM::FBindingSource, BindingContext)
 	SLATE_END_ARGS()
 
 	void Construct(const FArguments& Args);
diff --git a/Engine/Plugins/Runtime/ModelViewViewModel/Source/ModelViewViewModelEditor/Private/Widgets/SMVVMSourceSelector.cpp b/Engine/Plugins/Runtime/ModelViewViewModel/Source/ModelViewViewModelEditor/Private/Widgets/SMVVMSourceSelector.cpp
--- a/Engine/Plugins/Runtime/ModelViewViewModel/Source/ModelViewViewModelEditor/Private/Widgets/SMVVMSourceSelector.cpp
+++ b/Engine/Plugins/Runtime/ModelViewViewModel/Source/ModelViewViewModelEditor/Private/Widgets/SMVVMSourceSelector.cpp
@@ -18,11 +18,11 @@
 namespace UE::MVVM
 {
 
-void SSourceSelector::Construct(const FArguments& Args, const UWidgetBlueprint* InWidgetBlueprint)
+void SBindingContextSelector::Construct(const FArguments& Args, const UWidgetBlueprint* InWidgetBlueprint)
 {
 	TextStyle = Args._TextStyle;
 
-	SelectedSourceAttribute = Args._SelectedSource;
+	SelectedSourceAttribute = Args._SelectedBindingSource;
 	check(SelectedSourceAttribute.IsSet());
 
 	InitialSource = SelectedSourceAttribute.Get();
@@ -37,7 +37,7 @@ void SSourceSelector::Construct(const FArguments& Args, const UWidgetBlueprint*
 	Refresh();
 
 	MenuAnchor = SNew(SMenuAnchor)
-		.OnGetMenuContent(this, &SSourceSelector::OnGetMenuContent)
+		.OnGetMenuContent(this, &SBindingContextSelector::OnGetMenuContent)
 		[
 			SNew(SButton)
 			.ButtonStyle(&FAppStyle::Get().GetWidgetStyle<FComboButtonStyle>("ComboButton").ButtonStyle)
@@ -52,8 +52,8 @@ void SSourceSelector::Construct(const FArguments& Args, const UWidgetBlueprint*
 				.Padding(0, 2)
 				.HAlign(HAlign_Fill)
 				[
-					SAssignNew(SelectedSourceWidget, SSourceEntry)
-					.Source(SelectedSource)
+					SAssignNew(SelectedSourceWidget, SBindingContextEntry)
+					.BindingContext(SelectedSource)
 					.TextStyle(TextStyle)
 				]
 				+ SHorizontalBox::Slot()
@@ -87,8 +87,8 @@ void SSourceSelector::Construct(const FArguments& Args, const UWidgetBlueprint*
 				SNew(SSimpleButton)
 				.Icon(FAppStyle::Get().GetBrush("Icons.X"))
 				.ToolTipText(LOCTEXT("ClearField", "Clear source selection."))
-				.Visibility(this, &SSourceSelector::GetClearVisibility)
-				.OnClicked(this, &SSourceSelector::OnClearSource)
+				.Visibility(this, &SBindingContextSelector::GetClearVisibility)
+				.OnClicked(this, &SBindingContextSelector::OnClearSource)
 			]
 		];
 	}
@@ -104,15 +104,15 @@ void SSourceSelector::Construct(const FArguments& Args, const UWidgetBlueprint*
 	check(OnSelectionChangedDelegate.IsBound());
 }
 
-TSharedRef<SWidget> SSourceSelector::OnGetMenuContent()
+TSharedRef<SWidget> SBindingContextSelector::OnGetMenuContent()
 {
 	TSharedRef<SVerticalBox> VBox = SNew(SVerticalBox);
 
 	if (bViewModels)
 	{
 		ViewModelList = SNew(SListView<FBindingSource>)
-			.OnGenerateRow(this, &SSourceSelector::OnGenerateRow)
-			.OnSelectionChanged(this, &SSourceSelector::OnViewModelSelectionChanged)
+			.OnGenerateRow(this, &SBindingContextSelector::OnGenerateRow)
+			.OnSelectionChanged(this, &SBindingContextSelector::OnViewModelSelectionChanged)
 			.SelectionMode(ESelectionMode::Single)
 			.ListItemsSource(&ViewModelSources);
 
@@ -130,7 +130,7 @@ TSharedRef<SWidget> SSourceSelector::OnGetMenuContent()
 	else
 	{
 		WidgetHierarchy = SNew(SReadOnlyHierarchyView, WidgetBlueprint.Get())
-			.OnSelectionChanged(this, &SSourceSelector::OnWidgetSelectionChanged)
+			.OnSelectionChanged(this, &SBindingContextSelector::OnWidgetSelectionChanged)
 			.ShowSearch(false);
 
 		VBox->AddSlot()
@@ -158,8 +158,8 @@ TSharedRef<SWidget> SSourceSelector::OnGetMenuContent()
 			[
 				SNew(SPrimaryButton)
 				.Text(LOCTEXT("Select", "Select"))
-				.OnClicked(this, &SSourceSelector::HandleSelect)
-				.IsEnabled(this, &SSourceSelector::IsSelectEnabled)
+				.OnClicked(this, &SBindingContextSelector::HandleSelect)
+				.IsEnabled(this, &SBindingContextSelector::IsSelectEnabled)
 			]
 			+ SHorizontalBox::Slot()
 			.Padding(FAppStyle::Get().GetMargin("StandardDialog.SlotPadding"))
@@ -168,7 +168,7 @@ TSharedRef<SWidget> SSourceSelector::OnGetMenuContent()
 				SNew(SButton)
 				.Text(LOCTEXT("Cancel", "Cancel"))
 				.HAlign(HAlign_Center)
-				.OnClicked(this, &SSourceSelector::HandleCancel)
+				.OnClicked(this, &SBindingContextSelector::HandleCancel)
 			]
 		];
 
@@ -185,12 +185,12 @@ TSharedRef<SWidget> SSourceSelector::OnGetMenuContent()
 		];
 }
 
-bool SSourceSelector::IsSelectEnabled() const
+bool SBindingContextSelector::IsSelectEnabled() const
 {
 	return SelectedSource.IsValid();
 }
 
-FReply SSourceSelector::HandleCancel()
+FReply SBindingContextSelector::HandleCancel()
 {
 	OnViewModelSelectionChanged(InitialSource, ESelectInfo::Direct);
 
@@ -198,7 +198,7 @@ FReply SSourceSelector::HandleCancel()
 	return FReply::Handled();
 }
 
-FReply SSourceSelector::HandleSelect()
+FReply SBindingContextSelector::HandleSelect()
 {
 	if (SelectedSourceWidget.IsValid())
 	{
@@ -212,19 +212,19 @@ FReply SSourceSelector::HandleSelect()
 	return FReply::Handled();
 }
 
-void SSourceSelector::OnViewModelSelectionChanged(FBindingSource Selected, ESelectInfo::Type SelectionType)
+void SBindingContextSelector::OnViewModelSelectionChanged(FBindingSource Selected, ESelectInfo::Type SelectionType)
 {
 	SelectedSource = Selected;
 }
 
-void SSourceSelector::OnWidgetSelectionChanged(FName SelectedName, ESelectInfo::Type SelectionType)
+void SBindingContextSelector::OnWidgetSelectionChanged(FName SelectedName, ESelectInfo::Type SelectionType)
 {
 	FBindingSource Selected = FBindingSource::CreateForWidget(WidgetBlueprint.Get(), SelectedName);
 
 	SelectedSource = Selected;
 }
 
-void SSourceSelector::Refresh()
+void SBindingContextSelector::Refresh()
 {
 	SelectedSource = SelectedSourceAttribute.Get();
 
@@ -271,12 +271,12 @@ void SSourceSelector::Refresh()
 	}
 }
 
-EVisibility SSourceSelector::GetClearVisibility() const
+EVisibility SBindingContextSelector::GetClearVisibility() const
 {
 	return SelectedSource.IsValid() ? EVisibility::Visible : EVisibility::Collapsed;
 }
 
-FReply SSourceSelector::OnClearSource()
+FReply SBindingContextSelector::OnClearSource()
 {
 	if (ViewModelList.IsValid())
 	{
@@ -286,15 +286,15 @@ FReply SSourceSelector::OnClearSource()
 	return FReply::Handled();
 }
 
-TSharedRef<ITableRow> SSourceSelector::OnGenerateRow(FBindingSource Source, const TSharedRef<STableViewBase>& OwnerTable)
+TSharedRef<ITableRow> SBindingContextSelector::OnGenerateRow(FBindingSource Source, const TSharedRef<STableViewBase>& OwnerTable)
 {
 	return SNew(STableRow<FBindingSource>, OwnerTable)
 	[
 		SNew(SBox)
 		.Padding(4, 3)
 		[
-			SNew(SSourceEntry)
-			.Source(Source)
+			SNew(SBindingContextEntry)
+			.BindingContext(Source)
 			.TextStyle(TextStyle)
 		]
 	];
diff --git a/Engine/Plugins/Runtime/ModelViewViewModel/Source/ModelViewViewModelEditor/Private/Widgets/SMVVMSourceSelector.h b/Engine/Plugins/Runtime/ModelViewViewModel/Source/ModelViewViewModelEditor/Private/Widgets/SMVVMSourceSelector.h
--- a/Engine/Plugins/Runtime/ModelViewViewModel/Source/ModelViewViewModelEditor/Private/Widgets/SMVVMSourceSelector.h
+++ b/Engine/Plugins/Runtime/ModelViewViewModel/Source/ModelViewViewModelEditor/Private/Widgets/SMVVMSourceSelector.h
@@ -17,14 +17,14 @@ template <typename ItemType> class SListView;
 namespace UE::MVVM
 {
 
-class SSourceEntry;
+class SBindingContextEntry;
 
-class SSourceSelector : public SCompoundWidget
+class SBindingContextSelector : public SCompoundWidget
 {
 public:
 	DECLARE_DELEGATE_OneParam(FSelectionChanged, FBindingSource);
 
-	SLATE_BEGIN_ARGS(SSourceSelector) :
+	SLATE_BEGIN_ARGS(SBindingContextSelector) :
 		_TextStyle( &FCoreStyle::Get().GetWidgetStyle<FTextBlockStyle>( "NormalText" ) )
 		{
 		}
@@ -32,7 +32,7 @@ public:
 		SLATE_ARGUMENT_DEFAULT(bool, ShowClear) = true;
 		SLATE_ARGUMENT_DEFAULT(bool, AutoRefresh) = false;
 		SLATE_ARGUMENT_DEFAULT(bool, ViewModels) = false;
-		SLATE_ATTRIBUTE(FBindingSource, SelectedSource)
+		SLATE_ATTRIBUTE(FBindingSource, SelectedBindingSource)
 		SLATE_EVENT(FSelectionChanged, OnSelectionChanged)
 	SLATE_END_ARGS()
 
@@ -68,7 +68,7 @@ private:
 	TSharedPtr<SMenuAnchor> MenuAnchor;
 	TSharedPtr<SListView<FBindingSource>> ViewModelList;
 	TSharedPtr<SReadOnlyHierarchyView> WidgetHierarchy;
-	TSharedPtr<SSourceEntry> SelectedSourceWidget;
+	TSharedPtr<SBindingContextEntry> SelectedSourceWidget;
 
 	bool bAutoRefresh = false;
 	bool bViewModels = false;
diff --git a/Engine/Plugins/Runtime/ModelViewViewModel/Source/ModelViewViewModelEditor/Private/Widgets/SMVVMViewBindingListView.cpp b/Engine/Plugins/Runtime/ModelViewViewModel/Source/ModelViewViewModelEditor/Private/Widgets/SMVVMViewBindingListView.cpp
--- a/Engine/Plugins/Runtime/ModelViewViewModel/Source/ModelViewViewModelEditor/Private/Widgets/SMVVMViewBindingListView.cpp
+++ b/Engine/Plugins/Runtime/ModelViewViewModel/Source/ModelViewViewModelEditor/Private/Widgets/SMVVMViewBindingListView.cpp
@@ -27,6 +27,7 @@
 #include "Widgets/Input/SComboBox.h"
 #include "Widgets/Input/SEditableTextBox.h"
 #include "Widgets/SMVVMViewModelPanel.h" // IWYU pragma: keep
+#include "Widgets/SMVVMViewModelBindingListWidget.h"
 #include "Widgets/Views/STreeView.h"
 
 #define LOCTEXT_NAMESPACE "BindingListView"
@@ -183,11 +184,11 @@ public:
 					SNew(SBox)
 					.MinDesiredWidth(150)
 					[
-						SNew(SSourceSelector, InWidgetBlueprint)
+						SNew(SBindingContextSelector, InWidgetBlueprint)
 						.ShowClear(false)
 						.AutoRefresh(true)
 						.ViewModels(false)
-						.SelectedSource(this, &SWidgetRow::GetSelectedWidget)
+						.SelectedBindingSource(this, &SWidgetRow::GetSelectedWidget)
 						.OnSelectionChanged(this, &SWidgetRow::SetSelectedWidget)
 					]
 				]
@@ -282,11 +283,9 @@ public:
 		Entry = InEntry;
 		WidgetBlueprintWeak = InWidgetBlueprint;
 
-		OnBlueprintChangedHandle = InWidgetBlueprint->OnChanged().AddSP(this, &SBindingRow::HandleBlueprintChanged);
 		CVarDefaultExecutionMode = IConsoleManager::Get().FindConsoleVariable(TEXT("MVVM.DefaultExecutionMode"));
 		
 		FMVVMBlueprintViewBinding* ViewBinding = GetThisViewBinding();
-		FBindingSource WidgetSource = FBindingSource::CreateForWidget(InWidgetBlueprint, ViewBinding->DestinationPath.GetWidgetName());
 
 		STableRow<TSharedPtr<FBindingEntry>>::Construct(
 			STableRow<TSharedPtr<FBindingEntry>>::FArguments()
@@ -336,16 +335,12 @@ public:
 							SNew(SBox)
 							.MinDesiredWidth(150)
 							[
-								SAssignNew(WidgetFieldSelector, SFieldSelector, InWidgetBlueprint, false)
-								.BindingMode(this, &SBindingRow::GetCurrentBindingMode)
-								.SelectedField(this, &SBindingRow::GetSelectedWidgetProperty)
-								.OnFieldSelectionChanged(this, &SBindingRow::OnWidgetPropertySelected)
-								.SelectedConversionFunction(this, &SBindingRow::GetSelectedConversionFunction, false)
-								.OnConversionFunctionSelectionChanged(this, &SBindingRow::OnConversionFunctionChanged, false)
-								.ShowConversionFunctions(this, &SBindingRow::ShouldShowConversionFunctions, false)
-								.ShowSource(false)
-								.Source(WidgetSource)
-								.AssignableTo(this, &SBindingRow::GetAssignableToProperty, false)
+								SNew(SFieldSelector, InWidgetBlueprint)
+								.OnGetPropertyPath(this, &SBindingRow::GetSelectedPropertyPath, false)
+								.OnGetConversionFunction(this, &SBindingRow::GetSelectedConversionFunction, false)
+								.OnFieldSelectionChanged(this, &SBindingRow::HandleFieldSelectionChanged, false)
+								.OnGetSelectionContext(this, &SBindingRow::GetSelectedSelectionContext, false)
+								.ShowContext(false)
 							]
 						]
 
@@ -384,14 +379,11 @@ public:
 							SNew(SBox)
 							.MinDesiredWidth(150)
 							[
-								SAssignNew(ViewModelFieldSelector, SFieldSelector, InWidgetBlueprint, true)
-								.BindingMode(this, &SBindingRow::GetCurrentBindingMode)
-								.SelectedField(this, &SBindingRow::GetSelectedViewModelProperty)
-								.OnFieldSelectionChanged(this, &SBindingRow::OnViewModelPropertySelected)
-								.SelectedConversionFunction(this, &SBindingRow::GetSelectedConversionFunction, true)
-								.OnConversionFunctionSelectionChanged(this, &SBindingRow::OnConversionFunctionChanged, true)
-								.ShowConversionFunctions(this, &SBindingRow::ShouldShowConversionFunctions, true)
-								.AssignableTo(this, &SBindingRow::GetAssignableToProperty, true)
+								SNew(SFieldSelector, InWidgetBlueprint)
+								.OnGetPropertyPath(this, &SBindingRow::GetSelectedPropertyPath, true)
+								.OnGetConversionFunction(this, &SBindingRow::GetSelectedConversionFunction, true)
+								.OnFieldSelectionChanged(this, &SBindingRow::HandleFieldSelectionChanged, true)
+								.OnGetSelectionContext(this, &SBindingRow::GetSelectedSelectionContext, true)
 							]
 						]
 
@@ -453,14 +445,6 @@ public:
 		];
 	}
 
-	~SBindingRow()
-	{
-		if (UWidgetBlueprint* WidgetBlueprint = WidgetBlueprintWeak.Get())
-		{
-			WidgetBlueprint->OnChanged().Remove(OnBlueprintChangedHandle);
-		}
-	}
-
 	FMVVMBlueprintViewBinding* GetThisViewBinding() const
 	{
 		UMVVMEditorSubsystem* EditorSubsystem = GEditor->GetEditorSubsystem<UMVVMEditorSubsystem>();
@@ -486,20 +470,20 @@ private:
 
 	FSlateColor GetErrorBorderColor() const
 	{
-		UMVVMEditorSubsystem* EditorSubsystem = GEditor->GetEditorSubsystem<UMVVMEditorSubsystem>();
-		UMVVMBlueprintView* BlueprintViewPtr = EditorSubsystem->GetView(WidgetBlueprintWeak.Get());
-		if (BlueprintViewPtr->HasBindingMessage(GetThisViewBinding()->BindingId, EBindingMessageType::Error))
-		{
-			return FStyleColors::Error;
-		}
-		else if (BlueprintViewPtr->HasBindingMessage(GetThisViewBinding()->BindingId, EBindingMessageType::Warning))
-		{
-			return FStyleColors::Warning;
-		}
-		else
+		if (const FMVVMBlueprintViewBinding* ViewModelBinding = GetThisViewBinding())
 		{
-			return FStyleColors::Transparent;
+			UMVVMEditorSubsystem* EditorSubsystem = GEditor->GetEditorSubsystem<UMVVMEditorSubsystem>();
+			UMVVMBlueprintView* BlueprintViewPtr = EditorSubsystem->GetView(WidgetBlueprintWeak.Get());
+			if (BlueprintViewPtr->HasBindingMessage(ViewModelBinding->BindingId, EBindingMessageType::Error))
+			{
+				return FStyleColors::Error;
+			}
+			else if (BlueprintViewPtr->HasBindingMessage(ViewModelBinding->BindingId, EBindingMessageType::Warning))
+			{
+				return FStyleColors::Warning;
+			}
 		}
+		return FStyleColors::Transparent;
 	}
 
 	ECheckBoxState IsBindingEnabled() const
@@ -522,25 +506,33 @@ private:
 
 	EVisibility GetErrorButtonVisibility() const
 	{
-		UMVVMEditorSubsystem* EditorSubsystem = GEditor->GetEditorSubsystem<UMVVMEditorSubsystem>();
-		UMVVMBlueprintView* BlueprintViewPtr = EditorSubsystem->GetView(WidgetBlueprintWeak.Get());
-		bool HasBindingError = BlueprintViewPtr->HasBindingMessage(GetThisViewBinding()->BindingId, EBindingMessageType::Error) || BlueprintViewPtr->HasBindingMessage(GetThisViewBinding()->BindingId, EBindingMessageType::Warning);
-		return HasBindingError ? EVisibility::Visible : EVisibility::Collapsed;
+		if (const FMVVMBlueprintViewBinding* ViewModelBinding = GetThisViewBinding())
+		{
+			UMVVMEditorSubsystem* EditorSubsystem = GEditor->GetEditorSubsystem<UMVVMEditorSubsystem>();
+			UMVVMBlueprintView* BlueprintViewPtr = EditorSubsystem->GetView(WidgetBlueprintWeak.Get());
+			bool HasBindingError = BlueprintViewPtr->HasBindingMessage(GetThisViewBinding()->BindingId, EBindingMessageType::Error) || BlueprintViewPtr->HasBindingMessage(GetThisViewBinding()->BindingId, EBindingMessageType::Warning);
+			return HasBindingError ? EVisibility::Visible : EVisibility::Collapsed;
+		}
+		return EVisibility::Collapsed;
 	}
 
 	FText GetErrorButtonToolTip() const
 	{
 		// Get error messages of this binding
-		UMVVMEditorSubsystem* EditorSubsystem = GEditor->GetEditorSubsystem<UMVVMEditorSubsystem>();
-		UMVVMBlueprintView* BlueprintViewPtr = EditorSubsystem->GetView(WidgetBlueprintWeak.Get());
-		TArray<FText> BindingErrorList = BlueprintViewPtr->GetBindingMessages(GetThisViewBinding()->BindingId, EBindingMessageType::Error);
-		TArray<FText> BindingWarningList = BlueprintViewPtr->GetBindingMessages(GetThisViewBinding()->BindingId, EBindingMessageType::Warning);
-		BindingErrorList.Append(BindingWarningList);
-
-		static const FText NewLineText = FText::FromString(TEXT("\n"));
-		FText HintText = LOCTEXT("ErrorButtonText", "Errors and Warnings: (Click to show in a separate window)");
-		FText ErrorsText = FText::Join(NewLineText, BindingErrorList);
-		return FText::Join(NewLineText, HintText, ErrorsText);
+		if (const FMVVMBlueprintViewBinding* ViewModelBinding = GetThisViewBinding())
+		{
+			UMVVMEditorSubsystem* EditorSubsystem = GEditor->GetEditorSubsystem<UMVVMEditorSubsystem>();
+			UMVVMBlueprintView* BlueprintViewPtr = EditorSubsystem->GetView(WidgetBlueprintWeak.Get());
+			TArray<FText> BindingErrorList = BlueprintViewPtr->GetBindingMessages(ViewModelBinding->BindingId, EBindingMessageType::Error);
+			TArray<FText> BindingWarningList = BlueprintViewPtr->GetBindingMessages(ViewModelBinding->BindingId, EBindingMessageType::Warning);
+			BindingErrorList.Append(BindingWarningList);
+
+			static const FText NewLineText = FText::FromString(TEXT("\n"));
+			FText HintText = LOCTEXT("ErrorButtonText", "Errors and Warnings: (Click to show in a separate window)");
+			FText ErrorsText = FText::Join(NewLineText, BindingErrorList);
+			return FText::Join(NewLineText, HintText, ErrorsText);
+		}
+		return FText();
 	}
 
 	FReply OnErrorButtonClicked()
@@ -584,65 +576,6 @@ private:
 		return FReply::Handled();
 	}
 
-	EMVVMBindingMode GetCurrentBindingMode() const
-	{
-		if (const FMVVMBlueprintViewBinding* ViewModelBinding = GetThisViewBinding())
-		{
-			return ViewModelBinding->BindingType;
-		}
-		return EMVVMBindingMode::OneWayToDestination;
-	}
-
-	bool ShouldShowConversionFunctions(bool bViewModel) const
-	{
-		EMVVMBindingMode Mode = GetCurrentBindingMode();
-		if (IsForwardBinding(Mode))
-		{
-			return bViewModel;
-		}
-		else if (IsBackwardBinding(Mode))
-		{
-			return !bViewModel;
-		}
-		
-		return false;
-	}
-
-	const FProperty* GetAssignableToProperty(bool bViewModel) const
-	{
-		if (FMVVMBlueprintViewBinding* ViewBinding = GetThisViewBinding())
-		{
-			TArray<FMVVMConstFieldVariant> Fields;
-			if (bViewModel)
-			{
-				Fields = ViewBinding->DestinationPath.GetFields();
-			}
-			else
-			{
-				Fields = ViewBinding->SourcePath.GetFields();
-			}
-
-			if (Fields.Num() > 0)
-			{
-				FMVVMConstFieldVariant LastField = Fields.Last();
-				if (LastField.IsProperty())
-				{
-					return LastField.GetProperty();
-				}
-				else if (LastField.IsFunction())
-				{
-					const UFunction* Function = LastField.GetFunction();
-					if (Function != nullptr)
-					{
-						return BindingHelper::GetReturnProperty(Function);
-					}
-				}
-			}
-		}
-
-		return nullptr;
-	}
-
 	TSharedRef<ITableRow> OnGenerateErrorRow(TSharedPtr<FText> Text, const TSharedRef<STableViewBase>& TableView) const
 	{
 		return SNew(STableRow<TSharedPtr<FText>>, TableView)
@@ -655,43 +588,17 @@ private:
 			];
 	}
 
-	void OnConversionFunctionChanged(const UFunction* Function, bool bSourceToDest)
-	{
-		if (FMVVMBlueprintViewBinding* ViewBinding = GetThisViewBinding())
-		{
-			UMVVMEditorSubsystem* EditorSubsystem = GEditor->GetEditorSubsystem<UMVVMEditorSubsystem>();
-
-			if (bSourceToDest)
-			{
-				EditorSubsystem->SetSourceToDestinationConversionFunction(WidgetBlueprintWeak.Get(), *ViewBinding, Function);
-			} 
-			else
-			{
-				EditorSubsystem->SetDestinationToSourceConversionFunction(WidgetBlueprintWeak.Get(), *ViewBinding, Function);
-			}
-		}
-	}
-
 	TArray<FBindingSource> GetAvailableViewModels() const
 	{
 		UMVVMEditorSubsystem* EditorSubsystem = GEditor->GetEditorSubsystem<UMVVMEditorSubsystem>();
 		return EditorSubsystem->GetAllViewModels(WidgetBlueprintWeak.Get());
 	}
 
-	FMVVMBlueprintPropertyPath GetSelectedViewModelProperty() const
+	FMVVMBlueprintPropertyPath GetSelectedPropertyPath(bool bSource) const
 	{
 		if (FMVVMBlueprintViewBinding* ViewBinding = GetThisViewBinding())
 		{
-			return ViewBinding->SourcePath;
-		}
-		return FMVVMBlueprintPropertyPath();
-	}
-
-	FMVVMBlueprintPropertyPath GetSelectedWidgetProperty() const
-	{
-		if (FMVVMBlueprintViewBinding* ViewBinding = GetThisViewBinding())
-		{
-			return ViewBinding->DestinationPath;
+			return bSource ? ViewBinding->SourcePath : ViewBinding->DestinationPath;
 		}
 		return FMVVMBlueprintPropertyPath();
 	}
@@ -707,43 +614,84 @@ private:
 		return nullptr;
 	}
 
-	void OnViewModelPropertySelected(FMVVMBlueprintPropertyPath SelectedField)
+	void HandleFieldSelectionChanged(FMVVMBlueprintPropertyPath SelectedField, const UFunction* Function, bool bSource)
 	{
+		UWidgetBlueprint* WidgetBlueprintPtr = WidgetBlueprintWeak.Get();
 		if (FMVVMBlueprintViewBinding* ViewBinding = GetThisViewBinding())
 		{
-			if (ViewBinding->SourcePath != SelectedField)
+			UMVVMEditorSubsystem* Subsystem = GEditor->GetEditorSubsystem<UMVVMEditorSubsystem>();
+			if (bSource)
 			{
-				UMVVMEditorSubsystem* Subsystem = GEditor->GetEditorSubsystem<UMVVMEditorSubsystem>();
-				Subsystem->SetViewModelPropertyForBinding(WidgetBlueprintWeak.Get(), *ViewBinding, SelectedField);
-
-				if (WidgetFieldSelector.IsValid())
+				Subsystem->SetSourceToDestinationConversionFunction(WidgetBlueprintPtr, *ViewBinding, Function);
+				if (ViewBinding->SourcePath != SelectedField)
+				{
+					Subsystem->SetViewModelPropertyForBinding(WidgetBlueprintPtr, *ViewBinding, SelectedField);
+				}
+			}
+			else
+			{
+				Subsystem->SetDestinationToSourceConversionFunction(WidgetBlueprintPtr, *ViewBinding, Function);
+				if (ViewBinding->DestinationPath != SelectedField)
 				{
-					WidgetFieldSelector->Refresh();
+					Subsystem->SetWidgetPropertyForBinding(WidgetBlueprintPtr, *ViewBinding, SelectedField);
 				}
 			}
 		}
 	}
 
-	void OnWidgetPropertySelected(FMVVMBlueprintPropertyPath SelectedField)
+	FFieldSelectionContext GetSelectedSelectionContext(bool bSource) const
 	{
-		if (FMVVMBlueprintViewBinding* ViewBinding = GetThisViewBinding())
+		FFieldSelectionContext Result;
+
+		const UWidgetBlueprint* WidgetBlueprintPtr = WidgetBlueprintWeak.Get();
+		if (WidgetBlueprintPtr == nullptr)
 		{
-			if (ViewBinding->DestinationPath != SelectedField)
-			{
-				UMVVMEditorSubsystem* Subsystem = GEditor->GetEditorSubsystem<UMVVMEditorSubsystem>();
-				Subsystem->SetWidgetPropertyForBinding(WidgetBlueprintWeak.Get(), *ViewBinding, SelectedField);
+			return Result;
+		}
+		
+		if (const FMVVMBlueprintViewBinding* ViewBinding = GetThisViewBinding())
+		{
+			Result.BindingMode = ViewBinding->BindingType;
 
-				if (ViewModelFieldSelector.IsValid())
+			{
+				TArray<FMVVMConstFieldVariant> Fields = bSource ? ViewBinding->DestinationPath.GetFields() : ViewBinding->SourcePath.GetFields();
+				if (Fields.Num() > 0)
 				{
-					ViewModelFieldSelector->Refresh();
+					FMVVMConstFieldVariant LastField = Fields.Last();
+					if (LastField.IsProperty())
+					{
+						Result.AssignableTo = LastField.GetProperty();
+					}
+					else if (LastField.IsFunction())
+					{
+						if (const UFunction* Function = LastField.GetFunction())
+						{
+							Result.AssignableTo = BindingHelper::GetReturnProperty(Function);
+						}
+					}
 				}
+			}
 
-				if (SelectedWidget.IsValid())
-				{
-					SelectedWidget->SetField(SelectedField);
-				}
+			if (!bSource && !ViewBinding->DestinationPath.GetWidgetName().IsNone())
+			{
+				Result.FixedBindingSource = FBindingSource::CreateForWidget(WidgetBlueprintPtr, ViewBinding->DestinationPath.GetWidgetName());
+			}
+
+			Result.bAllowWidgets = !bSource;
+			Result.bAllowViewModels = bSource;
+			Result.bAllowConversionFunctions = false;
+			bool bIsReadingValue = (IsForwardBinding(ViewBinding->BindingType) && bSource)
+				|| (IsBackwardBinding(ViewBinding->BindingType) && !bSource);
+			if (!(IsBackwardBinding(ViewBinding->BindingType) && IsForwardBinding(ViewBinding->BindingType)))
+			{
+				Result.bAllowConversionFunctions = bIsReadingValue;
 			}
+
+			Result.bReadable = bIsReadingValue;
+			Result.bWritable = (IsForwardBinding(ViewBinding->BindingType) && !bSource)
+				|| (IsBackwardBinding(ViewBinding->BindingType) && bSource);
 		}
+		return Result;
 	}
 
 	ECheckBoxState IsExecutionModeOverrideChecked() const
@@ -976,45 +924,14 @@ private:
 
 			UMVVMEditorSubsystem* Subsystem = GEditor->GetEditorSubsystem<UMVVMEditorSubsystem>();
 			Subsystem->SetBindingTypeForBinding(WidgetBlueprintWeak.Get(), *ViewBinding, NewMode);
-
-			if (ViewModelFieldSelector.IsValid())
-			{
-				ViewModelFieldSelector->Refresh();
-			}
-
-			if (WidgetFieldSelector.IsValid())
-			{
-				WidgetFieldSelector->Refresh();
-			}
-		}
-	}
-
-	void HandleBlueprintChanged(UBlueprint* Blueprint)
-	{
-		if (ViewModelFieldSelector.IsValid())
-		{
-			ViewModelFieldSelector->Refresh();
-		}
-
-		if (WidgetFieldSelector.IsValid())
-		{
-			WidgetFieldSelector->Refresh();
 		}
 	}
 
 private:
 	TSharedPtr<FBindingEntry> Entry;
 	TWeakObjectPtr<UWidgetBlueprint> WidgetBlueprintWeak;
-	TSharedPtr<SFieldSelector> WidgetFieldSelector;
-	TSharedPtr<SFieldSelector> ViewModelFieldSelector;
-	TSharedPtr<SWidget> ContextMenuOptionHelper;
 	TSharedPtr<SCustomDialog> ErrorDialog;
 	TArray<TSharedPtr<FText>> ErrorItems;
-	FDelegateHandle OnBlueprintChangedHandle;
-	TSharedPtr<UE::MVVM::FFieldIterator_Bindable> WidgetFieldIterator;
-	TSharedPtr<SMenuAnchor> WidgetMenuAnchor;
-	TSharedPtr<SSourceBindingList> WidgetBindingList;
-	TSharedPtr<SFieldEntry> SelectedWidget;
 	IConsoleVariable* CVarDefaultExecutionMode = nullptr;
 };
 
@@ -1100,17 +1017,11 @@ class SFunctionParameterRow : public STableRow<TSharedPtr<FBindingEntry>>
 					.VAlign(VAlign_Center)
 					.AutoWidth()
 					[
-						SNew(SBox)
-						.MinDesiredWidth(200)
-						[
-							SNew(SFunctionParameter)
-							.OnGetBindingMode(this, &SFunctionParameterRow::OnGetBindingMode)
-							.WidgetBlueprint(InWidgetBlueprint)
-							.Binding(Binding)
-							.ParameterName(Entry->GetName())
-							.SourceToDestination(bSourceToDestination)
-							.AllowDefault(!bSimpleConversionFunction)
-						]
+						SNew(SFunctionParameter, InWidgetBlueprint)
+						.BindingId(Binding->BindingId)
+						.ParameterName(Entry->GetName())
+						.SourceToDestination(bSourceToDestination)
+						.AllowDefault(!bSimpleConversionFunction)
 					]
 				]
 			], 
