MVVM: Fix display of getter functions on Conversion function arguments. Source of a binding can now be a widget in the UI.

diff --git a/Engine/Plugins/Runtime/ModelViewViewModel/Source/ModelViewViewModelBlueprint/Private/Bindings/MVVMConversionFunctionHelper.cpp b/Engine/Plugins/Runtime/ModelViewViewModel/Source/ModelViewViewModelBlueprint/Private/Bindings/MVVMConversionFunctionHelper.cpp
--- a/Engine/Plugins/Runtime/ModelViewViewModel/Source/ModelViewViewModelBlueprint/Private/Bindings/MVVMConversionFunctionHelper.cpp
+++ b/Engine/Plugins/Runtime/ModelViewViewModel/Source/ModelViewViewModelBlueprint/Private/Bindings/MVVMConversionFunctionHelper.cpp
@@ -134,6 +134,10 @@ namespace Private
 			{
 				PathParts.Insert(Getter->VariableReference, 0);
 			}
+			else if (UK2Node_CallFunction* Function = Cast<UK2Node_CallFunction>(Node))
+			{
+				PathParts.Insert(Function->FunctionReference, 0);
+			}
 
 			for (UEdGraphPin* Pin : Node->Pins)
 			{
diff --git a/Engine/Plugins/Runtime/ModelViewViewModel/Source/ModelViewViewModelEditor/Private/Widgets/SMVVMFieldSelector.cpp b/Engine/Plugins/Runtime/ModelViewViewModel/Source/ModelViewViewModelEditor/Private/Widgets/SMVVMFieldSelector.cpp
--- a/Engine/Plugins/Runtime/ModelViewViewModel/Source/ModelViewViewModelEditor/Private/Widgets/SMVVMFieldSelector.cpp
+++ b/Engine/Plugins/Runtime/ModelViewViewModel/Source/ModelViewViewModelEditor/Private/Widgets/SMVVMFieldSelector.cpp
@@ -109,22 +109,6 @@ int32 SFieldSelector::GetCurrentDisplayIndex() const
 	return 2;
 }
 
-//void SFieldSelector::Refresh()
-//{
-//	CachedSelectedField = SelectedField.Get();
-//	CachedSelectedConversionFunction = SelectedConversionFunction.Get(nullptr);
-//
-//	//if (PropertyPathWidget.IsValid())
-//	//{
-//	//	PropertyPathWidget->SetPropertyPath(CachedSelectedField);
-//	//}
-//
-//	//if (SelectedSourceWidget.IsValid() && WidgetBlueprint.IsValid())
-//	//{
-//	//	SelectedSourceWidget->RefreshSource(Private::GetSourceFromPath(WidgetBlueprint.Get(), CachedSelectedField));
-//	//}
-//}
-
 TSharedRef<SWidget> SFieldSelector::HandleGetMenuContent()
 {
 	const UWidgetBlueprint* WidgetBlueprintPtr = WidgetBlueprint.Get();
diff --git a/Engine/Plugins/Runtime/ModelViewViewModel/Source/ModelViewViewModelEditor/Private/Widgets/SMVVMFieldSelectorMenu.cpp b/Engine/Plugins/Runtime/ModelViewViewModel/Source/ModelViewViewModelEditor/Private/Widgets/SMVVMFieldSelectorMenu.cpp
--- a/Engine/Plugins/Runtime/ModelViewViewModel/Source/ModelViewViewModelEditor/Private/Widgets/SMVVMFieldSelectorMenu.cpp
+++ b/Engine/Plugins/Runtime/ModelViewViewModel/Source/ModelViewViewModelEditor/Private/Widgets/SMVVMFieldSelectorMenu.cpp
@@ -385,6 +385,11 @@ void SFieldSelectorMenu::HandleViewModelSelected(FBindingSource Source, ESelectI
 		return;
 	}
 
+	if (WidgetList.IsValid())
+	{
+		WidgetList->ClearSelection();
+	}
+
 	if (ConversionFunctionCategoryTree.IsValid())
 	{
 		ConversionFunctionCategoryTree->ClearSelection();
@@ -414,6 +419,11 @@ void SFieldSelectorMenu::HandleWidgetSelected(FName WidgetName, ESelectInfo::Typ
 		return;
 	}
 
+	if (ViewModelList.IsValid())
+	{
+		ViewModelList->ClearSelection();
+	}
+
 	if (ConversionFunctionCategoryTree.IsValid())
 	{
 		ConversionFunctionCategoryTree->ClearSelection();
diff --git a/Engine/Plugins/Runtime/ModelViewViewModel/Source/ModelViewViewModelEditor/Private/Widgets/SMVVMViewBindingListView.cpp b/Engine/Plugins/Runtime/ModelViewViewModel/Source/ModelViewViewModelEditor/Private/Widgets/SMVVMViewBindingListView.cpp
--- a/Engine/Plugins/Runtime/ModelViewViewModel/Source/ModelViewViewModelEditor/Private/Widgets/SMVVMViewBindingListView.cpp
+++ b/Engine/Plugins/Runtime/ModelViewViewModel/Source/ModelViewViewModelEditor/Private/Widgets/SMVVMViewBindingListView.cpp
@@ -677,7 +677,7 @@ private:
 				Result.FixedBindingSource = FBindingSource::CreateForWidget(WidgetBlueprintPtr, ViewBinding->DestinationPath.GetWidgetName());
 			}
 
-			Result.bAllowWidgets = !bSource;
+			Result.bAllowWidgets = true;
 			Result.bAllowViewModels = bSource;
 			Result.bAllowConversionFunctions = false;
 			bool bIsReadingValue = (IsForwardBinding(ViewBinding->BindingType) && bSource)
