MVVM: The picker filter conversion function by display name and by keyword.

diff --git a/Engine/Plugins/Runtime/ModelViewViewModel/Source/ModelViewViewModelEditor/Private/Widgets/SMVVMFieldSelectorMenu.cpp b/Engine/Plugins/Runtime/ModelViewViewModel/Source/ModelViewViewModelEditor/Private/Widgets/SMVVMFieldSelectorMenu.cpp
--- a/Engine/Plugins/Runtime/ModelViewViewModel/Source/ModelViewViewModelEditor/Private/Widgets/SMVVMFieldSelectorMenu.cpp
+++ b/Engine/Plugins/Runtime/ModelViewViewModel/Source/ModelViewViewModelEditor/Private/Widgets/SMVVMFieldSelectorMenu.cpp
@@ -3,6 +3,7 @@
 #include "Widgets/SMVVMFieldSelectorMenu.h"
 
 #include "Bindings/MVVMBindingHelper.h"
+#include "EdGraphSchema_K2.h"
 #include "Editor.h"
 #include "Framework/Views/TableViewMetadata.h"
 #include "Hierarchy/SReadOnlyHierarchyView.h"
@@ -77,6 +78,11 @@ void SFieldSelectorMenu::Construct(const FArguments& InArgs, const UWidgetBluepr
 		AssignableToProperty = SelectionContext.AssignableTo;
 	}
 
+	if (SelectionContext.bAllowConversionFunctions)
+	{
+		GenerateConversionFunctionItems();
+	}
+
 	TSharedRef<SVerticalBox> VBox = SNew(SVerticalBox)
 	+ SVerticalBox::Slot()
 	.HAlign(HAlign_Fill)
@@ -322,14 +328,6 @@ void SFieldSelectorMenu::HandleSearchBoxTextChanged(const FText& NewText)
 		BindingList->SetRawFilterText(NewText);
 	}
 
-	// Do not search for widgets, only search for properties in the widgets
-	//if (WidgetList.IsValid())
-	//{
-	//	WidgetList->SetRawFilterText(NewText);
-	//}
-
-	//FilterViewModels(NewText);
-
 	if (ConversionFunctionCategoryTree.IsValid())
 	{
 		TArray<TSharedPtr<FConversionFunctionItem>> OldSelectedCategories;
@@ -477,23 +475,33 @@ int32 SFieldSelectorMenu::FilterConversionFunctionCategoryChildren(const TArray<
 
 	for (const TSharedPtr<FConversionFunctionItem>& SourceItem : SourceArray)
 	{
-		TArray<TSharedPtr<FConversionFunctionItem>> FilteredChildren;
-
 		// check if our name matches the filters
-		bool bMatchesFilters = true;
-
-		const FString ItemName = SourceItem->Function != nullptr ? SourceItem->Function->GetName() : FString();
-		
-		for (const FString& Filter : FilterStrings)
+		bool bMatchesFilters = false;
+		if (SourceItem->Function != nullptr)
 		{
-			if (!ItemName.Contains(Filter))
+			bMatchesFilters = true;
+			for (const FString& Filter : FilterStrings)
 			{
-				bMatchesFilters = false;
-				break;
+				bool bFoundMatches = false;
+				for (const FString& Keyword : SourceItem->SearchKeywords)
+				{
+					if (Keyword.Contains(Filter))
+					{
+						bFoundMatches = true;
+						break;
+					}
+				}
+
+				if (!bFoundMatches)
+				{
+					bMatchesFilters = false;
+					break;
+				}
 			}
 		}
 
 		int32 NumChildren = 0;
+		TArray<TSharedPtr<FConversionFunctionItem>> FilteredChildren;
 		if (bMatchesFilters)
 		{
 			ensureAlways(SourceItem->Function != nullptr);
@@ -608,6 +616,18 @@ void SFieldSelectorMenu::GenerateConversionFunctionItems()
 		Item->Function = Function;
 		Item->NumFunctions = 1;
 		Parent->NumFunctions += 1;
+
+		Item->SearchKeywords.Add(Function->GetName());
+		const FString& DisplayName = Function->GetMetaData(FBlueprintMetadata::MD_DisplayName);
+		if (DisplayName.Len() > 0)
+		{
+			Item->SearchKeywords.Add(DisplayName);
+		}
+		FString MetadataKeywords = Function->GetMetaDataText(FBlueprintMetadata::MD_FunctionKeywords, TEXT("UObjectKeywords"), Function->GetFullGroupName(false)).ToString();
+		if (MetadataKeywords.Len() > 0)
+		{
+			Item->SearchKeywords.Add(MoveTemp(MetadataKeywords));
+		}
 	};
 
 	TSharedPtr<FConversionFunctionItem> CurrentSelectedItem;
@@ -837,8 +857,6 @@ TSharedRef<SWidget> SFieldSelectorMenu::CreateBindingContextPanel(const FArgumen
 
 	if (SelectionContext.bAllowConversionFunctions)
 	{
-		GenerateConversionFunctionItems();
-
 		FilterConversionFunctionCategories();
 		
 		SAssignNew(ConversionFunctionCategoryTree, STreeView<TSharedPtr<FConversionFunctionItem>>)
@@ -985,25 +1003,25 @@ TSharedRef<ITableRow> SFieldSelectorMenu::HandleGenerateConversionFunctionCatego
 		[
 			SNew(SHorizontalBox)
 			+ SHorizontalBox::Slot()
-		.HAlign(HAlign_Left)
-		.VAlign(VAlign_Center)
-		.Padding(0, 2.0f, 4.0f, 2.0f)
-		.AutoWidth()
-		[
-			SNew(SImage)
-			.DesiredSizeOverride(FVector2D(16.0f, 16.0f))
-		.Image(FAppStyle::Get().GetBrush("GraphEditor.Function_16x"))
-		]
-	+ SHorizontalBox::Slot()
-		.HAlign(HAlign_Fill)
-		.VAlign(VAlign_Center)
-		[
-			SNew(STextBlock)
-			.Font(Item == FilteredConversionFunctionRoot[0] ? FAppStyle::Get().GetFontStyle("NormalText") : FAppStyle::Get().GetFontStyle("BoldFont"))
-		.Text(DisplayName)
-		.ToolTipText(FText::FromString(Item->GetCategoryName()))
-		.HighlightText_Lambda([this]() { return SearchBox.IsValid() ? SearchBox->GetText() : FText::GetEmpty(); })
-		]
+			.HAlign(HAlign_Left)
+			.VAlign(VAlign_Center)
+			.Padding(0, 2.0f, 4.0f, 2.0f)
+			.AutoWidth()
+			[
+				SNew(SImage)
+				.DesiredSizeOverride(FVector2D(16.0f, 16.0f))
+				.Image(FAppStyle::Get().GetBrush("GraphEditor.Function_16x"))
+			]
+			+ SHorizontalBox::Slot()
+			.HAlign(HAlign_Fill)
+			.VAlign(VAlign_Center)
+			[
+				SNew(STextBlock)
+				.Font(Item == FilteredConversionFunctionRoot[0] ? FAppStyle::Get().GetFontStyle("NormalText") : FAppStyle::Get().GetFontStyle("BoldFont"))
+				.Text(DisplayName)
+				.ToolTipText(FText::FromString(Item->GetCategoryName()))
+				.HighlightText_Lambda([this]() { return SearchBox.IsValid() ? SearchBox->GetText() : FText::GetEmpty(); })
+			]
 		];
 }
 
@@ -1056,7 +1074,7 @@ void SFieldSelectorMenu::AddConversionFunctionChildrenRecursive(const TSharedPtr
 			int32 Index = 0;
 			for (; Index < OutFunctions.Num(); ++Index)
 			{
-				if (OutFunctions[Index]->GetName() > Item->Function->GetName())
+				if (OutFunctions[Index]->GetFName().Compare(Item->Function->GetFName()) > 0)
 				{
 					break;
 				}
@@ -1085,14 +1103,17 @@ void SFieldSelectorMenu::FilterConversionFunctions()
 		return;
 	}
 
+	FilteredConversionFunctions.Reset();
 	for (const UFunction* Function : ConversionFunctions)
 	{
-		const FString FunctionName = Function->GetName();
+		FString FunctionName = Function->GetName();
+		const FString& DisplayName = Function->GetMetaData(FBlueprintMetadata::MD_DisplayName);
+		FString MetadataKeywords = Function->GetMetaDataText(FBlueprintMetadata::MD_FunctionKeywords, TEXT("UObjectKeywords"), Function->GetFullGroupName(false)).ToString();
 
 		bool bMatches = true;
 		for (const FString& Filter : FilterStrings)
 		{
-			if (!FunctionName.Contains(Filter))
+			if (!FunctionName.Contains(Filter) && !DisplayName.Contains(Filter) && !MetadataKeywords.Contains(Filter))
 			{
 				bMatches = false;
 				break;
diff --git a/Engine/Plugins/Runtime/ModelViewViewModel/Source/ModelViewViewModelEditor/Private/Widgets/SMVVMFieldSelectorMenu.h b/Engine/Plugins/Runtime/ModelViewViewModel/Source/ModelViewViewModelEditor/Private/Widgets/SMVVMFieldSelectorMenu.h
--- a/Engine/Plugins/Runtime/ModelViewViewModel/Source/ModelViewViewModelEditor/Private/Widgets/SMVVMFieldSelectorMenu.h
+++ b/Engine/Plugins/Runtime/ModelViewViewModel/Source/ModelViewViewModelEditor/Private/Widgets/SMVVMFieldSelectorMenu.h
@@ -63,6 +63,7 @@ private:
 		FString GetCategoryName() { return CategoryPath.Num() > 0 ? CategoryPath.Last() : FString(); }
 
 		TArray<FString> CategoryPath;
+		TArray<FString> SearchKeywords;
 		const UFunction* Function = nullptr;
 		TArray<TSharedPtr<FConversionFunctionItem>> Children;
 		int32 NumFunctions = 0;
