MVVM: Add test to test removing an element from the FieldNotificationDelegate in a broadcast of the same delegate.

diff --git a/Engine/Plugins/Runtime/ModelViewViewModel/Source/ModelViewViewModelEditor/Private/Tests/MVVMFieldValueChangedTest.cpp b/Engine/Plugins/Runtime/ModelViewViewModel/Source/ModelViewViewModelEditor/Private/Tests/MVVMFieldValueChangedTest.cpp
--- a/Engine/Plugins/Runtime/ModelViewViewModel/Source/ModelViewViewModelEditor/Private/Tests/MVVMFieldValueChangedTest.cpp
+++ b/Engine/Plugins/Runtime/ModelViewViewModel/Source/ModelViewViewModelEditor/Private/Tests/MVVMFieldValueChangedTest.cpp
@@ -330,6 +330,51 @@ bool FMVVMFieldValueChangedTest::RunTest(const FString& Parameters)
 	TestExpectedCountValues(1, 1, 1, 1, 1);
 	ResetAllValues();
 
+	// Remove element while iterating
+	struct FRemoveCallbackHandler
+	{
+		FDelegateHandle ToRemoveHandle;
+		UMVVMFieldValueChangedTest* SourceObj = nullptr;
+		FCallbackHandler* CallbackHandler = nullptr;
+		int32 NumberPropertyIntChanged = 0;
+		int32 NumberPropertyFloatChanged = 0;
+		void OnPropertyIntChanged(UObject* InObject, UE::FieldNotification::FFieldId InFieldId)
+		{
+			if (SourceObj && ToRemoveHandle.IsValid())
+			{
+				SourceObj->RemoveFieldValueChangedDelegate(UMVVMFieldValueChangedTest::FFieldNotificationClassDescriptor::PropertyInt, ToRemoveHandle);
+			}
+			++NumberPropertyIntChanged;
+		}
+		void OnPropertyFloatChanged(UObject* InObject, UE::FieldNotification::FFieldId InFieldId)
+		{
+			if (SourceObj && CallbackHandler)
+			{
+				SourceObj->RemoveAllFieldValueChangedDelegates(UMVVMFieldValueChangedTest::FFieldNotificationClassDescriptor::PropertyFloat, CallbackHandler);
+			}
+			++NumberPropertyFloatChanged;
+		}
+	};
+	FRemoveCallbackHandler RemoveCallbackInstance;
+
+	RemoveCallbackInstance.SourceObj = SourceObj;
+	RemoveCallbackInstance.ToRemoveHandle = HandleA_2;
+	RemoveCallbackInstance.CallbackHandler = &CallbackInstances[2];
+
+	SourceObj->AddFieldValueChangedDelegate(UMVVMFieldValueChangedTest::FFieldNotificationClassDescriptor::PropertyInt, UMVVMFieldValueChangedTest::FFieldValueChangedDelegate::CreateRaw(&RemoveCallbackInstance, &FRemoveCallbackHandler::OnPropertyIntChanged));
+	SourceObj->AddFieldValueChangedDelegate(UMVVMFieldValueChangedTest::FFieldNotificationClassDescriptor::PropertyFloat, UMVVMFieldValueChangedTest::FFieldValueChangedDelegate::CreateRaw(&RemoveCallbackInstance, &FRemoveCallbackHandler::OnPropertyFloatChanged));
+
+	TestBroadcastAll();
+	// Do not test the result here because the order is not deterministic.
+	ResetAllValues();
+
+	TestBroadcastAll();
+	TestAllExpectedCountValues(1, 3, 4, 4);
+	TestExpectedCountValues(0, 0, 2, 2, 2);
+	TestExpectedCountValues(1, 0, 1, 1, 1);
+	TestExpectedCountValues(2, 1, 0, 1, 1);
+	ResetAllValues();
+
 	return true;
 }
 
diff --git a/Engine/Source/Runtime/FieldNotification/Private/FieldNotificationDelegate.cpp b/Engine/Source/Runtime/FieldNotification/Private/FieldNotificationDelegate.cpp
--- a/Engine/Source/Runtime/FieldNotification/Private/FieldNotificationDelegate.cpp
+++ b/Engine/Source/Runtime/FieldNotification/Private/FieldNotificationDelegate.cpp
@@ -362,9 +362,9 @@ FFieldMulticastDelegate::FRemoveAllResult FFieldMulticastDelegate::RemoveAll(con
 	{
 		for (FInvocationElement& Element : Delegates)
 		{
-			if (Element.Key.Id == InFieldId && Element.Key.Object == InObject)
+			if (Element.Key.Object == InObject)
 			{
-				if (Element.Delegate.IsBoundToObject(InUserObject))
+				if (Element.Key.Id == InFieldId && Element.Delegate.IsBoundToObject(InUserObject))
 				{
 					Element.Delegate.Unbind();
 					++CompactionCount;
