MVVM: Do not add the view if there is no bindings.

diff --git a/Engine/Plugins/Runtime/ModelViewViewModel/Source/ModelViewViewModel/Private/View/MVVMViewClass.cpp b/Engine/Plugins/Runtime/ModelViewViewModel/Source/ModelViewViewModel/Private/View/MVVMViewClass.cpp
--- a/Engine/Plugins/Runtime/ModelViewViewModel/Source/ModelViewViewModel/Private/View/MVVMViewClass.cpp
+++ b/Engine/Plugins/Runtime/ModelViewViewModel/Source/ModelViewViewModel/Private/View/MVVMViewClass.cpp
@@ -492,6 +492,7 @@ void UMVVMViewClass::Initialize(UUserWidget* UserWidget)
 	UMVVMView* View = UserWidget->AddExtension<UMVVMView>();
 	if (ensure(View))
 	{
+		View->SetFlags(RF_Transient);
 		View->ConstructView(this);
 	}
 }
diff --git a/Engine/Plugins/Runtime/ModelViewViewModel/Source/ModelViewViewModelBlueprint/Private/MVVMWidgetBlueprintExtension_View.cpp b/Engine/Plugins/Runtime/ModelViewViewModel/Source/ModelViewViewModelBlueprint/Private/MVVMWidgetBlueprintExtension_View.cpp
--- a/Engine/Plugins/Runtime/ModelViewViewModel/Source/ModelViewViewModelBlueprint/Private/MVVMWidgetBlueprintExtension_View.cpp
+++ b/Engine/Plugins/Runtime/ModelViewViewModel/Source/ModelViewViewModelBlueprint/Private/MVVMWidgetBlueprintExtension_View.cpp
@@ -127,7 +127,11 @@ void UMVVMWidgetBlueprintExtension_View::HandleFinishCompilingClass(UWidgetBluep
 		if (bCompiled && UE::MVVM::Private::GAllowViewClass)
 		{
 			check(ViewExtension);
-			CurrentCompilerContext->AddExtension(Class, ViewExtension);
+			// Does it have any bindings or any sources
+			if (ViewExtension->GetViewModelCreators().Num() > 0 || const_cast<const UMVVMViewClass*>(ViewExtension)->GetCompiledBindings().Num() > 0)
+			{
+				CurrentCompilerContext->AddExtension(Class, ViewExtension);
+			}
 		}
 	}
 	else
