MVVM: Enable bindings from view to view or viewmodel to viewmodel for the compiler. The binding execution does not care. It was restricted by the UI.

diff --git a/Engine/Plugins/Runtime/ModelViewViewModel/Source/ModelViewViewModelBlueprint/Private/MVVMViewBlueprintCompiler.cpp b/Engine/Plugins/Runtime/ModelViewViewModel/Source/ModelViewViewModelBlueprint/Private/MVVMViewBlueprintCompiler.cpp
--- a/Engine/Plugins/Runtime/ModelViewViewModel/Source/ModelViewViewModelBlueprint/Private/MVVMViewBlueprintCompiler.cpp
+++ b/Engine/Plugins/Runtime/ModelViewViewModel/Source/ModelViewViewModelBlueprint/Private/MVVMViewBlueprintCompiler.cpp
@@ -441,21 +441,10 @@ void FMVVMViewBlueprintCompiler::CreateSourceLists(const FWidgetBlueprintCompile
 		}
 
 		FMVVMViewBlueprintCompiler* Self = this;
-		auto GenerateCompilerSourceContext = [Self, BlueprintView, DefaultWidgetCategory, Class = SkeletonClass, &Binding, &ViewModelGuids, &WidgetSources](const FMVVMBlueprintPropertyPath& PropertyPath, bool bViewModelPath, FName ArgumentName = FName()) -> bool
+		auto GenerateCompilerSourceContext = [Self, BlueprintView, DefaultWidgetCategory, Class = SkeletonClass, &Binding, &ViewModelGuids, &WidgetSources](const FMVVMBlueprintPropertyPath& PropertyPath,  FName ArgumentName = FName()) -> bool
 		{
 			if (PropertyPath.IsFromWidget())
 			{
-				if (bViewModelPath)
-				{
-					Self->AddMessageForBinding(Binding,
-							BlueprintView,
-							FText::Format(LOCTEXT("ExpectedViewModelPath", "Expected a viewmodel path, but received a path from widget: {0}"), 
-								FText::FromName(PropertyPath.GetWidgetName())
-							),
-							EBindingMessageType::Error
-						);
-				}
-
 				if (PropertyPath.GetWidgetName() == Class->ClassGeneratedBy->GetFName())
 				{
 					// it's the userwidget
@@ -506,18 +495,6 @@ void FMVVMViewBlueprintCompiler::CreateSourceLists(const FWidgetBlueprintCompile
 					return false;
 				}
 
-				if (!bViewModelPath)
-				{
-					Self->AddMessageForBinding(Binding,
-							BlueprintView,
-							FText::Format(LOCTEXT("ExpectedWidgetPath", "Expected a widget path, but received a path from viewmodel: {0}"), 
-								SourceViewModelContext->GetDisplayName()
-							),
-							EBindingMessageType::Error,
-							ArgumentName
-						);
-				}
-
 				if (!ViewModelGuids.Contains(SourceViewModelContext->GetViewModelId()))
 				{
 					Self->AddMessageForBinding(Binding,
@@ -534,13 +511,7 @@ void FMVVMViewBlueprintCompiler::CreateSourceLists(const FWidgetBlueprintCompile
 			}
 			else
 			{
-				Self->AddMessageForBinding(Binding,
-					BlueprintView,
-					bViewModelPath ? LOCTEXT("ViewModelPathNotSet", "A viewmodel path is required, but not set.") :
-							LOCTEXT("WidgetPathNotSet", "A widget path is required, but not set."),
-					EBindingMessageType::Error,
-					ArgumentName
-				);
+				Self->AddMessageForBinding(Binding, BlueprintView, LOCTEXT("SourcePathNotSet", "A source path is required, but not set."), EBindingMessageType::Error, ArgumentName);
 				return false;
 			}
 			return true;
@@ -566,29 +537,29 @@ void FMVVMViewBlueprintCompiler::CreateSourceLists(const FWidgetBlueprintCompile
 				{
 					if (bIsForwardBinding)
 					{
-						bAreSourceContextsValid &= GenerateCompilerSourceContext(Arg.Value, true, Arg.Key);
+						bAreSourceContextsValid &= GenerateCompilerSourceContext(Arg.Value, Arg.Key);
 					}
 					else
 					{
-						bAreSourceContextsValid &= GenerateCompilerSourceContext(Arg.Value, false, Arg.Key);
+						bAreSourceContextsValid &= GenerateCompilerSourceContext(Arg.Value, Arg.Key);
 					}
 				}
 
 				// generate destination source
 				if (bIsForwardBinding)
 				{
-					bAreSourceContextsValid &= GenerateCompilerSourceContext(Binding.DestinationPath, false);
+					bAreSourceContextsValid &= GenerateCompilerSourceContext(Binding.DestinationPath);
 				}
 				else
 				{
-					bAreSourceContextsValid &= GenerateCompilerSourceContext(Binding.SourcePath, true);
+					bAreSourceContextsValid &= GenerateCompilerSourceContext(Binding.SourcePath);
 				}
 			}
 			else
 			{
 				// if we aren't using a conversion function, just validate the widget and viewmodel paths
-				bAreSourceContextsValid &= GenerateCompilerSourceContext(Binding.DestinationPath, false);
-				bAreSourceContextsValid &= GenerateCompilerSourceContext(Binding.SourcePath, true);
+				bAreSourceContextsValid &= GenerateCompilerSourceContext(Binding.DestinationPath);
+				bAreSourceContextsValid &= GenerateCompilerSourceContext(Binding.SourcePath);
 			}
 		}
 	}
