MVVM: Before deleting a viewmodel, check if the viewmodel is in use. It will check in the bindings and in the conversion functions.

diff --git a/Engine/Plugins/Runtime/ModelViewViewModel/Source/ModelViewViewModelEditor/Private/MVVMEditorSubsystem.cpp b/Engine/Plugins/Runtime/ModelViewViewModel/Source/ModelViewViewModelEditor/Private/MVVMEditorSubsystem.cpp
--- a/Engine/Plugins/Runtime/ModelViewViewModel/Source/ModelViewViewModelEditor/Private/MVVMEditorSubsystem.cpp
+++ b/Engine/Plugins/Runtime/ModelViewViewModel/Source/ModelViewViewModelEditor/Private/MVVMEditorSubsystem.cpp
@@ -1136,5 +1136,66 @@ TArray<FMVVMAvailableBinding> UMVVMEditorSubsystem::GetChildViewModels(TSubclass
 	return ViewModelAvailableBindingsList;
 }
 
+FGuid UMVVMEditorSubsystem::GetFirstBindingThatUsesViewModel(const UWidgetBlueprint* WidgetBlueprint, FGuid ViewModelId) const
+{
+	if (UMVVMBlueprintView* View = GetView(WidgetBlueprint))
+	{
+		for (const FMVVMBlueprintViewBinding& Binding : View->GetBindings())
+		{
+			if (Binding.SourcePath.GetViewModelId() == ViewModelId)
+			{
+				return Binding.BindingId;
+			}
+			if (Binding.DestinationPath.GetViewModelId() == ViewModelId)
+			{
+				return Binding.BindingId;
+			}
+
+			auto TestConversionFunction = [&](bool bForward)
+			{
+				if (UE::MVVM::IsForwardBinding(Binding.BindingType))
+				{
+					const UFunction* Function = GetConversionFunction(WidgetBlueprint, Binding, bForward);
+					if (Function != nullptr)
+					{
+						TValueOrError<TArray<const FProperty*>, FText> ArgumentsResult = UE::MVVM::BindingHelper::TryGetArgumentsForConversionFunction(Function);
+						if (ArgumentsResult.HasValue())
+						{
+							for (const FProperty* Property : ArgumentsResult.GetValue())
+							{
+								FMVVMBlueprintPropertyPath Path = GetPathForConversionFunctionArgument(WidgetBlueprint, Binding, Property->GetFName(), bForward);
+								if (Path.GetViewModelId() == ViewModelId)
+								{
+									return Binding.BindingId;
+								}
+							}
+						}
+					}
+				}
+				return FGuid();
+			};
+
+			if (UE::MVVM::IsForwardBinding(Binding.BindingType))
+			{
+				FGuid Forward = TestConversionFunction(true);
+				if (Forward.IsValid())
+				{
+					return Forward;
+				}
+			}
+
+			if (UE::MVVM::IsBackwardBinding(Binding.BindingType))
+			{
+				FGuid Backward = TestConversionFunction(false);
+				if (Backward.IsValid())
+				{
+					return Backward;
+				}
+			}
+		}
+	}
+	return FGuid();
+}
+
 #undef LOCTEXT_NAMESPACE
 
diff --git a/Engine/Plugins/Runtime/ModelViewViewModel/Source/ModelViewViewModelEditor/Private/Widgets/SMVVMViewModelPanel.cpp b/Engine/Plugins/Runtime/ModelViewViewModel/Source/ModelViewViewModelEditor/Private/Widgets/SMVVMViewModelPanel.cpp
--- a/Engine/Plugins/Runtime/ModelViewViewModel/Source/ModelViewViewModelEditor/Private/Widgets/SMVVMViewModelPanel.cpp
+++ b/Engine/Plugins/Runtime/ModelViewViewModel/Source/ModelViewViewModelEditor/Private/Widgets/SMVVMViewModelPanel.cpp
@@ -3,6 +3,7 @@
 #include "Widgets/SMVVMViewModelPanel.h"
 
 #include "DetailsViewArgs.h"
+#include "Dialogs/Dialogs.h"
 #include "IStructureDetailsView.h"
 #include "MVVMBlueprintView.h"
 #include "MVVMEditorSubsystem.h"
@@ -14,7 +15,6 @@
 #include "Framework/Commands/GenericCommands.h"
 #include "Framework/MultiBox/MultiBoxBuilder.h"
 
-
 #include "Widgets/Images/SImage.h"
 #include "Widgets/Layout/SBorder.h"
 #include "Widgets/Layout/SSplitter.h"
@@ -378,6 +378,26 @@ void SMVVMViewModelPanel::CreateCommandList()
 }
 
 
+namespace Private
+{
+static bool DisplayInUseWarningAndEarlyExit(const UWidgetBlueprint* WidgetBlueprint, const FMVVMBlueprintViewModelContext* ViewModelContext, const FMVVMBlueprintViewBinding* Binding)
+{
+	const FText DeleteConfirmationPrompt = FText::Format(LOCTEXT("DeleteConfirmationPrompt", "The viewmodel {0} is in use by binding {1}! Do you really want to delete it?")
+		, ViewModelContext->GetDisplayName()
+		, Binding ? FText::FromString(Binding->GetDisplayNameString(WidgetBlueprint)) : LOCTEXT("DeleteConfirmation_Unknowed", "Unknowed"));
+	const FText DeleteConfirmationTitle = LOCTEXT("DeleteConfirmationTitle", "Delete viewmodel");
+
+	// Warn the user that this may result in data loss
+	FSuppressableWarningDialog::FSetupInfo Info(DeleteConfirmationPrompt, DeleteConfirmationTitle, TEXT("Viewmodel_Warning"));
+	Info.ConfirmText = LOCTEXT("DeleteConfirmation_Yes", "Yes");
+	Info.CancelText = LOCTEXT("DeleteConfirmation_No", "No");
+
+	FSuppressableWarningDialog DeleteFunctionInUse(Info);
+	return DeleteFunctionInUse.ShowModal() == FSuppressableWarningDialog::Cancel;
+}
+}
+
+
 void SMVVMViewModelPanel::HandleDeleteViewModel()
 {
 	TSharedPtr<FWidgetBlueprintEditor> WidgetBlueprintEditor = WeakBlueprintEditor.Pin();
@@ -407,8 +427,17 @@ void SMVVMViewModelPanel::HandleDeleteViewModel()
 			{
 				if (const FMVVMBlueprintViewModelContext* ViewModelContext = BlueprintView->FindViewModel(*VMGuidPtr))
 				{
-					GEditor->GetEditorSubsystem<UMVVMEditorSubsystem>()->RemoveViewModel(WidgetBP, ViewModelContext->GetViewModelName());
+					FGuid BindingId = GEditor->GetEditorSubsystem<UMVVMEditorSubsystem>()->GetFirstBindingThatUsesViewModel(WidgetBP, *VMGuidPtr);
+					if (BindingId.IsValid())
+					{
+						if (UE::MVVM::Private::DisplayInUseWarningAndEarlyExit(WidgetBP, ViewModelContext, BlueprintView->GetBinding(BindingId)))
+						{
+							return;
+						}
+					}
 				}
+
+				BlueprintView->RemoveViewModel(*VMGuidPtr);
 			}
 		}
 	}
diff --git a/Engine/Plugins/Runtime/ModelViewViewModel/Source/ModelViewViewModelEditor/Public/MVVMEditorSubsystem.h b/Engine/Plugins/Runtime/ModelViewViewModel/Source/ModelViewViewModelEditor/Public/MVVMEditorSubsystem.h
--- a/Engine/Plugins/Runtime/ModelViewViewModel/Source/ModelViewViewModelEditor/Public/MVVMEditorSubsystem.h
+++ b/Engine/Plugins/Runtime/ModelViewViewModel/Source/ModelViewViewModelEditor/Public/MVVMEditorSubsystem.h
@@ -92,6 +92,8 @@ public:
 	
 	TArray<UE::MVVM::FBindingSource> GetAllViewModels(const UWidgetBlueprint* WidgetBlueprint) const;
 
+	FGuid GetFirstBindingThatUsesViewModel(const UWidgetBlueprint* WidgetBlueprint, FGuid ViewModelId) const;
+
 private:
 	UEdGraph* CreateConversionFunctionWrapperGraph(UWidgetBlueprint* WidgetBlueprint, const FMVVMBlueprintViewBinding& Binding, const UFunction* ConversionFunction, bool bSourceToDestination);
 };
