MVVM: Fix static analysis warning
#rnx

diff --git a/Engine/Plugins/Runtime/ModelViewViewModel/Source/ModelViewViewModelBlueprint/Private/MVVMViewBlueprintCompiler.cpp b/Engine/Plugins/Runtime/ModelViewViewModel/Source/ModelViewViewModelBlueprint/Private/MVVMViewBlueprintCompiler.cpp
--- a/Engine/Plugins/Runtime/ModelViewViewModel/Source/ModelViewViewModelBlueprint/Private/MVVMViewBlueprintCompiler.cpp
+++ b/Engine/Plugins/Runtime/ModelViewViewModel/Source/ModelViewViewModelBlueprint/Private/MVVMViewBlueprintCompiler.cpp
@@ -535,14 +535,7 @@ void FMVVMViewBlueprintCompiler::CreateSourceLists(const FWidgetBlueprintCompile
 				// generate sources for conversion function arguments
 				for (const TPair<FName, FMVVMBlueprintPropertyPath>& Arg : ConversionFunctionArguments)
 				{
-					if (bIsForwardBinding)
-					{
-						bAreSourceContextsValid &= GenerateCompilerSourceContext(Arg.Value, Arg.Key);
-					}
-					else
-					{
-						bAreSourceContextsValid &= GenerateCompilerSourceContext(Arg.Value, Arg.Key);
-					}
+					bAreSourceContextsValid &= GenerateCompilerSourceContext(Arg.Value, Arg.Key);
 				}
 
 				// generate destination source
