MVVM: Add a message at compile time when 2 views will be use at runtime.

diff --git a/Engine/Plugins/Runtime/ModelViewViewModel/Source/ModelViewViewModelBlueprint/Private/MVVMWidgetBlueprintExtension_View.cpp b/Engine/Plugins/Runtime/ModelViewViewModel/Source/ModelViewViewModelBlueprint/Private/MVVMWidgetBlueprintExtension_View.cpp
--- a/Engine/Plugins/Runtime/ModelViewViewModel/Source/ModelViewViewModelBlueprint/Private/MVVMWidgetBlueprintExtension_View.cpp
+++ b/Engine/Plugins/Runtime/ModelViewViewModel/Source/ModelViewViewModelBlueprint/Private/MVVMWidgetBlueprintExtension_View.cpp
@@ -130,6 +130,12 @@ void UMVVMWidgetBlueprintExtension_View::HandleFinishCompilingClass(UWidgetBluep
 			// Does it have any bindings or any sources
 			if (ViewExtension->GetViewModelCreators().Num() > 0 || const_cast<const UMVVMViewClass*>(ViewExtension)->GetCompiledBindings().Num() > 0)
 			{
+				// Test if parent also has a view
+				if (Class->GetExtension<UMVVMViewClass>(true))
+				{
+					CurrentCompilerContext->GetCompilerContext().MessageLog.Warning(*LOCTEXT("MoreThanOneViewWarning", "There is more than one view.").ToString());
+				}
+
 				CurrentCompilerContext->AddExtension(Class, ViewExtension);
 			}
 		}
