MVVM: Update compiler error message when the field is not found. It's probably because it's missing the FieldNotify

diff --git a/Engine/Plugins/Runtime/ModelViewViewModel/Source/ModelViewViewModelBlueprint/Private/Bindings/MVVMCompiledBindingLibraryCompiler.cpp b/Engine/Plugins/Runtime/ModelViewViewModel/Source/ModelViewViewModelBlueprint/Private/Bindings/MVVMCompiledBindingLibraryCompiler.cpp
--- a/Engine/Plugins/Runtime/ModelViewViewModel/Source/ModelViewViewModelBlueprint/Private/Bindings/MVVMCompiledBindingLibraryCompiler.cpp
+++ b/Engine/Plugins/Runtime/ModelViewViewModel/Source/ModelViewViewModelBlueprint/Private/Bindings/MVVMCompiledBindingLibraryCompiler.cpp
@@ -144,7 +144,7 @@ TValueOrError<FCompiledBindingLibraryCompiler::FFieldIdHandle, FText> FCompiledB
 
 	if (FieldId.IsNone())
 	{
-		return MakeError(LOCTEXT("FieldNotDefined", "The Field/Property is not defined."));
+		return MakeError(LOCTEXT("FieldNotDefined", "The Field is not defined. Is it missing the FieldNotify flag?."));
 	}
 
 	if (!SourceClass->ImplementsInterface(UNotifyFieldValueChanged::StaticClass()))
