MVVM: Add the Auto ExecutionMode. It will decide at compile time if it should run the binding should be run immediately or delayed. A binding will be run in delay mode if the binding can be triggered from more than one field.

diff --git a/Engine/Plugins/Runtime/ModelViewViewModel/Source/ModelViewViewModel/Private/ModelViewViewModelModule.cpp b/Engine/Plugins/Runtime/ModelViewViewModel/Source/ModelViewViewModel/Private/ModelViewViewModelModule.cpp
--- a/Engine/Plugins/Runtime/ModelViewViewModel/Source/ModelViewViewModel/Private/ModelViewViewModelModule.cpp
+++ b/Engine/Plugins/Runtime/ModelViewViewModel/Source/ModelViewViewModel/Private/ModelViewViewModelModule.cpp
@@ -29,7 +29,10 @@ void FModelViewViewModelModule::ShutdownModule()
 void FModelViewViewModelModule::HandleDefaultExecutionModeChanged(IConsoleVariable* Variable)
 {
 	EMVVMExecutionMode Value = (EMVVMExecutionMode)Variable->GetInt();
-	if (Value != EMVVMExecutionMode::Delayed && Value != EMVVMExecutionMode::Immediate && Value != EMVVMExecutionMode::Tick)
+	if (Value != EMVVMExecutionMode::Delayed
+		&& Value != EMVVMExecutionMode::Immediate
+		&& Value != EMVVMExecutionMode::Tick
+		&& Value != EMVVMExecutionMode::DelayedWhenSharedElseImmediate)
 	{
 		Variable->Set((int32)EMVVMExecutionMode::Immediate, (EConsoleVariableFlags)(Variable->GetFlags() & ECVF_SetByMask));
 	}
diff --git a/Engine/Plugins/Runtime/ModelViewViewModel/Source/ModelViewViewModel/Private/View/MVVMViewClass.cpp b/Engine/Plugins/Runtime/ModelViewViewModel/Source/ModelViewViewModel/Private/View/MVVMViewClass.cpp
--- a/Engine/Plugins/Runtime/ModelViewViewModel/Source/ModelViewViewModel/Private/View/MVVMViewClass.cpp
+++ b/Engine/Plugins/Runtime/ModelViewViewModel/Source/ModelViewViewModel/Private/View/MVVMViewClass.cpp
@@ -273,7 +273,8 @@ namespace UE::MVVM::Private
 EMVVMExecutionMode FMVVMViewClass_CompiledBinding::GetExecuteMode() const
 {
 	EMVVMExecutionMode DefaultMode = (EMVVMExecutionMode)UE::MVVM::Private::GDefaultEvaluationMode;
-	return (Flags & EBindingFlags::OverrideExecuteMode) == 0 ? DefaultMode : ExecutionMode;
+	EMVVMExecutionMode Result = (Flags & EBindingFlags::OverrideExecuteMode) == 0 ? DefaultMode : ExecutionMode;
+	return Result == EMVVMExecutionMode::DelayedWhenSharedElseImmediate ? (Binding.IsShared() ? EMVVMExecutionMode::Delayed : EMVVMExecutionMode::Immediate) : Result;
 }
 
 #if UE_WITH_MVVM_DEBUGGING
diff --git a/Engine/Plugins/Runtime/ModelViewViewModel/Source/ModelViewViewModel/Public/Bindings/MVVMCompiledBindingLibrary.h b/Engine/Plugins/Runtime/ModelViewViewModel/Source/ModelViewViewModel/Public/Bindings/MVVMCompiledBindingLibrary.h
--- a/Engine/Plugins/Runtime/ModelViewViewModel/Source/ModelViewViewModel/Public/Bindings/MVVMCompiledBindingLibrary.h
+++ b/Engine/Plugins/Runtime/ModelViewViewModel/Source/ModelViewViewModel/Public/Bindings/MVVMCompiledBindingLibrary.h
@@ -229,6 +229,12 @@ public:
 		return (Flags & AllFunctionFlags) == AllFunctionFlags;
 	}
 
+	/** @return true if multiple field can trigger this binding. */
+	bool IsShared() const
+	{
+		return (Flags & (uint8)EFlags::IsShared) != 0;
+	}
+
 private:
 	using IndexType = int16;
 
@@ -246,6 +252,7 @@ private:
 		None = 0,
 		HasConversionFunction = 1 << 0,
 		IsConversionFunctionComplex = 1 << 1,
+		IsShared = 1 << 2,
 	};
 
 	UPROPERTY()
diff --git a/Engine/Plugins/Runtime/ModelViewViewModel/Source/ModelViewViewModel/Public/Types/MVVMExecutionMode.h b/Engine/Plugins/Runtime/ModelViewViewModel/Source/ModelViewViewModel/Public/Types/MVVMExecutionMode.h
--- a/Engine/Plugins/Runtime/ModelViewViewModel/Source/ModelViewViewModel/Public/Types/MVVMExecutionMode.h
+++ b/Engine/Plugins/Runtime/ModelViewViewModel/Source/ModelViewViewModel/Public/Types/MVVMExecutionMode.h
@@ -10,11 +10,13 @@ UENUM()
 enum class EMVVMExecutionMode : uint8
 {
 	/** Execute the binding as soon as the source value changes. */
-	Immediate,
+	Immediate = 0,
 	/** Execute the binding at the end of the frame before drawing when the source value changes. */
-	Delayed,
+	Delayed = 1,
 	/** Always execute the binding at the end of the frame. */
-	Tick,
+	Tick = 2,
+	/** When the binding can be triggered from multiple fields, use Delayed. Else, uses Immediate. */
+	DelayedWhenSharedElseImmediate = 3 UMETA(DisplayName="Auto"),
 };
 
 #if UE_ENABLE_INCLUDE_ORDER_DEPRECATED_IN_5_2
diff --git a/Engine/Plugins/Runtime/ModelViewViewModel/Source/ModelViewViewModelBlueprint/Private/Bindings/MVVMCompiledBindingLibraryCompiler.cpp b/Engine/Plugins/Runtime/ModelViewViewModel/Source/ModelViewViewModelBlueprint/Private/Bindings/MVVMCompiledBindingLibraryCompiler.cpp
--- a/Engine/Plugins/Runtime/ModelViewViewModel/Source/ModelViewViewModelBlueprint/Private/Bindings/MVVMCompiledBindingLibraryCompiler.cpp
+++ b/Engine/Plugins/Runtime/ModelViewViewModel/Source/ModelViewViewModelBlueprint/Private/Bindings/MVVMCompiledBindingLibraryCompiler.cpp
@@ -72,6 +72,7 @@ struct FRawBinding
 
 	FCompiledBindingLibraryCompiler::FBindingHandle BindingHandle;
 	FMVVMVCompiledBinding CompiledBinding;
+	int32 BindingCount = 1;
 	bool bIsConversionFunctionComplex = false;
 
 public:
@@ -550,6 +551,7 @@ TValueOrError<FCompiledBindingLibraryCompiler::FBindingHandle, FText> FCompiledB
 		}
 		else
 		{
+			++Impl->Bindings[FoundSameBindingIndex].BindingCount;
 			return MakeValue(Impl->Bindings[FoundSameBindingIndex].BindingHandle);
 		}
 	}
@@ -794,6 +796,7 @@ TValueOrError<FCompiledBindingLibraryCompiler::FCompileResult, FText> FCompiledB
 		Binding.CompiledBinding.Flags = 0;
 		Binding.CompiledBinding.Flags |= (Binding.ConversionFunctionPathHandle.IsValid()) ? (uint8)FMVVMVCompiledBinding::EFlags::HasConversionFunction : 0;
 		Binding.CompiledBinding.Flags |= (Binding.bIsConversionFunctionComplex) ? (uint8)FMVVMVCompiledBinding::EFlags::IsConversionFunctionComplex : 0;
+		Binding.CompiledBinding.Flags |= Binding.BindingCount > 1 ? (uint8)FMVVMVCompiledBinding::EFlags::IsShared : 0;
 
 		Result.Bindings.Add(Binding.BindingHandle, Binding.CompiledBinding);
 	}
diff --git a/Engine/Plugins/Runtime/ModelViewViewModel/Source/ModelViewViewModelBlueprint/Private/MVVMDeveloperProjectSettings.cpp b/Engine/Plugins/Runtime/ModelViewViewModel/Source/ModelViewViewModelBlueprint/Private/MVVMDeveloperProjectSettings.cpp
--- a/Engine/Plugins/Runtime/ModelViewViewModel/Source/ModelViewViewModelBlueprint/Private/MVVMDeveloperProjectSettings.cpp
+++ b/Engine/Plugins/Runtime/ModelViewViewModel/Source/ModelViewViewModelBlueprint/Private/MVVMDeveloperProjectSettings.cpp
@@ -15,6 +15,7 @@ UMVVMDeveloperProjectSettings::UMVVMDeveloperProjectSettings()
 	AllowedExecutionMode.Add(EMVVMExecutionMode::Immediate);
 	AllowedExecutionMode.Add(EMVVMExecutionMode::Delayed);
 	AllowedExecutionMode.Add(EMVVMExecutionMode::Tick);
+	AllowedExecutionMode.Add(EMVVMExecutionMode::DelayedWhenSharedElseImmediate);
 	
 	AllowedContextCreationType.Add(EMVVMBlueprintViewModelContextCreationType::Manual);
 	AllowedContextCreationType.Add(EMVVMBlueprintViewModelContextCreationType::CreateInstance);
