MVVM: Add the viewmodel tag data to the userwidget generated class. It was present on the blueprint asset but that data is not present in coocked build.

diff --git a/Engine/Plugins/Runtime/ModelViewViewModel/Source/ModelViewViewModelBlueprint/Private/MVVMBlueprintView.cpp b/Engine/Plugins/Runtime/ModelViewViewModel/Source/ModelViewViewModelBlueprint/Private/MVVMBlueprintView.cpp
--- a/Engine/Plugins/Runtime/ModelViewViewModel/Source/ModelViewViewModelBlueprint/Private/MVVMBlueprintView.cpp
+++ b/Engine/Plugins/Runtime/ModelViewViewModel/Source/ModelViewViewModelBlueprint/Private/MVVMBlueprintView.cpp
@@ -270,7 +270,7 @@ void UMVVMBlueprintView::PostEditChangeChainProperty(FPropertyChangedChainEvent&
 	}
 }
 
-void UMVVMBlueprintView::AddAssetTags(TArray<FAssetRegistryTag>& OutTags)
+void UMVVMBlueprintView::AddAssetTags(TArray<FAssetRegistryTag>& OutTags) const
 {
 	if (AvailableViewModels.Num() > 0)
 	{
@@ -289,7 +289,7 @@ void UMVVMBlueprintView::AddAssetTags(TArray<FAssetRegistryTag>& OutTags)
 
 		if (Builder.Len() > 0)
 		{
-			OutTags.Emplace(TEXT("Viewmodels"), Builder.ToString(), FAssetRegistryTag::TT_Hidden);
+			OutTags.Emplace(FName("Viewmodels"), Builder.ToString(), FAssetRegistryTag::TT_Hidden);
 		}
 	}
 }
diff --git a/Engine/Plugins/Runtime/ModelViewViewModel/Source/ModelViewViewModelBlueprint/Public/MVVMBlueprintView.h b/Engine/Plugins/Runtime/ModelViewViewModel/Source/ModelViewViewModelBlueprint/Public/MVVMBlueprintView.h
--- a/Engine/Plugins/Runtime/ModelViewViewModel/Source/ModelViewViewModelBlueprint/Public/MVVMBlueprintView.h
+++ b/Engine/Plugins/Runtime/ModelViewViewModel/Source/ModelViewViewModelBlueprint/Public/MVVMBlueprintView.h
@@ -91,7 +91,7 @@ public:
 	virtual void PostEditChangeProperty(FPropertyChangedEvent& PropertyChangedEvent) override;
 	virtual void PostEditChangeChainProperty(FPropertyChangedChainEvent& PropertyChainEvent) override;
 
-	void AddAssetTags(TArray<FAssetRegistryTag>& OutTags);
+	void AddAssetTags(TArray<FAssetRegistryTag>& OutTags) const;
 	void WidgetRenamed(FName OldObjectName, FName NewObjectName);
 #endif
 
diff --git a/Engine/Plugins/Runtime/ModelViewViewModel/Source/ModelViewViewModelEditor/Private/ModelViewViewModelEditorModule.cpp b/Engine/Plugins/Runtime/ModelViewViewModel/Source/ModelViewViewModelEditor/Private/ModelViewViewModelEditorModule.cpp
--- a/Engine/Plugins/Runtime/ModelViewViewModel/Source/ModelViewViewModelEditor/Private/ModelViewViewModelEditorModule.cpp
+++ b/Engine/Plugins/Runtime/ModelViewViewModel/Source/ModelViewViewModelEditor/Private/ModelViewViewModelEditorModule.cpp
@@ -2,6 +2,7 @@
 
 #include "ModelViewViewModelEditorModule.h"
 
+#include "Blueprint/WidgetBlueprintGeneratedClass.h"
 #include "BlueprintModes/WidgetBlueprintApplicationMode.h"
 #include "BlueprintModes/WidgetBlueprintApplicationModes.h"
 #include "Customizations/MVVMBlueprintViewModelContextCustomization.h"
@@ -55,12 +56,14 @@ void FModelViewViewModelEditorModule::StartupModule()
 	}
 
 	FMVVMEditorCommands::Register();
-	FWidgetBlueprintDelegates::GetAssetTags.AddRaw(this, &FModelViewViewModelEditorModule::HandleAssetTags);
+	FWidgetBlueprintDelegates::GetAssetTags.AddRaw(this, &FModelViewViewModelEditorModule::HandleWidgetBlueprintAssetTags);
+	FWidgetBlueprintGeneratedClassDelegates::GetAssetTags.AddRaw(this, &FModelViewViewModelEditorModule::HandleClassBlueprintAssetTags);
 }
 
 
 void FModelViewViewModelEditorModule::ShutdownModule()
 {
+	FWidgetBlueprintGeneratedClassDelegates::GetAssetTags.RemoveAll(this);
 	FWidgetBlueprintDelegates::GetAssetTags.RemoveAll(this);
 	if (FMessageLogModule* MessageLogModule = FModuleManager::GetModulePtr<FMessageLogModule>("MessageLog"))
 	{
@@ -192,7 +195,7 @@ void FModelViewViewModelEditorModule::HandleActivateMode(FWidgetBlueprintApplica
 	}
 }
 
-void FModelViewViewModelEditorModule::HandleAssetTags(const UWidgetBlueprint* WidgetBlueprint, TArray<UObject::FAssetRegistryTag>& OutTags)
+void FModelViewViewModelEditorModule::HandleWidgetBlueprintAssetTags(const UWidgetBlueprint* WidgetBlueprint, TArray<UObject::FAssetRegistryTag>& OutTags)
 {
 	if (WidgetBlueprint && GEditor)
 	{
@@ -206,6 +209,17 @@ void FModelViewViewModelEditorModule::HandleAssetTags(const UWidgetBlueprint* Wi
 	}
 }
 
+void FModelViewViewModelEditorModule::HandleClassBlueprintAssetTags(const UWidgetBlueprintGeneratedClass* GeneratedClass, TArray<UObject::FAssetRegistryTag>& OutTags)
+{
+	if (GeneratedClass && GEditor && GeneratedClass->ClassGeneratedBy)
+	{
+		if (UWidgetBlueprint* WidgetBlueprint = Cast<UWidgetBlueprint>(GeneratedClass->ClassGeneratedBy))
+		{
+			HandleWidgetBlueprintAssetTags(WidgetBlueprint, OutTags);
+		}
+	}
+}
+
 IMPLEMENT_MODULE(FModelViewViewModelEditorModule, ModelViewViewModelEditor);
 
 #undef LOCTEXT_NAMESPACE
diff --git a/Engine/Plugins/Runtime/ModelViewViewModel/Source/ModelViewViewModelEditor/Private/ModelViewViewModelEditorModule.h b/Engine/Plugins/Runtime/ModelViewViewModel/Source/ModelViewViewModelEditor/Private/ModelViewViewModelEditorModule.h
--- a/Engine/Plugins/Runtime/ModelViewViewModel/Source/ModelViewViewModelEditor/Private/ModelViewViewModelEditorModule.h
+++ b/Engine/Plugins/Runtime/ModelViewViewModel/Source/ModelViewViewModelEditor/Private/ModelViewViewModelEditorModule.h
@@ -16,6 +16,7 @@ class FWidgetBlueprintApplicationMode;
 class FWorkflowAllowedTabSet;
 class UBlueprint;
 class UWidgetBlueprint;
+class UWidgetBlueprintGeneratedClass;
 
 /**
  *
@@ -35,7 +36,8 @@ private:
 	void HandleRenameVariableReferences(UBlueprint* Blueprint, UClass* VariableClass, const FName& OldVarName, const FName& NewVarName);
 	void HandleDeactiveMode(FWidgetBlueprintApplicationMode& InDesignerMode);
 	void HandleActivateMode(FWidgetBlueprintApplicationMode& InDesignerMode);
-	void HandleAssetTags(const UWidgetBlueprint* Widget, TArray<UObject::FAssetRegistryTag>& OutTags);
+	void HandleWidgetBlueprintAssetTags(const UWidgetBlueprint* Widget, TArray<UObject::FAssetRegistryTag>& OutTags);
+	void HandleClassBlueprintAssetTags(const UWidgetBlueprintGeneratedClass* GeneratedClass, TArray<UObject::FAssetRegistryTag>& OutTags);
 
 private:
 	TSharedPtr<FMVVMPropertyBindingExtension> PropertyBindingExtension;
