MVVM: Save the function name and object name, if available, to the DebugSnapshot. Add a detait view for the DebugSnapshot. Change the main tab default layout.

diff --git a/Engine/Plugins/Runtime/ModelViewViewModel/Source/ModelViewViewModelDebugger/Private/MVVMDebugSnapshot.cpp b/Engine/Plugins/Runtime/ModelViewViewModel/Source/ModelViewViewModelDebugger/Private/MVVMDebugSnapshot.cpp
--- a/Engine/Plugins/Runtime/ModelViewViewModel/Source/ModelViewViewModelDebugger/Private/MVVMDebugSnapshot.cpp
+++ b/Engine/Plugins/Runtime/ModelViewViewModel/Source/ModelViewViewModelDebugger/Private/MVVMDebugSnapshot.cpp
@@ -73,10 +73,12 @@ FMVVMViewSourceDebugEntry CreateViewModelEntry(const FMVVMViewSource& ViewSource
 FMVVMViewModelFieldBoundDebugEntry CreateViewModelFieldBoundDebugEntry(const UE::FieldNotification::FFieldMulticastDelegate::FDelegateView& DelegateView)
 {
 	FMVVMViewModelFieldBoundDebugEntry Result;
-	Result.ObjectName = DelegateView.Object ? DelegateView.Object->GetPathName() : FString();
-	Result.FieldId = FFieldNotificationId(DelegateView.Field.GetName());
-	Result.FunctionName = DelegateView.FunctionName;
-	Result.LiveInstanceObject = DelegateView.Object;
+	Result.KeyObjectName = DelegateView.KeyObject ? DelegateView.KeyObject->GetFName() : FName();
+	Result.KeyFieldId = FFieldNotificationId(DelegateView.KeyField.GetName());
+	Result.BindingFunctionName = DelegateView.BindingFunctionName;
+	Result.BindingObjectPathName = DelegateView.BindingObject ? DelegateView.BindingObject->GetPathName() : FString();
+	Result.LiveInstanceKeyObject = DelegateView.KeyObject;
+	Result.LiveInstanceBindingObject = DelegateView.BindingObject;
 	return Result;
 }
 }//namespace private
@@ -142,7 +144,7 @@ TSharedPtr<FDebugSnapshot> FDebugSnapshot::CreateSnapshot()
 
 		TSharedPtr<FMVVMViewModelDebugEntry> DebugEntry = MakeShared<FMVVMViewModelDebugEntry>();
 		DebugEntry->Name = ViewModel->GetFName();
-		DebugEntry->FullName = ViewModel->GetPathName();
+		DebugEntry->PathName = ViewModel->GetPathName();
 		DebugEntry->ViewModelAsset = FAssetData(ViewModel->GetClass());
 		for (const UE::FieldNotification::FFieldMulticastDelegate::FDelegateView& DelegateView : ViewModel->GetNotificationDelegateView())
 		{
diff --git a/Engine/Plugins/Runtime/ModelViewViewModel/Source/ModelViewViewModelDebugger/Public/MVVMDebugViewModel.h b/Engine/Plugins/Runtime/ModelViewViewModel/Source/ModelViewViewModelDebugger/Public/MVVMDebugViewModel.h
--- a/Engine/Plugins/Runtime/ModelViewViewModel/Source/ModelViewViewModelDebugger/Public/MVVMDebugViewModel.h
+++ b/Engine/Plugins/Runtime/ModelViewViewModel/Source/ModelViewViewModelDebugger/Public/MVVMDebugViewModel.h
@@ -17,15 +17,19 @@ struct FMVVMViewModelFieldBoundDebugEntry
 	GENERATED_BODY()
 
 	UPROPERTY(VisibleAnywhere, Category = "Viewmodel")
-	FString ObjectName;
+	FName KeyObjectName;
 
 	UPROPERTY(VisibleAnywhere, Category = "Viewmodel")
-	FFieldNotificationId FieldId;
+	FFieldNotificationId KeyFieldId;
 
 	UPROPERTY(VisibleAnywhere, Category = "Viewmodel")
-	FName FunctionName;
+	FString BindingObjectPathName;
 
-	TWeakObjectPtr<const UObject> LiveInstanceObject;
+	UPROPERTY(VisibleAnywhere, Category = "Viewmodel")
+	FName BindingFunctionName;
+
+	TWeakObjectPtr<const UObject> LiveInstanceKeyObject;
+	TWeakObjectPtr<const UObject> LiveInstanceBindingObject;
 };
 
 
@@ -38,7 +42,7 @@ struct FMVVMViewModelDebugEntry
 	FName Name;
 
 	UPROPERTY(VisibleAnywhere, Category = "Viewmodel")
-	FString FullName;
+	FString PathName;
 
 	UPROPERTY(VisibleAnywhere, Category = "Viewmodel")
 	FAssetData ViewModelAsset;
diff --git a/Engine/Plugins/Runtime/ModelViewViewModel/Source/ModelViewViewModelDebuggerEditor/Private/Widgets/SDetailsTab.cpp b/Engine/Plugins/Runtime/ModelViewViewModel/Source/ModelViewViewModelDebuggerEditor/Private/Widgets/SDetailsTab.cpp
--- a/Engine/Plugins/Runtime/ModelViewViewModel/Source/ModelViewViewModelDebuggerEditor/Private/Widgets/SDetailsTab.cpp
+++ b/Engine/Plugins/Runtime/ModelViewViewModel/Source/ModelViewViewModelDebuggerEditor/Private/Widgets/SDetailsTab.cpp
@@ -3,6 +3,7 @@
 #include "Widgets/SDetailsTab.h"
 
 #include "IDetailsView.h"
+#include "IStructureDetailsView.h"
 #include "Modules/ModuleManager.h"
 #include "PropertyEditorModule.h"
 
@@ -26,18 +27,49 @@ void SDetailsTab::Construct(const FArguments& InArgs)
 		DetailsViewArgs.bHideSelectionTip = true;
 	}
 
-	TSharedRef<IDetailsView> PropertyView = PropertyEditorModule.CreateDetailView(DetailsViewArgs);
-	DetailView = PropertyView;
+	if (InArgs._UseStructDetailView)
+	{
+		FStructureDetailsViewArgs StructureViewArgs;
+		{
+			StructureViewArgs.bShowObjects = true;
+			StructureViewArgs.bShowAssets = true;
+			StructureViewArgs.bShowClasses = true;
+			StructureViewArgs.bShowInterfaces = true;
+		}
+
+		TSharedRef<IStructureDetailsView> PropertyView = PropertyEditorModule.CreateStructureDetailView(DetailsViewArgs, StructureViewArgs, TSharedPtr<FStructOnScope>());
+		StructDetailView = PropertyView;
+		ChildSlot
+		[
+			PropertyView->GetWidget().ToSharedRef()
+		];
+	}
+	else
+	{
+		TSharedRef<IDetailsView> PropertyView = PropertyEditorModule.CreateDetailView(DetailsViewArgs);
+		DetailView = PropertyView;
 
-	ChildSlot
-	[
-		PropertyView
-	];
+		ChildSlot
+		[
+			PropertyView
+		];
+	}
 }
 
 void SDetailsTab::SetObjects(const TArray<UObject*>& InObjects)
 {
-	DetailView->SetObjects(InObjects);
+	if (DetailView)
+	{
+		DetailView->SetObjects(InObjects);
+	}
+}
+
+void SDetailsTab::SetStruct(TSharedPtr<FStructOnScope> InStructData)
+{
+	if (StructDetailView)
+	{
+		StructDetailView->SetStructureData(InStructData);
+	}
 }
 
 } //namespace
diff --git a/Engine/Plugins/Runtime/ModelViewViewModel/Source/ModelViewViewModelDebuggerEditor/Private/Widgets/SDetailsTab.h b/Engine/Plugins/Runtime/ModelViewViewModel/Source/ModelViewViewModelDebuggerEditor/Private/Widgets/SDetailsTab.h
--- a/Engine/Plugins/Runtime/ModelViewViewModel/Source/ModelViewViewModelDebuggerEditor/Private/Widgets/SDetailsTab.h
+++ b/Engine/Plugins/Runtime/ModelViewViewModel/Source/ModelViewViewModelDebuggerEditor/Private/Widgets/SDetailsTab.h
@@ -5,6 +5,8 @@
 #include "Widgets/SCompoundWidget.h"
 
 class IDetailsView;
+class IStructureDetailsView;
+class FStructOnScope;
 
 namespace UE::MVVM
 {
@@ -13,14 +15,17 @@ class SDetailsTab : public SCompoundWidget
 {
 public:
 	SLATE_BEGIN_ARGS(SDetailsTab) { }
+	SLATE_ARGUMENT_DEFAULT(bool, UseStructDetailView) = false;
 	SLATE_END_ARGS()
 
 public:
 	void Construct(const FArguments& InArgs);
 	void SetObjects(const TArray<UObject*>& InObjects);
+	void SetStruct(TSharedPtr<FStructOnScope> InStructData);
 
 private:
 	TSharedPtr<IDetailsView> DetailView;
+	TSharedPtr<IStructureDetailsView> StructDetailView;
 };
 
 } //namespace
diff --git a/Engine/Plugins/Runtime/ModelViewViewModel/Source/ModelViewViewModelDebuggerEditor/Private/Widgets/SMainDebugTab.cpp b/Engine/Plugins/Runtime/ModelViewViewModel/Source/ModelViewViewModelDebuggerEditor/Private/Widgets/SMainDebugTab.cpp
--- a/Engine/Plugins/Runtime/ModelViewViewModel/Source/ModelViewViewModelDebuggerEditor/Private/Widgets/SMainDebugTab.cpp
+++ b/Engine/Plugins/Runtime/ModelViewViewModel/Source/ModelViewViewModelDebuggerEditor/Private/Widgets/SMainDebugTab.cpp
@@ -3,6 +3,9 @@
 #include "Widgets/SMainDebugTab.h"
 
 #include "MVVMDebugSnapshot.h"
+#include "MVVMDebugViewModel.h"
+#include "UObject/StructOnScope.h"
+#include "View/MVVMView.h"
 
 #include "ToolMenus.h"
 
@@ -16,8 +19,9 @@
 #include "Widgets/SBoxPanel.h"
 #include "Widgets/SDetailsTab.h"
 #include "Widgets/SMessagesLog.h"
-#include "Widgets/SSelectionTab.h"
 #include "Widgets/SViewModelBindingDetail.h"
+#include "Widgets/SViewModelSelection.h"
+#include "Widgets/SViewSelection.h"
 
 
 #define LOCTEXT_NAMESPACE "MVVMDebuggerMainDebug"
@@ -27,9 +31,11 @@ namespace UE::MVVM
 
 namespace Private
 {
-const FLazyName Stack_Selection = "Selection";
+const FLazyName Stack_ViewSelection = "ViewSelection";
+const FLazyName Stack_ViewmodelSelection = "ViewmodelSelection";
 const FLazyName Stack_Binding = "Binding";
-const FLazyName Stack_Detail = "Detail";
+const FLazyName Stack_LiveDetail = "LiveDetail";
+const FLazyName Stack_EntryDetail = "EntryDetail";
 const FLazyName Stack_Messages = "Messages";
 }
 
@@ -170,11 +176,18 @@ TSharedRef<SWidget> SMainDebug::HandleSnapshotMenuContent()
 
 void SMainDebug::HandleTakeSnapshot()
 {
+	CurrentSelection = ESelection::None;
+	Selection();
+
 	Snapshot = FDebugSnapshot::CreateSnapshot();
-	if (TSharedPtr<SSelectionTab> SelectionViewPtr = SelectionView.Pin())
+	if (TSharedPtr<SViewSelection> SelectionViewPtr = ViewSelection.Pin())
 	{
 		SelectionViewPtr->SetSnapshot(Snapshot);
 	}
+	if (TSharedPtr<SViewModelSelection> SelectionViewModelPtr = ViewModelSelection.Pin())
+	{
+		SelectionViewModelPtr->SetSnapshot(Snapshot);
+	}
 }
 
 
@@ -196,14 +209,40 @@ bool SMainDebug::HasValidSnapshot() const
 }
 
 
-void SMainDebug::HandleObjectSelectionChanged()
+void SMainDebug::HandleViewSelectionChanged()
+{
+	CurrentSelection = ESelection::View;
+	Selection();
+	if (TSharedPtr<SViewModelSelection> ViewModelSelectionPtr = ViewModelSelection.Pin())
+	{
+		ViewModelSelectionPtr->SetSelection(FGuid());
+	}
+}
+
+
+void SMainDebug::HandleViewModleSelectionChanged()
+{
+	CurrentSelection = ESelection::ViewModel;
+	Selection();
+	if (TSharedPtr<SViewSelection> ViewSelectionPtr = ViewSelection.Pin())
+	{
+		ViewSelectionPtr->SetSelection(FGuid());
+	}
+}
+
+
+void SMainDebug::Selection()
 {
 	if (Snapshot == nullptr)
 	{
-		if (TSharedPtr<SDetailsTab> DetailViewPtr = DetailView.Pin())
+		if (TSharedPtr<SDetailsTab> DetailViewPtr = LiveDetailView.Pin())
 		{
 			DetailViewPtr->SetObjects(TArray<UObject*>());
 		}
+		if (TSharedPtr<SDetailsTab> DetailViewPtr = EntryDetailView.Pin())
+		{
+			DetailViewPtr->SetStruct(TSharedPtr<FStructOnScope>());
+		}
 		if (TSharedPtr<SViewModelBindingDetail> ViewModelBindingDetailPtr = ViewModelBindingDetail.Pin())
 		{
 			ViewModelBindingDetailPtr->SetViewModels(TArray<TSharedPtr<FMVVMViewModelDebugEntry>>());
@@ -211,28 +250,82 @@ void SMainDebug::HandleObjectSelectionChanged()
 	}
 	else
 	{
-		if (TSharedPtr<SSelectionTab> SelectionPtr = SelectionView.Pin())
+		TArray<UObject*> SelectedObjects;
+		TSharedPtr<FStructOnScope> SelectedStruct;
+		if (CurrentSelection == ESelection::View)
+		{
+			if (TSharedPtr<SViewSelection> ViewSelectionPtr = ViewSelection.Pin())
+			{
+				for (TSharedPtr<FMVVMViewDebugEntry>& Selection : ViewSelectionPtr->GetSelectedItems())
+				{
+					SelectedObjects.Add(Selection->LiveView.Get());
+					if (!SelectedStruct.IsValid())
+					{
+						SelectedStruct = MakeShared<FStructOnScope>(FMVVMViewDebugEntry::StaticStruct(), reinterpret_cast<uint8*>(Selection.Get()));
+					}
+				}
+			}
+		}
+		else if (CurrentSelection == ESelection::ViewModel)
 		{
-			if (TSharedPtr<SDetailsTab> DetailViewPtr = DetailView.Pin())
+			if (TSharedPtr<SViewModelSelection> ViewModelSelectionPtr = ViewModelSelection.Pin())
 			{
-				DetailViewPtr->SetObjects(SelectionPtr->GetSelectedObjects());
+				for (TSharedPtr<FMVVMViewModelDebugEntry>& Selection : ViewModelSelectionPtr->GetSelectedItems())
+				{
+					SelectedObjects.Add(Selection->LiveViewModel.Get());
+					if (!SelectedStruct.IsValid())
+					{
+						SelectedStruct = MakeShared<FStructOnScope>(FMVVMViewModelDebugEntry::StaticStruct(), reinterpret_cast<uint8*>(Selection.Get()));
+					}
+				}
 			}
+		}
 
-			if (TSharedPtr<SViewModelBindingDetail> ViewModelBindingDetailPtr = ViewModelBindingDetail.Pin())
+		TArray<FDebugItemId> SelectedItems;
+		if (CurrentSelection == ESelection::View)
+		{
+			if (TSharedPtr<SViewSelection> ViewSelectionPtr = ViewSelection.Pin())
 			{
-				TArray<TSharedPtr<FMVVMViewModelDebugEntry>> ViewModels;
-				for (FDebugItemId DebugItem : SelectionPtr->GetSelectedItems())
+				for (const TSharedPtr<FMVVMViewDebugEntry>& Selection : ViewSelectionPtr->GetSelectedItems())
 				{
-					if (DebugItem.Type == FDebugItemId::EType::ViewModel)
+					SelectedItems.Add(FDebugItemId(FDebugItemId::EType::View, Selection->ViewInstanceDebugId));
+				}
+			}
+		}
+		else if (CurrentSelection == ESelection::ViewModel)
+		{
+			if (TSharedPtr<SViewModelSelection> ViewModelSelectionPtr = ViewModelSelection.Pin())
+			{
+				for (const TSharedPtr<FMVVMViewModelDebugEntry>& Selection : ViewModelSelectionPtr->GetSelectedItems())
+				{
+					SelectedItems.Add(FDebugItemId(FDebugItemId::EType::ViewModel, Selection->ViewModelDebugId));
+				}
+			}
+		}
+
+		if (TSharedPtr<SDetailsTab> DetailViewPtr = LiveDetailView.Pin())
+		{
+			DetailViewPtr->SetObjects(SelectedObjects);
+		}
+		if (TSharedPtr<SDetailsTab> DetailViewPtr = EntryDetailView.Pin())
+		{
+			DetailViewPtr->SetStruct(SelectedStruct);
+		}
+
+		if (TSharedPtr<SViewModelBindingDetail> ViewModelBindingDetailPtr = ViewModelBindingDetail.Pin())
+		{
+			TArray<TSharedPtr<FMVVMViewModelDebugEntry>> ViewModels;
+			for (FDebugItemId DebugItem : SelectedItems)
+			{
+				if (DebugItem.Type == FDebugItemId::EType::ViewModel)
+				{
+					if (TSharedPtr<FMVVMViewModelDebugEntry> Found = Snapshot->FindViewModel(DebugItem.Id))
 					{
-						if (TSharedPtr<FMVVMViewModelDebugEntry> Found = Snapshot->FindViewModel(DebugItem.Id))
-						{
-							ViewModels.Add(Found);
-						}
+						ViewModels.Add(Found);
 					}
 				}
-				ViewModelBindingDetailPtr->SetViewModels(ViewModels);
 			}
+			ViewModelBindingDetailPtr->SetViewModels(ViewModels);
 		}
 	}
 }
@@ -240,7 +333,7 @@ void SMainDebug::HandleObjectSelectionChanged()
 
 TSharedRef<SWidget> SMainDebug::CreateDockingArea(const TSharedRef<SDockTab>& InParentTab)
 {
-	const FName LayoutName = TEXT("MVVMDebugger_Layout_v1.8");
+	const FName LayoutName = TEXT("MVVMDebugger_Layout_v1.0");
 	TSharedRef<FTabManager::FLayout> DefaultLayer = FTabManager::NewLayout(LayoutName)
 		->AddArea
 		(
@@ -252,10 +345,22 @@ TSharedRef<SWidget> SMainDebug::CreateDockingArea(const TSharedRef<SDockTab>& In
 				->SetOrientation(Orient_Horizontal)
 				->Split
 				(
-					FTabManager::NewStack()
-					->SetSizeCoefficient(0.25f)
-					->SetHideTabWell(true)
-					->AddTab(Private::Stack_Selection, ETabState::OpenedTab)
+					FTabManager::NewSplitter()
+					->SetOrientation(Orient_Vertical)
+					->Split
+					(
+						FTabManager::NewStack()
+						->SetSizeCoefficient(0.25f)
+						->SetHideTabWell(true)
+						->AddTab(Private::Stack_ViewSelection, ETabState::OpenedTab)
+					)
+					->Split
+					(
+						FTabManager::NewStack()
+						->SetSizeCoefficient(0.25f)
+						->SetHideTabWell(true)
+						->AddTab(Private::Stack_ViewmodelSelection, ETabState::OpenedTab)
+					)
 				)
 				->Split
 				(
@@ -266,10 +371,22 @@ TSharedRef<SWidget> SMainDebug::CreateDockingArea(const TSharedRef<SDockTab>& In
 				)
 				->Split
 				(
-					FTabManager::NewStack()
-					->SetSizeCoefficient(0.25f)
-					->SetHideTabWell(true)
-					->AddTab(Private::Stack_Detail, ETabState::OpenedTab)
+					FTabManager::NewSplitter()
+					->SetOrientation(Orient_Vertical)
+					->Split
+					(
+						FTabManager::NewStack()
+						->SetSizeCoefficient(0.25f)
+						->SetHideTabWell(true)
+						->AddTab(Private::Stack_LiveDetail, ETabState::OpenedTab)
+					)
+					->Split
+					(
+						FTabManager::NewStack()
+						->SetSizeCoefficient(0.25f)
+						->SetHideTabWell(true)
+						->AddTab(Private::Stack_EntryDetail, ETabState::OpenedTab)
+					)
 				)
 			)
 			->Split
@@ -287,12 +404,16 @@ TSharedRef<SWidget> SMainDebug::CreateDockingArea(const TSharedRef<SDockTab>& In
 
 	TabManager = FGlobalTabmanager::Get()->NewTabManager(InParentTab);
 
-	TabManager->RegisterTabSpawner(Private::Stack_Selection, FOnSpawnTab::CreateSP(this, &SMainDebug::SpawnSelectionTab))
+	TabManager->RegisterTabSpawner(Private::Stack_ViewSelection, FOnSpawnTab::CreateSP(this, &SMainDebug::SpawnViewSelectionTab))
+		.SetDisplayName(LOCTEXT("SelectionTab", "Selection"));
+	TabManager->RegisterTabSpawner(Private::Stack_ViewmodelSelection, FOnSpawnTab::CreateSP(this, &SMainDebug::SpawnViewModelSelectionTab))
 		.SetDisplayName(LOCTEXT("SelectionTab", "Selection"));
 	TabManager->RegisterTabSpawner(Private::Stack_Binding, FOnSpawnTab::CreateSP(this, &SMainDebug::SpawnBindingTab))
 		.SetDisplayName(LOCTEXT("BindingTab", "Bindings"));
-	TabManager->RegisterTabSpawner(Private::Stack_Detail, FOnSpawnTab::CreateSP(this, &SMainDebug::SpawnDetailTab))
-		.SetDisplayName(LOCTEXT("DetailTab", "Details"));
+	TabManager->RegisterTabSpawner(Private::Stack_LiveDetail, FOnSpawnTab::CreateSP(this, &SMainDebug::SpawnLiveDetailTab))
+		.SetDisplayName(LOCTEXT("LiveDetailTab", "Live Details"));
+	TabManager->RegisterTabSpawner(Private::Stack_EntryDetail, FOnSpawnTab::CreateSP(this, &SMainDebug::SpawnEntryDetailTab))
+		.SetDisplayName(LOCTEXT("EntryDetailTab", "Detail"));
 	TabManager->RegisterTabSpawner(Private::Stack_Messages, FOnSpawnTab::CreateSP(this, &SMainDebug::SpawnMessagesTab))
 		.SetDisplayName(LOCTEXT("MessageLogTab", "Messages Log"));
 
@@ -300,14 +421,27 @@ TSharedRef<SWidget> SMainDebug::CreateDockingArea(const TSharedRef<SDockTab>& In
 }
 
 
-TSharedRef<SDockTab> SMainDebug::SpawnSelectionTab(const FSpawnTabArgs& Args)
+TSharedRef<SDockTab> SMainDebug::SpawnViewSelectionTab(const FSpawnTabArgs& Args)
 {
-	TSharedRef<SSelectionTab> LocalSelectionView = SNew(SSelectionTab)
-		.OnSelectionChanged(this, &SMainDebug::HandleObjectSelectionChanged);
-	SelectionView = LocalSelectionView;
+	TSharedRef<SViewSelection> LocalSelectionView = SNew(SViewSelection)
+		.OnSelectionChanged(this, &SMainDebug::HandleViewSelectionChanged);
+	ViewSelection = LocalSelectionView;
 	LocalSelectionView->SetSnapshot(Snapshot);
 	return SNew(SDockTab)
-		.Label(LOCTEXT("BindingTab", "Bindings"))
+		.Label(LOCTEXT("ViewTab", "View"))
+		[
+			LocalSelectionView
+		];
+}
+
+TSharedRef<SDockTab> SMainDebug::SpawnViewModelSelectionTab(const FSpawnTabArgs& Args)
+{
+	TSharedRef<SViewModelSelection> LocalSelectionView = SNew(SViewModelSelection)
+		.OnSelectionChanged(this, &SMainDebug::HandleViewModleSelectionChanged);
+	ViewModelSelection = LocalSelectionView;
+	LocalSelectionView->SetSnapshot(Snapshot);
+	return SNew(SDockTab)
+		.Label(LOCTEXT("ViewmodelTab", "Viewmodel"))
 		[
 			LocalSelectionView
 		];
@@ -324,12 +458,24 @@ TSharedRef<SDockTab> SMainDebug::SpawnBindingTab(const FSpawnTabArgs& Args)
 }
 
 
-TSharedRef<SDockTab> SMainDebug::SpawnDetailTab(const FSpawnTabArgs& Args)
+TSharedRef<SDockTab> SMainDebug::SpawnLiveDetailTab(const FSpawnTabArgs& Args)
+{
+	return SNew(SDockTab)
+		.Label(LOCTEXT("LiveDetailTab", "Live Object Details"))
+		[
+			SAssignNew(LiveDetailView, SDetailsTab)
+			.UseStructDetailView(false)
+		];
+}
+
+
+TSharedRef<SDockTab> SMainDebug::SpawnEntryDetailTab(const FSpawnTabArgs& Args)
 {
 	return SNew(SDockTab)
-		.Label(LOCTEXT("DetailTab", "Details"))
+		.Label(LOCTEXT("EntryDetailTab", "Entry Details"))
 		[
-			SAssignNew(DetailView, SDetailsTab)
+			SAssignNew(EntryDetailView, SDetailsTab)
+			.UseStructDetailView(true)
 		];
 }
 
diff --git a/Engine/Plugins/Runtime/ModelViewViewModel/Source/ModelViewViewModelDebuggerEditor/Private/Widgets/SMainDebugTab.h b/Engine/Plugins/Runtime/ModelViewViewModel/Source/ModelViewViewModelDebuggerEditor/Private/Widgets/SMainDebugTab.h
--- a/Engine/Plugins/Runtime/ModelViewViewModel/Source/ModelViewViewModelDebuggerEditor/Private/Widgets/SMainDebugTab.h
+++ b/Engine/Plugins/Runtime/ModelViewViewModel/Source/ModelViewViewModelDebuggerEditor/Private/Widgets/SMainDebugTab.h
@@ -18,8 +18,9 @@ class FDebugSnapshot;
 
 class SDetailsTab;
 class SMessagesLog;
-class SSelectionTab;
+class SViewModelSelection;
 class SViewModelBindingDetail;
+class SViewSelection;
 
 class SMainDebug : public SCompoundWidget
 {
@@ -40,23 +41,37 @@ private:
 	void HandleSaveSnapshot();
 	bool HasValidSnapshot() const;
 
-	void HandleObjectSelectionChanged();
+	void HandleViewSelectionChanged();
+	void HandleViewModleSelectionChanged();
+	void Selection();
 
 	TSharedRef<SWidget> CreateDockingArea(const TSharedRef<SDockTab>& InParentTab);
-	TSharedRef<SDockTab> SpawnSelectionTab(const FSpawnTabArgs& Args);
+	TSharedRef<SDockTab> SpawnViewSelectionTab(const FSpawnTabArgs& Args);
+	TSharedRef<SDockTab> SpawnViewModelSelectionTab(const FSpawnTabArgs& Args);
 	TSharedRef<SDockTab> SpawnBindingTab(const FSpawnTabArgs& Args);
-	TSharedRef<SDockTab> SpawnDetailTab(const FSpawnTabArgs& Args);
+	TSharedRef<SDockTab> SpawnLiveDetailTab(const FSpawnTabArgs& Args);
+	TSharedRef<SDockTab> SpawnEntryDetailTab(const FSpawnTabArgs& Args);
 	TSharedRef<SDockTab> SpawnMessagesTab(const FSpawnTabArgs& Args);
 
 private:
 	TSharedPtr<FTabManager> TabManager;
 
-	TWeakPtr<SDetailsTab> DetailView;
+	TWeakPtr<SDetailsTab> LiveDetailView;
+	TWeakPtr<SDetailsTab> EntryDetailView;
 	TWeakPtr<SMessagesLog> MessageLog;
-	TWeakPtr<SSelectionTab> SelectionView;
+	TWeakPtr<SViewSelection> ViewSelection;
+	TWeakPtr<SViewModelSelection> ViewModelSelection;
 	TWeakPtr<SViewModelBindingDetail> ViewModelBindingDetail;
 
 	TSharedPtr<FDebugSnapshot> Snapshot;
+
+	enum class ESelection
+	{
+		None,
+		View,
+		ViewModel,
+	};
+	ESelection CurrentSelection = ESelection::None;
 };
 
 } //namespace
diff --git a/Engine/Plugins/Runtime/ModelViewViewModel/Source/ModelViewViewModelDebuggerEditor/Private/Widgets/SMessagesLog.cpp b/Engine/Plugins/Runtime/ModelViewViewModel/Source/ModelViewViewModelDebuggerEditor/Private/Widgets/SMessagesLog.cpp
--- a/Engine/Plugins/Runtime/ModelViewViewModel/Source/ModelViewViewModelDebuggerEditor/Private/Widgets/SMessagesLog.cpp
+++ b/Engine/Plugins/Runtime/ModelViewViewModel/Source/ModelViewViewModelDebuggerEditor/Private/Widgets/SMessagesLog.cpp
@@ -4,6 +4,7 @@
 
 #include "IMessageLogListing.h"
 #include "MessageLogModule.h"
+#include "MVVMMessageLog.h"
 #include "Modules/ModuleManager.h"
 
 #include "Widgets/SBoxPanel.h"
@@ -17,7 +18,7 @@ namespace UE::MVVM
 
 void SMessagesLog::Construct(const FArguments& InArgs)
 {
-	const FName LogName = "Live Link";
+	const FName LogName = UE::MVVM::FMessageLog::LogName;
 	TSharedPtr<IMessageLogListing> MessageLogListing;
 
 	FMessageLogModule& MessageLogModule = FModuleManager::LoadModuleChecked<FMessageLogModule>("MessageLog");
@@ -29,24 +30,24 @@ void SMessagesLog::Construct(const FArguments& InArgs)
 
 	ChildSlot
 	[
-		SNew(SOverlay)
-		+SOverlay::Slot()
-		[
+		//SNew(SOverlay)
+		//+SOverlay::Slot()
+		//[
 			MessageLogListingWidget
-		]
-		+SOverlay::Slot()
-		.HAlign(HAlign_Left)
-		.VAlign(VAlign_Bottom)
-		[
-			SNew(SHorizontalBox)
-			+SHorizontalBox::Slot()
-			.AutoWidth()
-			.Padding(10, 4, 4, 10)
-			[
-				SNew(STextBlock)
-				.Text(this, &SMessagesLog::GetMessageCountText)
-			]
-		]
+		//]
+		//+SOverlay::Slot()
+		//.HAlign(HAlign_Left)
+		//.VAlign(VAlign_Bottom)
+		//[
+		//	SNew(SHorizontalBox)
+		//	+SHorizontalBox::Slot()
+		//	.AutoWidth()
+		//	.Padding(10, 4, 4, 10)
+		//	[
+		//		SNew(STextBlock)
+		//		.Text(this, &SMessagesLog::GetMessageCountText)
+		//	]
+		//]
 	];
 }
 
@@ -55,7 +56,7 @@ FText SMessagesLog::GetMessageCountText() const
 	int32 ErrorCount = 0;
 	int32 WarningCount = 0;
 	int32 InfoCount = 0;
-	//FLiveLinkLog::GetInstance()->GetLogCount(ErrorCount, WarningCount, InfoCount);
+	//FMessageLog::GetLogCount(ErrorCount, WarningCount, InfoCount);
 	return FText::Format(LOCTEXT("MessageCountText", "{0} Error(s)  {1} Warning(s)"), FText::AsNumber(ErrorCount), FText::AsNumber(WarningCount));
 }
 
diff --git a/Engine/Plugins/Runtime/ModelViewViewModel/Source/ModelViewViewModelDebuggerEditor/Private/Widgets/SSelectionTab.cpp b/Engine/Plugins/Runtime/ModelViewViewModel/Source/ModelViewViewModelDebuggerEditor/Private/Widgets/SSelectionTab.cpp
deleted file mode 100644
diff --git a/Engine/Plugins/Runtime/ModelViewViewModel/Source/ModelViewViewModelDebuggerEditor/Private/Widgets/SSelectionTab.h b/Engine/Plugins/Runtime/ModelViewViewModel/Source/ModelViewViewModelDebuggerEditor/Private/Widgets/SSelectionTab.h
deleted file mode 100644
diff --git a/Engine/Plugins/Runtime/ModelViewViewModel/Source/ModelViewViewModelDebuggerEditor/Private/Widgets/SViewModelBindingDetail.cpp b/Engine/Plugins/Runtime/ModelViewViewModel/Source/ModelViewViewModelDebuggerEditor/Private/Widgets/SViewModelBindingDetail.cpp
--- a/Engine/Plugins/Runtime/ModelViewViewModel/Source/ModelViewViewModelDebuggerEditor/Private/Widgets/SViewModelBindingDetail.cpp
+++ b/Engine/Plugins/Runtime/ModelViewViewModel/Source/ModelViewViewModelDebuggerEditor/Private/Widgets/SViewModelBindingDetail.cpp
@@ -22,8 +22,10 @@ namespace Private
 
 namespace ViewColumnName
 {
-FLazyName FieldId = "FieldId";
-FLazyName ObjectName = "ObjectName";
+FLazyName KeyFieldId = "FieldId";
+FLazyName KeyObjectName = "KeyObjectName";
+FLazyName BindingObjectName = "BindingObjectName";
+FLazyName BindingFunctionName = "BindingFunction";
 }
 
 class SViewModelBindingDetailItem : public SMultiColumnTableRow<SViewModelBindingDetail::FViewModelBindingId>
@@ -47,15 +49,26 @@ class SViewModelBindingDetailItem : public SMultiColumnTableRow<SViewModelBindin
 	{
 		if (BindingId.ViewModel && BindingId.ViewModel->FieldBound.IsValidIndex(BindingId.Index))
 		{
-			if (ColumnName == ViewColumnName::FieldId)
+			const FMVVMViewModelFieldBoundDebugEntry& Entry = BindingId.ViewModel->FieldBound[BindingId.Index];
+			if (ColumnName == ViewColumnName::KeyObjectName)
 			{
 				return SNew(STextBlock)
-					.Text(FText::FromName(BindingId.ViewModel->FieldBound[BindingId.Index].FieldId.GetFieldName()));
+					.Text(FText::FromName(Entry.KeyObjectName));
 			}
-			if (ColumnName == ViewColumnName::ObjectName)
+			if (ColumnName == ViewColumnName::KeyFieldId)
 			{
 				return SNew(STextBlock)
-					.Text(FText::FromString(BindingId.ViewModel->FullName));
+					.Text(FText::FromName(Entry.KeyFieldId.GetFieldName()));
+			}
+			if (ColumnName == ViewColumnName::BindingObjectName)
+			{
+				return SNew(STextBlock)
+					.Text(FText::FromString(Entry.BindingObjectPathName));
+			}
+			if (ColumnName == ViewColumnName::BindingFunctionName)
+			{
+				return SNew(STextBlock)
+					.Text(FText::FromName(Entry.BindingFunctionName));
 			}
 		}
 
@@ -72,20 +85,26 @@ private:
 void SViewModelBindingDetail::Construct(const FArguments& InArgs)
 {
 	TSharedRef<SHeaderRow> HeaderRow = SNew(SHeaderRow)
-		+ SHeaderRow::Column(Private::ViewColumnName::FieldId)
+		+ SHeaderRow::Column(Private::ViewColumnName::KeyObjectName)
+		.FillWidth(0.25f)
+		.DefaultLabel(LOCTEXT("KeyObjectdColumnHeaderName", "Key Object"))
+		+ SHeaderRow::Column(Private::ViewColumnName::KeyFieldId)
 		.FillWidth(0.25f)
 		.DefaultLabel(LOCTEXT("FieldIdColumnHeaderName", "Field"))
-		+ SHeaderRow::Column(Private::ViewColumnName::ObjectName)
-		.FillWidth(0.75f)
-		.DefaultLabel(LOCTEXT("ObjectColumnHeaderName", "Object"));
+		+ SHeaderRow::Column(Private::ViewColumnName::BindingObjectName)
+		.FillWidth(0.50f)
+		.DefaultLabel(LOCTEXT("BindingObjectColumnHeaderName", "Delegate Object"))
+		+SHeaderRow::Column(Private::ViewColumnName::BindingFunctionName)
+		.FillWidth(0.50f)
+		.DefaultLabel(LOCTEXT("BindingFunctionColumnHeaderName", "Delegate Function"));
 
 	SAssignNew(ListView, SListView<FViewModelBindingId>)
 		.ListItemsSource(&Entries)
 		//.SelectionMode(ESelectionMode::Single)
 		.OnGenerateRow(this, &SViewModelBindingDetail::HandleGenerateWidgetForItem)
 		.OnSelectionChanged(this, &SViewModelBindingDetail::HandleListSelectionChanged)
-		//.OnMouseButtonDoubleClick(this, &SSceneOutliner::OnOutlinerTreeDoubleClick)
-		//.OnContextMenuOpening(this, &SLiveLinkClientPanel::OnOpenVirtualSubjectContextMenu)
+		//.OnMouseButtonDoubleClick(this, &SViewModelBindingDetail::OnOutlinerTreeDoubleClick)
+		//.OnContextMenuOpening(this, &SViewModelBindingDetail::OnOpenVirtualSubjectContextMenu)
 		.HeaderRow(HeaderRow);
 
 	ChildSlot
diff --git a/Engine/Plugins/Runtime/ModelViewViewModel/Source/ModelViewViewModelDebuggerEditor/Private/Widgets/SViewModelSelection.cpp b/Engine/Plugins/Runtime/ModelViewViewModel/Source/ModelViewViewModelDebuggerEditor/Private/Widgets/SViewModelSelection.cpp
--- a/Engine/Plugins/Runtime/ModelViewViewModel/Source/ModelViewViewModelDebuggerEditor/Private/Widgets/SViewModelSelection.cpp
+++ b/Engine/Plugins/Runtime/ModelViewViewModel/Source/ModelViewViewModelDebuggerEditor/Private/Widgets/SViewModelSelection.cpp
@@ -121,7 +121,7 @@ TSharedRef<SHeaderRow> SViewModelSelection::BuildHeaderRowImpl() const
 
 void SViewModelSelection::GatherFilterStringsImpl(TSharedPtr<FMVVMViewModelDebugEntry> Item, TArray<FString>& OutStrings) const
 {
-	OutStrings.Add(Item->FullName);
+	OutStrings.Add(Item->PathName);
 	OutStrings.Add(Item->ViewModelAsset.GetObjectPathString());
 }
 
diff --git a/Engine/Source/Runtime/FieldNotification/Private/FieldNotificationDelegate.cpp b/Engine/Source/Runtime/FieldNotification/Private/FieldNotificationDelegate.cpp
--- a/Engine/Source/Runtime/FieldNotification/Private/FieldNotificationDelegate.cpp
+++ b/Engine/Source/Runtime/FieldNotification/Private/FieldNotificationDelegate.cpp
@@ -515,8 +515,7 @@ TArray<FFieldMulticastDelegate::FDelegateView> FFieldMulticastDelegate::GetView(
 	{
 		if (Element.Delegate.IsBound())
 		{
-			//Result.Emplace(Element.Key.Object.Get(), Element.Key.Id, Element.Delegate.TryGetBoundFunctionName());
-			Result.Emplace(Element.Key.Object.Get(), Element.Key.Id, Element.Key.DynamicName);
+			Result.Emplace(Element.Key.Object.Get(), Element.Key.Id, Element.Delegate.GetUObject(), Element.Key.DynamicName);
 		}
 	}
 	return Result;
diff --git a/Engine/Source/Runtime/FieldNotification/Public/FieldNotificationDelegate.h b/Engine/Source/Runtime/FieldNotification/Public/FieldNotificationDelegate.h
--- a/Engine/Source/Runtime/FieldNotification/Public/FieldNotificationDelegate.h
+++ b/Engine/Source/Runtime/FieldNotification/Public/FieldNotificationDelegate.h
@@ -105,12 +105,13 @@ public:
 
 	struct FDelegateView
 	{
-		FDelegateView(const UObject* Obj, FFieldId Id, FName Function)
-			: Object(Obj), Field(Id), FunctionName(Function)
+		FDelegateView(const UObject* InKeyObj, FFieldId InKeyId, const UObject* InObj, FName InFunction)
+			: KeyObject(InKeyObj), KeyField(InKeyId), BindingObject(InObj), BindingFunctionName(InFunction)
 		{}
-		const UObject* Object;
-		FFieldId Field;
-		FName FunctionName;
+		const UObject* KeyObject;
+		FFieldId KeyField;
+		const UObject* BindingObject;
+		FName BindingFunctionName;
 	};
 	TArray<FDelegateView> GetView() const;
 
