MVVM: Fix Localization issue in Viewmodel debugger
#rnx

diff --git a/Engine/Plugins/Runtime/ModelViewViewModel/Source/ModelViewViewModelDebuggerEditor/Private/Widgets/SMainDebugTab.cpp b/Engine/Plugins/Runtime/ModelViewViewModel/Source/ModelViewViewModelDebuggerEditor/Private/Widgets/SMainDebugTab.cpp
--- a/Engine/Plugins/Runtime/ModelViewViewModel/Source/ModelViewViewModelDebuggerEditor/Private/Widgets/SMainDebugTab.cpp
+++ b/Engine/Plugins/Runtime/ModelViewViewModel/Source/ModelViewViewModelDebuggerEditor/Private/Widgets/SMainDebugTab.cpp
@@ -411,9 +411,9 @@ TSharedRef<SWidget> SMainDebug::CreateDockingArea(const TSharedRef<SDockTab>& In
 	TabManager->RegisterTabSpawner(Private::Stack_Binding, FOnSpawnTab::CreateSP(this, &SMainDebug::SpawnBindingTab))
 		.SetDisplayName(LOCTEXT("BindingTab", "Bindings"));
 	TabManager->RegisterTabSpawner(Private::Stack_LiveDetail, FOnSpawnTab::CreateSP(this, &SMainDebug::SpawnLiveDetailTab))
-		.SetDisplayName(LOCTEXT("LiveDetailTab", "Live Details"));
+		.SetDisplayName(LOCTEXT("LiveDetailTab", "Live Object Details"));
 	TabManager->RegisterTabSpawner(Private::Stack_EntryDetail, FOnSpawnTab::CreateSP(this, &SMainDebug::SpawnEntryDetailTab))
-		.SetDisplayName(LOCTEXT("EntryDetailTab", "Detail"));
+		.SetDisplayName(LOCTEXT("EntryDetailTab", "Entry Details"));
 	TabManager->RegisterTabSpawner(Private::Stack_Messages, FOnSpawnTab::CreateSP(this, &SMainDebug::SpawnMessagesTab))
 		.SetDisplayName(LOCTEXT("MessageLogTab", "Messages Log"));
 
