MVVM: In View Binding Panel, the "add" button now add the selected widget. The user can use the default behavior, saved in a setting.

diff --git a/Engine/Plugins/Runtime/ModelViewViewModel/Source/ModelViewViewModelEditor/Private/Widgets/SMVVMViewBindingListView.cpp b/Engine/Plugins/Runtime/ModelViewViewModel/Source/ModelViewViewModelEditor/Private/Widgets/SMVVMViewBindingListView.cpp
--- a/Engine/Plugins/Runtime/ModelViewViewModel/Source/ModelViewViewModelEditor/Private/Widgets/SMVVMViewBindingListView.cpp
+++ b/Engine/Plugins/Runtime/ModelViewViewModel/Source/ModelViewViewModelEditor/Private/Widgets/SMVVMViewBindingListView.cpp
@@ -156,6 +156,23 @@ namespace Private
 			ExpandAll(TreeView, Child);
 		}
 	}
+
+	TSharedPtr<FBindingEntry> FindBinding(FGuid BindingId, TConstArrayView<TSharedPtr<FBindingEntry>> Entries)
+	{
+		for (const TSharedPtr<FBindingEntry>& Entry : Entries)
+		{
+			if (Entry->GetRowType() == FBindingEntry::ERowType::Binding && Entry->GetBindingId() == BindingId)
+			{
+				return Entry;
+			}
+			TSharedPtr<FBindingEntry> Result = FindBinding(BindingId, Entry->GetChildren());
+			if (Result)
+			{
+				return Result;
+			}
+		}
+		return TSharedPtr<FBindingEntry>();
+	}
 }
 
 class SWidgetRow : public STableRow<TSharedPtr<FBindingEntry>>
@@ -1366,6 +1383,15 @@ TSharedPtr<SWidget> SBindingsList::OnSourceConstructContextMenu()
 	return MenuBuilder.MakeWidget();
 }
 
+void SBindingsList::RequestNavigateToBinding(FGuid BindingId)
+{
+	TSharedPtr<FBindingEntry> Entry = Private::FindBinding(BindingId, RootGroups);
+	if (Entry && TreeView)
+	{
+		TreeView->RequestNavigateToItem(Entry);
+	}
+}
+
 FReply SBindingsList::OnKeyDown(const FGeometry& MyGeometry, const FKeyEvent& InKeyEvent)
 {
 	if (InKeyEvent.GetKey() == EKeys::Delete)
diff --git a/Engine/Plugins/Runtime/ModelViewViewModel/Source/ModelViewViewModelEditor/Private/Widgets/SMVVMViewBindingListView.h b/Engine/Plugins/Runtime/ModelViewViewModel/Source/ModelViewViewModelEditor/Private/Widgets/SMVVMViewBindingListView.h
--- a/Engine/Plugins/Runtime/ModelViewViewModel/Source/ModelViewViewModelEditor/Private/Widgets/SMVVMViewBindingListView.h
+++ b/Engine/Plugins/Runtime/ModelViewViewModel/Source/ModelViewViewModelEditor/Private/Widgets/SMVVMViewBindingListView.h
@@ -32,6 +32,8 @@ public:
 	/** Constructs context menu used for right click and dropdown button */
 	TSharedPtr<SWidget> OnSourceConstructContextMenu();
 
+	void RequestNavigateToBinding(FGuid BindingId);
+
 private:
 	TSharedRef<ITableRow> GenerateEntryRow(TSharedPtr<FBindingEntry> Entry, const TSharedRef<STableViewBase>& OwnerTable) const;
 	void OnSourceListSelectionChanged(TSharedPtr<FBindingEntry> Entry, ESelectInfo::Type SelectionType) const;
diff --git a/Engine/Plugins/Runtime/ModelViewViewModel/Source/ModelViewViewModelEditor/Private/Widgets/SMVVMViewBindingPanel.cpp b/Engine/Plugins/Runtime/ModelViewViewModel/Source/ModelViewViewModelEditor/Private/Widgets/SMVVMViewBindingPanel.cpp
--- a/Engine/Plugins/Runtime/ModelViewViewModel/Source/ModelViewViewModelEditor/Private/Widgets/SMVVMViewBindingPanel.cpp
+++ b/Engine/Plugins/Runtime/ModelViewViewModel/Source/ModelViewViewModelEditor/Private/Widgets/SMVVMViewBindingPanel.cpp
@@ -3,12 +3,15 @@
 #include "Widgets/SMVVMViewBindingPanel.h"
 
 #include "DetailsViewArgs.h"
-#include "WidgetBlueprintEditor.h"
+#include "Engine/Engine.h"
 #include "IDetailsView.h"
-#include "WidgetBlueprintToolMenuContext.h"
 #include "MVVMBlueprintView.h"
 #include "MVVMDeveloperProjectSettings.h"
+#include "MVVMEditorSubsystem.h"
 #include "MVVMWidgetBlueprintExtension_View.h"
+#include "WidgetBlueprintEditor.h"
+#include "WidgetBlueprintToolMenuContext.h"
+#include "WidgetReference.h"
 
 #include "IStructureDetailsView.h"
 #include "PropertyEditorModule.h"
@@ -18,6 +21,7 @@
 #include "Tabs/MVVMBindingSummoner.h"
 #include "Tabs/MVVMViewModelSummoner.h"
 #include "Styling/MVVMEditorStyle.h"
+#include "Styling/ToolBarStyle.h"
 #include "Toolkits/IToolkitHost.h"
 #include "Widgets/Docking/SDockTab.h"
 #include "Widgets/Input/SButton.h"
@@ -28,7 +32,6 @@
 #include "Widgets/SMVVMViewBindingListView.h"
 #include "SPositiveActionButton.h"
 
-
 #define LOCTEXT_NAMESPACE "BindingPanel"
 
 namespace UE::MVVM
@@ -65,6 +68,8 @@ void SBindingsPanel::Construct(const FArguments& InArgs, TSharedPtr<FWidgetBluep
 	WeakBlueprintEditor = WidgetBlueprintEditor;
 	bIsDrawerTab = bInIsDrawerTab;
 
+	LoadSettings();
+
 	UWidgetBlueprint* WidgetBlueprint = WidgetBlueprintEditor->GetWidgetBlueprintObj();
 	check(WidgetBlueprint);
 
@@ -118,6 +123,25 @@ SBindingsPanel::~SBindingsPanel()
 	}
 }
 
+void SBindingsPanel::SaveSettings()
+{
+	GConfig->SetInt(TEXT("MVVMViewBindingPanel"), TEXT("LastAddBindingMode"), static_cast<int32>(AddBindingMode), *GEditorPerProjectIni);
+}
+
+void SBindingsPanel::LoadSettings()
+{
+	GConfig->SetInt(TEXT("MVVMViewBindingPanel"), TEXT("LastAddBindingMode"), static_cast<int32>(AddBindingMode), *GEditorPerProjectIni);
+	if (GConfig->DoesSectionExist(TEXT("MVVMViewBindingPanel"), *GEditorPerProjectIni))
+	{
+		int32 AddBindingModeAsInt = static_cast<int32>(EAddBindingMode::Selected);
+		GConfig->GetInt(TEXT("MVVMViewBindingPanel"), TEXT("LastAddBindingMode"), AddBindingModeAsInt, *GEditorPerProjectIni);
+		if (AddBindingModeAsInt >= 0 && AddBindingModeAsInt <= 1)
+		{
+			AddBindingMode = static_cast<EAddBindingMode>(AddBindingModeAsInt);
+		}
+	}
+}
+
 FReply SBindingsPanel::OnKeyDown(const FGeometry& MyGeometry, const FKeyEvent& InKeyEvent)
 {
 	FReply Reply = FReply::Unhandled();
@@ -177,18 +201,49 @@ void SBindingsPanel::Tick(const FGeometry& AllottedGeometry, const double InCurr
 	Super::Tick(AllottedGeometry, InCurrentTime, InDeltaTime);
 }
 
-FReply SBindingsPanel::AddDefaultBinding()
+void SBindingsPanel::AddDefaultBinding()
 {
+	if (!CanAddBinding())
+	{
+		return;
+	}
 	if (UMVVMWidgetBlueprintExtension_View* MVVMExtensionPtr = MVVMExtension.Get())
 	{
-		if (UMVVMBlueprintView* EditorData = MVVMExtensionPtr->GetBlueprintView())
+		FGuid AddedBindingId;
+		UMVVMEditorSubsystem* EditorSubsystem = GEditor->GetEditorSubsystem<UMVVMEditorSubsystem>();
+		if (TSharedPtr<FWidgetBlueprintEditor> BlueprintEditor = WeakBlueprintEditor.Pin())
 		{
-			check(BindingsList);
-			EditorData->AddDefaultBinding();
-			BindingsList->Refresh();
+			bool bBindingAdded = false;
+
+			if (AddBindingMode == EAddBindingMode::Selected)
+			{
+				for (const FWidgetReference& WidgetReference : BlueprintEditor->GetSelectedWidgets())
+				{
+					if (WidgetReference.IsValid() && WidgetReference.GetTemplate())
+					{
+						FMVVMBlueprintViewBinding& Binding = EditorSubsystem->AddBinding(MVVMExtensionPtr->GetWidgetBlueprint());
+						FMVVMBlueprintPropertyPath Path;
+						Path.SetWidgetName(WidgetReference.GetTemplate()->GetFName());
+						EditorSubsystem->SetWidgetPropertyForBinding(MVVMExtensionPtr->GetWidgetBlueprint(), Binding, Path);
+						AddedBindingId = Binding.BindingId;
+
+						bBindingAdded = true;
+					}
+				}
+			}
+
+			if (!bBindingAdded)
+			{
+				FMVVMBlueprintViewBinding& Binding = EditorSubsystem->AddBinding(MVVMExtensionPtr->GetWidgetBlueprint());
+				AddedBindingId = Binding.BindingId;
+			}
+
+			if (AddedBindingId.IsValid() && BindingsList)
+			{
+				BindingsList->RequestNavigateToBinding(AddedBindingId);
+			}
 		}
 	}
-	return FReply::Handled();
 }
 
 bool SBindingsPanel::CanAddBinding() const
@@ -204,11 +259,51 @@ bool SBindingsPanel::CanAddBinding() const
 	return false;
 }
 
+FText SBindingsPanel::GetAddBindingText() const
+{
+	if (TSharedPtr<FWidgetBlueprintEditor> WidgetEditor = WeakBlueprintEditor.Pin())
+	{
+		if (AddBindingMode == EAddBindingMode::Selected)
+			{
+			const TSet<FWidgetReference>& SelectedWidgets = WidgetEditor->GetSelectedWidgets();
+			int32 NumberOfWidgetSelected = SelectedWidgets.Num();
+			if (NumberOfWidgetSelected > 1)
+			{
+				return LOCTEXT("AddWidgets", "Add Widgets");
+			}
+			else if (NumberOfWidgetSelected == 1)
+			{
+				for (const FWidgetReference& Item : SelectedWidgets)
+				{
+					UWidget* WidgetTemplate = Item.GetTemplate();
+					if (WidgetTemplate)
+					{
+						FText LabelText = WidgetTemplate->GetLabelText();
+						if (!LabelText.IsEmpty())
+						{ 
+							return FText::Format(LOCTEXT("AddForWidget", "Add Widget {0}"), LabelText);
+						}
+					}
+					break;
+				}
+			}
+		}
+	}
+	return LOCTEXT("AddWidget", "Add Widget");
+}
+
 FText SBindingsPanel::GetAddBindingToolTip() const
 {
 	if (CanAddBinding())
 	{
-		return LOCTEXT("AddBindingTooltip", "Add an empty binding.");
+		if (AddBindingMode == EAddBindingMode::Selected)
+		{
+			return LOCTEXT("AddBindingSelectedTooltip", "Add a binding for each selected widget.");
+		}
+		else
+		{
+			return LOCTEXT("AddBindingTooltip", "Add an empty binding.");
+		}
 	}
 	else
 	{
@@ -216,6 +311,40 @@ FText SBindingsPanel::GetAddBindingToolTip() const
 	}
 }
 
+TSharedRef<SWidget> SBindingsPanel::HandleAddDefaultBindingContextMenu()
+{
+	const bool bShouldCloseWindowAfterMenuSelection = true;
+	FMenuBuilder MenuBuilder(bShouldCloseWindowAfterMenuSelection, nullptr);
+
+	MenuBuilder.AddMenuEntry(
+		LOCTEXT("AddSelectedWidget", "Add Selected Widget(s)"),
+		FText::GetEmpty(),
+		FSlateIcon(),
+		FUIAction(
+			FExecuteAction::CreateSP(this, &SBindingsPanel::HandleAddDefaultBindingButtonClick, EAddBindingMode::Selected)
+		));
+
+	MenuBuilder.AddMenuEntry(
+		LOCTEXT("AddEmptyWidget", "Add Empty binding"),
+		FText::GetEmpty(),
+		FSlateIcon(),
+		FUIAction(
+			FExecuteAction::CreateSP(this, &SBindingsPanel::HandleAddDefaultBindingButtonClick, EAddBindingMode::Empty)
+		));
+
+	return MenuBuilder.MakeWidget();
+}
+
+void SBindingsPanel::HandleAddDefaultBindingButtonClick(EAddBindingMode NewMode)
+{
+	if (AddBindingMode != NewMode)
+	{
+		AddBindingMode = NewMode;
+		SaveSettings();
+	}
+	AddDefaultBinding();
+}
+
 TSharedRef<SWidget> SBindingsPanel::CreateDrawerDockButton()
 {
 	if (bIsDrawerTab)
@@ -315,18 +444,69 @@ void SBindingsPanel::RegisterSettingsMenu()
 
 TSharedRef<SWidget> SBindingsPanel::GenerateEditViewWidget()
 {
+	FSlateIcon EmptyIcon(FAppStyle::GetAppStyleSetName(), "Icon.Empty");
+
 	FSlimHorizontalToolBarBuilder ToolbarBuilderGlobal(TSharedPtr<const FUICommandList>(), FMultiBoxCustomization::None);
+	//ToolbarBuilderGlobal.SetStyle(&FAppStyle::Get(), "SlimToolBar");
+	//{
+	//	ToolbarBuilderGlobal.BeginSection("Binding");
+	//	ToolbarBuilderGlobal.AddWidget(
+	//		SNew(SPositiveActionButton)
+	//		.OnClicked(this, &SBindingsPanel::AddDefaultBinding)
+	//		.Text(this, &SBindingsPanel::GetAddBindingText)
+	//		.ToolTipText(this, &SBindingsPanel::GetAddBindingToolTip)
+	//		.IsEnabled(this, &SBindingsPanel::CanAddBinding)
+	//	);
+
+	//	const FToolBarStyle& ToolBarStyle = ToolbarBuilderGlobal.GetStyleSet()->GetWidgetStyle<FToolBarStyle>(ToolbarBuilderGlobal.GetStyleName());
+	//	ToolbarBuilderGlobal.AddWidget(
+	//		SNew(SComboButton)
+	//		.IsEnabled(this, &SBindingsPanel::CanAddBinding)
+	//		.ContentPadding(0.f)
+	//		.ComboButtonStyle(&ToolBarStyle.SettingsComboButton)
+	//		.ButtonStyle(&ToolBarStyle.ButtonStyle)
+	//		.ToolTipText(this, &SBindingsPanel::GetAddBindingToolTip)
+	//		.ForegroundColor(FSlateColor::UseStyle())
+	//		.OnGetMenuContent(this, &SBindingsPanel::HandleAddDefaultBindingContextMenu)
+	//		.ButtonContent()
+	//		[
+	//			SNullWidget::NullWidget
+	//		]
+	//	);
+	//	ToolbarBuilderGlobal.EndSection();
+	//}
+
+	ToolbarBuilderGlobal.BeginSection("Picking");
 	{
-		ToolbarBuilderGlobal.BeginSection("Binding");
-		ToolbarBuilderGlobal.AddWidget(
-			SNew(SPositiveActionButton)
-			.OnClicked(this, &SBindingsPanel::AddDefaultBinding)
-			.Text(LOCTEXT("AddWidget", "Add Widget"))
-			.ToolTipText(this, &SBindingsPanel::GetAddBindingToolTip)
-			.IsEnabled(this, &SBindingsPanel::CanAddBinding)
+		ToolbarBuilderGlobal.AddToolBarButton(
+			FUIAction(
+				FExecuteAction::CreateSP(this, &SBindingsPanel::AddDefaultBinding),
+				FCanExecuteAction::CreateSP(this, &SBindingsPanel::CanAddBinding),
+				FGetActionCheckState()
+			),
+			NAME_None,
+			MakeAttributeSP(this, &SBindingsPanel::GetAddBindingText),
+			MakeAttributeSP(this, &SBindingsPanel::GetAddBindingToolTip),
+			FSlateIcon(FAppStyle::Get().GetStyleSetName(), FName("Icons.Plus")),
+			EUserInterfaceActionType::Button
 		);
-		ToolbarBuilderGlobal.EndSection();
+
+		ToolbarBuilderGlobal.AddComboButton(
+			FUIAction(
+				FExecuteAction(),
+				FCanExecuteAction::CreateSP(this, &SBindingsPanel::CanAddBinding),
+				FGetActionCheckState()
+			),
+			FOnGetContent::CreateSP(this, &SBindingsPanel::HandleAddDefaultBindingContextMenu),
+			FText::GetEmpty(),
+			MakeAttributeSP(this, &SBindingsPanel::GetAddBindingToolTip),
+			EmptyIcon,
+			true
+		);
+
+
 	}
+	ToolbarBuilderGlobal.EndSection();
 
 	{
 		ToolbarBuilderGlobal.AddWidget(SNew(SSpacer), NAME_None, true, HAlign_Right);
diff --git a/Engine/Plugins/Runtime/ModelViewViewModel/Source/ModelViewViewModelEditor/Private/Widgets/SMVVMViewBindingPanel.h b/Engine/Plugins/Runtime/ModelViewViewModel/Source/ModelViewViewModelEditor/Private/Widgets/SMVVMViewBindingPanel.h
--- a/Engine/Plugins/Runtime/ModelViewViewModel/Source/ModelViewViewModelEditor/Private/Widgets/SMVVMViewBindingPanel.h
+++ b/Engine/Plugins/Runtime/ModelViewViewModel/Source/ModelViewViewModelEditor/Private/Widgets/SMVVMViewBindingPanel.h
@@ -44,13 +44,25 @@ public:
 	static void RegisterSettingsMenu();
 
 private:
+	enum class EAddBindingMode
+	{
+		Selected = 0,
+		Empty = 1,
+	};
+
+	void LoadSettings();
+	void SaveSettings();
+
 	void HandleBlueprintViewChangedDelegate();
 
 	TSharedRef<SWidget> GenerateEditViewWidget();
 
-	FReply AddDefaultBinding();
+	void AddDefaultBinding();
 	bool CanAddBinding() const;
+	FText GetAddBindingText() const;
 	FText GetAddBindingToolTip() const;
+	TSharedRef<SWidget> HandleAddDefaultBindingContextMenu();
+	void HandleAddDefaultBindingButtonClick(EAddBindingMode NewMode);
 
 	TSharedRef<SWidget> GenerateSettingsMenu();
 
@@ -75,6 +87,7 @@ private:
 	TWeakObjectPtr<UMVVMWidgetBlueprintExtension_View> MVVMExtension;
 	TPimplPtr<Private::FStructDetailNotifyHook> NotifyHook;
 	FDelegateHandle BlueprintViewChangedDelegateHandle;
+	EAddBindingMode AddBindingMode = EAddBindingMode::Selected;
 	bool bIsDrawerTab;
 };
 
