MVVM: Fix order when adding a delegate and removing delegate while a delegate is executing.
#prefligth 647947ac6f48a706660482d4

diff --git a/Engine/Plugins/Runtime/ModelViewViewModel/Source/ModelViewViewModelEditor/Private/Tests/MVVMFieldValueChangedTest.cpp b/Engine/Plugins/Runtime/ModelViewViewModel/Source/ModelViewViewModelEditor/Private/Tests/MVVMFieldValueChangedTest.cpp
--- a/Engine/Plugins/Runtime/ModelViewViewModel/Source/ModelViewViewModelEditor/Private/Tests/MVVMFieldValueChangedTest.cpp
+++ b/Engine/Plugins/Runtime/ModelViewViewModel/Source/ModelViewViewModelEditor/Private/Tests/MVVMFieldValueChangedTest.cpp
@@ -123,19 +123,19 @@ bool FMVVMFieldValueChangedTest::RunTest(const FString& Parameters)
 
 		if (NumberPropertyIntChanged != InPropertyInt)
 		{
-			AddError(TEXT("Wrong number of callbacked called with PropertyInt"));
+			AddError(TEXT("Wrong number of callback called with PropertyInt"));
 		}
 		if (NumberPropertyFloatChanged != InPropertyFloat)
 		{
-			AddError(TEXT("Wrong number of callbacked called with PropertyFloat"));
+			AddError(TEXT("Wrong number of callback called with PropertyFloat"));
 		}
 		if (NumberFunctionIntChanged != InFunctionInt)
 		{
-			AddError(TEXT("Wrong number of callbacked called with FunctionInt"));
+			AddError(TEXT("Wrong number of callback called with FunctionInt"));
 		}
 		if (NumberFunctionFloatChanged != InFunctionFloat)
 		{
-			AddError(TEXT("Wrong number of callbacked called with FunctionFloat"));
+			AddError(TEXT("Wrong number of callback called with FunctionFloat"));
 		}
 	};
 
@@ -144,19 +144,19 @@ bool FMVVMFieldValueChangedTest::RunTest(const FString& Parameters)
 		FCallbackHandler& Instance = CallbackInstances[EntryIndex];
 		if (Instance.NumberPropertyIntChanged != InPropertyInt)
 		{
-			AddError(TEXT("Wrong number of callbacked called with PropertyInt"));
+			AddError(TEXT("Wrong number of callback called with PropertyInt"));
 		}
 		if (Instance.NumberPropertyFloatChanged != InPropertyFloat)
 		{
-			AddError(TEXT("Wrong number of callbacked called with PropertyFloat"));
+			AddError(TEXT("Wrong number of callback called with PropertyFloat"));
 		}
 		if (Instance.NumberFunctionIntChanged != InFunctionInt)
 		{
-			AddError(TEXT("Wrong number of callbacked called with FunctionInt"));
+			AddError(TEXT("Wrong number of callback called with FunctionInt"));
 		}
 		if (Instance.NumberFunctionFloatChanged != InFunctionFloat)
 		{
-			AddError(TEXT("Wrong number of callbacked called with FunctionFloat"));
+			AddError(TEXT("Wrong number of callback called with FunctionFloat"));
 		}
 	};
 
@@ -338,6 +338,10 @@ bool FMVVMFieldValueChangedTest::RunTest(const FString& Parameters)
 		FCallbackHandler* CallbackHandler = nullptr;
 		int32 NumberPropertyIntChanged = 0;
 		int32 NumberPropertyFloatChanged = 0;
+		int32 NumberFunctionIntChanged = 0;
+		int32 NumberFunctionFloatChanged = 0;
+		int32 NumberWasAddedWhileIteratingOther = 0;
+		int32 NumberWasAddedWhileIterating = 0;
 		void OnPropertyIntChanged(UObject* InObject, UE::FieldNotification::FFieldId InFieldId)
 		{
 			if (SourceObj && ToRemoveHandle.IsValid())
@@ -354,6 +358,72 @@ bool FMVVMFieldValueChangedTest::RunTest(const FString& Parameters)
 			}
 			++NumberPropertyFloatChanged;
 		}
+		void OnFunctionIntChanged(UObject* InObject, UE::FieldNotification::FFieldId InFieldId)
+		{
+			if (SourceObj && CallbackHandler)
+			{
+				SourceObj->AddFieldValueChangedDelegate(UMVVMFieldValueChangedTest::FFieldNotificationClassDescriptor::FunctionFloat, UMVVMFieldValueChangedTest::FFieldValueChangedDelegate::CreateRaw(this, &FRemoveCallbackHandler::OnWasAddedWhileIteratingOther));
+			}
+			++NumberFunctionIntChanged;
+		}
+		void OnFunctionFloatChanged(UObject* InObject, UE::FieldNotification::FFieldId InFieldId)
+		{
+			if (SourceObj && CallbackHandler)
+			{
+				SourceObj->AddFieldValueChangedDelegate(UMVVMFieldValueChangedTest::FFieldNotificationClassDescriptor::FunctionFloat, UMVVMFieldValueChangedTest::FFieldValueChangedDelegate::CreateRaw(this, &FRemoveCallbackHandler::OnWasAddedWhileIterating));
+			}
+			++NumberFunctionFloatChanged;
+		}
+		void OnWasAddedWhileIteratingOther(UObject* InObject, UE::FieldNotification::FFieldId InFieldId)
+		{
+			// this one should be called when OnFunctionFloatChanged is called, not when OnFunctionIntChanged is called
+			++NumberWasAddedWhileIteratingOther;
+			if (SourceObj && CallbackHandler)
+			{
+				SourceObj->AddFieldValueChangedDelegate(UMVVMFieldValueChangedTest::FFieldNotificationClassDescriptor::FunctionFloat, UMVVMFieldValueChangedTest::FFieldValueChangedDelegate::CreateRaw(this, &FRemoveCallbackHandler::OnWasAddedWhileIterating));
+			}
+		}
+		void OnWasAddedWhileIterating(UObject* InObject, UE::FieldNotification::FFieldId InFieldId)
+		{
+			// this one should be called when OnFunctionFloatChanged and OnWasAddedWhileIteratingOther are called
+			++NumberWasAddedWhileIterating;
+		}
+		void Reset()
+		{
+			NumberPropertyIntChanged = 0;
+			NumberPropertyFloatChanged = 0;
+			NumberFunctionIntChanged = 0;
+			NumberFunctionFloatChanged = 0;
+			NumberWasAddedWhileIteratingOther = 0;
+			NumberWasAddedWhileIterating = 0;
+		}
+		void TestExpectedCount(FMVVMFieldValueChangedTest* TestOwner, int32 InPropertyInt, int32 InPropertyFloat, int32 InFunctionInt, int32 InFunctionFloat, int32 InAddedWhileIteratingOther, int32 InAddedWhileIterating)
+		{
+			if (NumberPropertyIntChanged != InPropertyInt)
+			{
+				TestOwner->AddError(TEXT("Wrong number of callback called with PropertyInt while iterating"));
+			}
+			if (NumberPropertyFloatChanged != InPropertyFloat)
+			{
+				TestOwner->AddError(TEXT("Wrong number of callback called with PropertyFloat while iterating"));
+			}
+			if (NumberFunctionIntChanged != InFunctionInt)
+			{
+				TestOwner->AddError(TEXT("Wrong number of callback called with FunctionInt while iterating"));
+			}
+			if (NumberFunctionFloatChanged != InFunctionFloat)
+			{
+				TestOwner->AddError(TEXT("Wrong number of callback called with FunctionFloat while iterating"));
+			}
+			if (NumberWasAddedWhileIteratingOther != InAddedWhileIteratingOther)
+			{
+				TestOwner->AddError(TEXT("Wrong number of callback called with AddedWhileIteratingOther while iterating"));
+			}
+			if (NumberWasAddedWhileIterating != InAddedWhileIterating)
+			{
+				TestOwner->AddError(TEXT("Wrong number of callback called with AddedWhileIterating while iterating"));
+			}
+		}
 	};
 	FRemoveCallbackHandler RemoveCallbackInstance;
 
@@ -366,14 +436,29 @@ bool FMVVMFieldValueChangedTest::RunTest(const FString& Parameters)
 
 	TestBroadcastAll();
 	// Do not test the result here because the order is not deterministic.
+	//TestAllExpectedCountValues
+	RemoveCallbackInstance.TestExpectedCount(this, 1, 1, 0, 0, 0, 0);
+	ResetAllValues();
+	RemoveCallbackInstance.Reset();
+
+	TestBroadcastAll();
+	TestAllExpectedCountValues(1, 3, 4, 4);
+	TestExpectedCountValues(0, 0, 2, 2, 2);
+	TestExpectedCountValues(1, 0, 1, 1, 1);
+	TestExpectedCountValues(2, 1, 0, 1, 1);
+	RemoveCallbackInstance.TestExpectedCount(this, 1, 1, 0, 0, 0, 0);
 	ResetAllValues();
+	RemoveCallbackInstance.Reset();
 
+	SourceObj->AddFieldValueChangedDelegate(UMVVMFieldValueChangedTest::FFieldNotificationClassDescriptor::PropertyFloat, UMVVMFieldValueChangedTest::FFieldValueChangedDelegate::CreateRaw(&RemoveCallbackInstance, &FRemoveCallbackHandler::OnFunctionIntChanged));
 	TestBroadcastAll();
 	TestAllExpectedCountValues(1, 3, 4, 4);
 	TestExpectedCountValues(0, 0, 2, 2, 2);
 	TestExpectedCountValues(1, 0, 1, 1, 1);
 	TestExpectedCountValues(2, 1, 0, 1, 1);
+	RemoveCallbackInstance.TestExpectedCount(this, 1, 1, 1, 0, 1, 1);
 	ResetAllValues();
+	RemoveCallbackInstance.Reset();
 
 	return true;
 }
diff --git a/Engine/Source/Runtime/FieldNotification/Private/FieldNotificationDelegate.cpp b/Engine/Source/Runtime/FieldNotification/Private/FieldNotificationDelegate.cpp
--- a/Engine/Source/Runtime/FieldNotification/Private/FieldNotificationDelegate.cpp
+++ b/Engine/Source/Runtime/FieldNotification/Private/FieldNotificationDelegate.cpp
@@ -476,6 +476,11 @@ void FFieldMulticastDelegate::ExecuteLockOperations()
 				if (Element.Delegate.IsCompactable())
 				{
 					Delegates.RemoveAt(Index);
+					if (AddedEmplaceAt != std::numeric_limits<uint16>::max() && Index < AddedEmplaceAt)
+					{
+						check(AddedEmplaceAt > 0);
+						--AddedEmplaceAt;
+					}
 				}
 			}
 			CompactionCount = 0;
