MVVM: Support custom pin for complex conversion function.
#b editor-ui-systems

diff --git a/Engine/Plugins/Runtime/ModelViewViewModel/Source/ModelViewViewModelEditor/Private/Widgets/SMVVMFunctionParameter.cpp b/Engine/Plugins/Runtime/ModelViewViewModel/Source/ModelViewViewModelEditor/Private/Widgets/SMVVMFunctionParameter.cpp
--- a/Engine/Plugins/Runtime/ModelViewViewModel/Source/ModelViewViewModelEditor/Private/Widgets/SMVVMFunctionParameter.cpp
+++ b/Engine/Plugins/Runtime/ModelViewViewModel/Source/ModelViewViewModelEditor/Private/Widgets/SMVVMFunctionParameter.cpp
@@ -56,7 +56,7 @@ void SFunctionParameter::Construct(const FArguments& InArgs, UWidgetBlueprint* I
 	{
 		bIsBooleanPin = Pin->PinType.PinCategory == UEdGraphSchema_K2::PC_Boolean;
 		// create a new pin widget so that we can get the default value widget out of it
-		if (TSharedPtr<SGraphPin> PinWidget = FNodeFactory::CreateK2PinWidget(Pin))
+		if (TSharedPtr<SGraphPin> PinWidget = FNodeFactory::CreatePinWidget(Pin))
 		{
 			GraphPin = PinWidget;
 			ValueWidget = PinWidget->GetDefaultValueWidget();
