MVVM: Supports FieldNotify broadcast on BP defined variable.

diff --git a/Engine/Plugins/Runtime/ModelViewViewModel/Source/ModelViewViewModelBlueprint/Private/MVVMFunctionGraphHelper.cpp b/Engine/Plugins/Runtime/ModelViewViewModel/Source/ModelViewViewModelBlueprint/Private/MVVMFunctionGraphHelper.cpp
--- a/Engine/Plugins/Runtime/ModelViewViewModel/Source/ModelViewViewModelBlueprint/Private/MVVMFunctionGraphHelper.cpp
+++ b/Engine/Plugins/Runtime/ModelViewViewModel/Source/ModelViewViewModelBlueprint/Private/MVVMFunctionGraphHelper.cpp
@@ -2,24 +2,32 @@
 
 #include "MVVMFunctionGraphHelper.h"
 
+#include "EdGraph/EdGraph.h"
 #include "K2Node_CallFunction.h"
 #include "K2Node_DynamicCast.h"
 #include "K2Node_FunctionEntry.h"
 #include "K2Node_VariableGet.h"
+#include "K2Node_VariableSet.h"
 #include "KismetCompiler.h"
 #include "Kismet2/BlueprintEditorUtils.h"
 
 #include "View/MVVMView.h"
 #include "MVVMViewModelBase.h"
 
+#if WITH_EDITOR
+#include "EdGraphSchema_K2.h"
+#endif
 
 namespace UE::MVVM::FunctionGraphHelper
 {
 
-UEdGraph* CreateFunctionGraph(FKismetCompilerContext& InContext, FStringView InFunctionName, EFunctionFlags ExtraFunctionFlag, const FStringView Category, bool bIsEditable)
+UEdGraph* CreateFunctionGraph(UBlueprint* InBlueprint, FStringView InFunctionName, EFunctionFlags ExtraFunctionFlag, const FStringView Category, bool bIsEditable)
 {
-	FName UniqueFunctionName = FBlueprintEditorUtils::FindUniqueKismetName(InContext.Blueprint, InFunctionName.GetData());
-	UEdGraph* FunctionGraph = FBlueprintEditorUtils::CreateNewGraph(InContext.Blueprint, UniqueFunctionName, UEdGraph::StaticClass(), UEdGraphSchema_K2::StaticClass());
+	FName UniqueFunctionName = FBlueprintEditorUtils::FindUniqueKismetName(InBlueprint, InFunctionName.GetData());
+	UEdGraph* FunctionGraph = FBlueprintEditorUtils::CreateNewGraph(InBlueprint, UniqueFunctionName, UEdGraph::StaticClass(), UEdGraphSchema_K2::StaticClass());
+
+	// This is sometime needed. Needs more investigation.
+	//InBlueprint->FunctionGraphs.Add(FunctionGraph);
 
 	const UEdGraphSchema_K2* Schema = CastChecked<UEdGraphSchema_K2>(FunctionGraph->GetSchema());
 	Schema->MarkFunctionEntryAsEditable(FunctionGraph, bIsEditable);
@@ -59,6 +67,19 @@ UEdGraph* CreateIntermediateFunctionGraph(FKismetCompilerContext& InContext, FSt
 }
 
 
+void SetVariableNodeMember(UK2Node_Variable* InVariableNode, const FProperty* InProperty, UBlueprint* InTargetBlueprint)
+{
+	if (InProperty->GetOwnerStruct() && InTargetBlueprint->GeneratedClass && InTargetBlueprint->GeneratedClass->IsChildOf(InProperty->GetOwnerStruct()))
+	{
+		FGuid Guid = FBlueprintEditorUtils::FindMemberVariableGuidByName(InTargetBlueprint, InProperty->GetFName());
+		InVariableNode->VariableReference.SetSelfMember(InProperty->GetFName(), Guid);
+	}
+	else
+	{
+		UEdGraphSchema_K2::ConfigureVarNode(InVariableNode, InProperty->GetFName(), InProperty->GetOwnerStruct(), InTargetBlueprint);
+	}
+}
+
 bool AddFunctionArgument(UEdGraph* InFunctionGraph, TSubclassOf<UObject> InArgument, FName InArgumentName)
 {
 	for (UEdGraphNode* Node : InFunctionGraph->Nodes)
@@ -80,7 +101,7 @@ bool AddFunctionArgument(UEdGraph* InFunctionGraph, TSubclassOf<UObject> InArgum
 }
 
 
-bool AddFunctionArgument(UEdGraph* InFunctionGraph, FProperty* InArgument, FName InArgumentName)
+bool AddFunctionArgument(UEdGraph* InFunctionGraph, const FProperty* InArgument, FName InArgumentName)
 {
 	for (UEdGraphNode* Node : InFunctionGraph->Nodes)
 	{
@@ -322,4 +343,202 @@ bool GenerateViewModelFielNotifyBroadcast(FKismetCompilerContext& InContext, UEd
 	return bResult;
 }
 
+namespace Private
+{
+
+UK2Node_CallFunction* CreateFunctionNode(UEdGraph* InFunctionGraph, const UFunction* InFunction, UBlueprint* InBlueprint, FKismetCompilerContext* InContext, UK2Node_FunctionEntry* InFunctionEntry)
+{
+	if (InContext)
+	{
+		check(InFunctionEntry);
+		UK2Node_CallFunction* FunctionNode = InContext->SpawnIntermediateNode<UK2Node_CallFunction>(InFunctionEntry, InFunctionGraph);
+		FunctionNode->SetFromFunction(InFunction);
+		FunctionNode->AllocateDefaultPins();
+		return FunctionNode;
+	}
+	else
+	{
+		FGraphNodeCreator<UK2Node_CallFunction> RootCreator(*InFunctionGraph);
+		UK2Node_CallFunction* FunctionNode = RootCreator.CreateNode();
+		FunctionNode->NodePosX = 0;
+		FunctionNode->NodePosY = 0;
+		FunctionNode->SetFromFunction(InFunction);
+		RootCreator.Finalize();
+		return FunctionNode;
+	}
+}
+
+template<typename NodeType>
+NodeType* CreateVariableNode(UEdGraph* InFunctionGraph, const FProperty* InProperty, UBlueprint* InBlueprint, FKismetCompilerContext* InContext, UK2Node_FunctionEntry* InFunctionEntry)
+{
+	if (InContext)
+	{
+		check(InFunctionEntry);
+		NodeType* VariableNode = InContext->SpawnIntermediateNode<NodeType>(InFunctionEntry, InFunctionGraph);
+		SetVariableNodeMember(VariableNode, InProperty, InBlueprint);
+		VariableNode->AllocateDefaultPins();
+		return VariableNode;
+	}
+	else
+	{
+		FGraphNodeCreator<NodeType> RootCreator(*InFunctionGraph);
+		NodeType* VariableNode = RootCreator.CreateNode();
+		VariableNode->NodePosX = 0;
+		VariableNode->NodePosY = 0;
+		SetVariableNodeMember(VariableNode, InProperty, InBlueprint);
+		RootCreator.Finalize();
+		return VariableNode;
+	}
+}
+
+bool GenerateSetter(UEdGraph* InFunctionGraph, TArrayView<UE::MVVM::FMVVMConstFieldVariant> InSetterPath, UBlueprint* InBlueprint, FKismetCompilerContext* InContext)
+{
+	check(InFunctionGraph);
+
+	UK2Node_FunctionEntry* FunctionEntry = nullptr;
+	for (UEdGraphNode* Node : InFunctionGraph->Nodes)
+	{
+		FunctionEntry = Cast<UK2Node_FunctionEntry>(Node);
+		if (FunctionEntry)
+		{
+			break;
+		}
+	}
+
+	if (FunctionEntry == nullptr)
+	{
+		return false;
+	}
+
+	UEdGraphPin* FunctionInputPin = FunctionEntry->FindPinByPredicate([](UEdGraphPin* OtherPin)
+		{
+			return OtherPin->Direction == EGPD_Output
+				&& OtherPin->PinName != UEdGraphSchema_K2::PN_Then;
+		});
+
+	const UEdGraphSchema* Schema = InFunctionGraph->GetSchema();
+	if (Schema == nullptr)
+	{
+		return false;
+	}
+
+	UEdGraphPin* ThenPin = FunctionEntry->FindPin(UEdGraphSchema_K2::PN_Then, EGPD_Output);
+	UEdGraphPin* OutContextPin = nullptr;
+
+	bool bResult = true;
+	for (int32 Index = 0; Index < InSetterPath.Num(); ++Index)
+	{
+		UE::MVVM::FMVVMConstFieldVariant& Field = InSetterPath[Index];
+		const bool bLastPath = Index == InSetterPath.Num() - 1;
+
+		UEdGraphPin* NextThenPin = nullptr;
+		UEdGraphPin* NextOutContextPin = nullptr;
+
+		UEdGraphPin* SetterPin = nullptr;
+		UEdGraphPin* CurrentExecPin = nullptr;
+		UEdGraphPin* CurrentInContextPin = nullptr;
+		if (Field.IsProperty())
+		{
+			check(Field.GetProperty());
+
+			if (Field.GetProperty()->GetOwnerClass() == nullptr)
+			{
+				ensureMsgf(false, TEXT("Property point to a structure member and that is not supported right now."));
+				return false;
+			}
+
+			if (bLastPath)
+			{
+				UK2Node_VariableSet* VariableSetNode = CreateVariableNode<UK2Node_VariableSet>(InFunctionGraph, Field.GetProperty(), InBlueprint, InContext, FunctionEntry);
+
+				NextThenPin = VariableSetNode->FindPin(UEdGraphSchema_K2::PN_Then, EGPD_Output);
+				CurrentExecPin = VariableSetNode->FindPin(UEdGraphSchema_K2::PN_Execute, EGPD_Input);
+				CurrentInContextPin = VariableSetNode->FindPin(UEdGraphSchema_K2::PN_Self, EGPD_Input);
+				SetterPin = VariableSetNode->FindPin(Field.GetProperty()->GetFName(), EGPD_Input);
+			}
+			else
+			{
+				UK2Node_VariableGet* VariableGetNode = CreateVariableNode<UK2Node_VariableGet>(InFunctionGraph, Field.GetProperty(), InBlueprint, InContext, FunctionEntry);
+
+				CurrentInContextPin = VariableGetNode->FindPin(UEdGraphSchema_K2::PN_Self, EGPD_Input);
+				NextOutContextPin = VariableGetNode->FindPin(Field.GetProperty()->GetFName(), EGPD_Output);
+			}
+		}
+		else
+		{
+			ensure(Field.IsFunction());
+			check(Field.GetFunction());
+
+			UK2Node_CallFunction* CallFunctionNode = CreateFunctionNode(InFunctionGraph, Field.GetFunction(), InBlueprint, InContext, FunctionEntry);
+
+			NextThenPin = CallFunctionNode->FindPin(UEdGraphSchema_K2::PN_Then, EGPD_Output);
+			CurrentExecPin = CallFunctionNode->FindPin(UEdGraphSchema_K2::PN_Execute, EGPD_Input);
+			CurrentInContextPin = CallFunctionNode->FindPin(UEdGraphSchema_K2::PN_Self, EGPD_Input);
+
+			if (bLastPath)
+			{
+				// find the input of the setter function
+				SetterPin = CallFunctionNode->FindPinByPredicate([](UEdGraphPin* OtherPin)
+					{
+						return OtherPin->Direction == EGPD_Input
+							&& OtherPin->PinName != UEdGraphSchema_K2::PN_Execute
+							&& OtherPin->PinName != UEdGraphSchema_K2::PN_Self;
+					});
+			}
+			else
+			{
+				// find the output of the getter function
+				NextOutContextPin = CallFunctionNode->FindPinByPredicate([](UEdGraphPin* OtherPin)
+					{
+						return OtherPin->Direction == EGPD_Output
+							&& OtherPin->PinName != UEdGraphSchema_K2::PN_Then;
+					});
+			}
+		}
+
+		if (ThenPin && CurrentExecPin)
+		{
+			ensure(Schema->TryCreateConnection(ThenPin, CurrentExecPin));
+		}
+
+		if (OutContextPin)
+		{
+			if (ensure(CurrentInContextPin))
+			{
+				ensure(Schema->TryCreateConnection(OutContextPin, CurrentInContextPin));
+			}
+		}
+
+		if (bLastPath)
+		{
+			if (ensure(SetterPin && FunctionInputPin))
+			{
+				ensure(Schema->TryCreateConnection(FunctionInputPin, SetterPin));
+			}
+		}
+
+		if (NextThenPin)
+		{
+			ThenPin = NextThenPin;
+		}
+		if (NextOutContextPin)
+		{
+			OutContextPin = NextOutContextPin;
+		}
+	}
+
+	return true;
+}
+} //Private
+
+bool GenerateIntermediateSetter(FKismetCompilerContext& InContext, UEdGraph* InFunctionGraph, TArrayView<UE::MVVM::FMVVMConstFieldVariant> InSetterPath)
+{
+	return Private::GenerateSetter(InFunctionGraph, InSetterPath, InContext.Blueprint, &InContext);
+}
+
+bool GenerateSetter(UBlueprint* InBlueprint, UEdGraph* InFunctionGraph, TArrayView<UE::MVVM::FMVVMConstFieldVariant> InSetterPath)
+{
+	return Private::GenerateSetter(InFunctionGraph, InSetterPath, InBlueprint, nullptr);
+}
+
 } //namespace
diff --git a/Engine/Plugins/Runtime/ModelViewViewModel/Source/ModelViewViewModelBlueprint/Private/MVVMFunctionGraphHelper.h b/Engine/Plugins/Runtime/ModelViewViewModel/Source/ModelViewViewModelBlueprint/Private/MVVMFunctionGraphHelper.h
--- a/Engine/Plugins/Runtime/ModelViewViewModel/Source/ModelViewViewModelBlueprint/Private/MVVMFunctionGraphHelper.h
+++ b/Engine/Plugins/Runtime/ModelViewViewModel/Source/ModelViewViewModelBlueprint/Private/MVVMFunctionGraphHelper.h
@@ -3,6 +3,7 @@
 #pragma once
 
 #include "Containers/StringFwd.h"
+#include "Types/MVVMFieldVariant.h"
 
 class FName;
 class FProperty;
@@ -13,21 +14,25 @@ template <typename T> class TSubclassOf;
 class FKismetCompilerContext;
 class UBlueprint;
 class UEdGraph;
+class UK2Node_Variable;
 
 namespace UE::MVVM::FunctionGraphHelper
 {
 
 	/** Generate a function graph that can be viewed in the BP editor. */
-	MODELVIEWVIEWMODELBLUEPRINT_API UEdGraph* CreateFunctionGraph(FKismetCompilerContext& Context, const FStringView FunctionName, EFunctionFlags ExtraFunctionFlag, const FStringView Category, bool bIsEditable);
+	MODELVIEWVIEWMODELBLUEPRINT_API UEdGraph* CreateFunctionGraph(UBlueprint* Blueprint, const FStringView FunctionName, EFunctionFlags ExtraFunctionFlag, const FStringView Category, bool bIsEditable);
 
 	/** Generate a function graph that cannot be viewed in the BP editor. */
 	UEdGraph* CreateIntermediateFunctionGraph(FKismetCompilerContext& Context, const FStringView FunctionName, EFunctionFlags ExtraFunctionFlag, const FStringView Category, bool bIsEditable);
 
+	/** Set the UK2Node_Variable::VariableReference correctly base on the variable's owner. */
+	MODELVIEWVIEWMODELBLUEPRINT_API void SetVariableNodeMember(UK2Node_Variable* InVarNode, const FProperty* InProperty, UBlueprint* InTargetBlueprint);
+
 	/** Add an input argument to an existing function graph. */
 	bool AddFunctionArgument(UEdGraph* FunctionGraph, TSubclassOf<UObject> Argument, FName ArgumentName);
 
 	/** Add an input argument to an existing function graph. */
-	bool AddFunctionArgument(UEdGraph* FunctionGraph, FProperty* Argument, FName ArgumentName);
+	bool AddFunctionArgument(UEdGraph* FunctionGraph, const FProperty* Argument, FName ArgumentName);
 
 	/** Generate the nodes needed for to call the MVVMView::SetViewModel function. */
 	bool GenerateViewModelSetter(FKismetCompilerContext& Context, UEdGraph* FunctionGraph, FName ViewModelName);
@@ -38,4 +43,9 @@ namespace UE::MVVM::FunctionGraphHelper
 	/** Generate the nodes needed to broadcast field notify function. */
 	bool GenerateViewModelFielNotifyBroadcast(FKismetCompilerContext& Context, UEdGraph* FunctionGraph, FProperty* Property);
 
+	/** Generate the nodes needed for to set a variable/function value. */
+	bool GenerateSetter(UBlueprint* Blueprint, UEdGraph* FunctionGraph, TArrayView<UE::MVVM::FMVVMConstFieldVariant> SetterPath);
+	/** Generate the nodes needed to set a variable/function value. To use while while compiling the Blueprint. */
+	bool GenerateIntermediateSetter(FKismetCompilerContext& Context, UEdGraph* FunctionGraph, TArrayView<UE::MVVM::FMVVMConstFieldVariant> SetterPath);
+
 } //namespace
diff --git a/Engine/Plugins/Runtime/ModelViewViewModel/Source/ModelViewViewModelBlueprint/Private/MVVMViewBlueprintCompiler.cpp b/Engine/Plugins/Runtime/ModelViewViewModel/Source/ModelViewViewModelBlueprint/Private/MVVMViewBlueprintCompiler.cpp
--- a/Engine/Plugins/Runtime/ModelViewViewModel/Source/ModelViewViewModelBlueprint/Private/MVVMViewBlueprintCompiler.cpp
+++ b/Engine/Plugins/Runtime/ModelViewViewModel/Source/ModelViewViewModelBlueprint/Private/MVVMViewBlueprintCompiler.cpp
@@ -28,6 +28,23 @@ FAutoConsoleVariable CVarLogViewCompliedResult(
 	TEXT("After the view is compiled log the compiled bindings and sources.")
 );
 
+#if UE_BUILD_DEVELOPMENT || UE_BUILD_DEBUG
+FAutoConsoleCommand CVarTestGenerateSetter(
+	TEXT("MVVM.TestGenerateSetter"),
+	TEXT("Generate a setter function base on the input string."),
+	FConsoleCommandWithArgsDelegate::CreateLambda([](const TArray<FString>& Args)
+	{
+		if (Args.Num() < 3)
+		{
+			return;
+		}
+		FMVVMViewBlueprintCompiler::TestGenerateSetter(Args[0], Args[1], Args[2]);
+	})
+);
+#endif
+
+static const FText CouldNotCreateSourceFieldPathFormat = LOCTEXT("CouldNotCreateSourceFieldPath", "Couldn't create the source field path '{0}'. {1}");
+
 FString PropertyPathToString(const UMVVMBlueprintView* BlueprintView, const FMVVMBlueprintPropertyPath& PropertyPath)
 {
 	if (PropertyPath.IsEmpty())
@@ -183,6 +200,9 @@ void FMVVMViewBlueprintCompiler::CreateFunctions(UMVVMBlueprintView* BlueprintVi
 	{
 		return;
 	}
+	
+	// Build the list of BP destinations
+	CreateBindingDestinationContexts(BlueprintView);
 
 	if (!GetDefault<UMVVMDeveloperProjectSettings>()->bAllowGeneratedViewModelSetter)
 	{
@@ -615,6 +635,134 @@ void FMVVMViewBlueprintCompiler::CreateFunctionsDeclaration(const FWidgetBluepri
 }
 
 
+void FMVVMViewBlueprintCompiler::CreateBindingDestinationContexts(UMVVMBlueprintView* BlueprintView)
+{
+	if (!bAreSourcesCreatorValid || !bAreSourceContextsValid || !bIsBindingsValid)
+	{
+		return;
+	}
+
+	BindingDestinationContexts.Empty();
+
+	// Use the Skeleton class. The class bind and not all functions are generated yet
+	UWidgetBlueprintGeneratedClass* NewSkeletonClass = Cast<UWidgetBlueprintGeneratedClass>(WidgetBlueprintCompilerContext.Blueprint->SkeletonGeneratedClass);
+	if (NewSkeletonClass == nullptr)
+	{
+		WidgetBlueprintCompilerContext.MessageLog.Error(*LOCTEXT("InvalidNewClass", "Internal error. The skeleton class is not valid.").ToString());
+		return;
+	}
+
+	const int32 NumBindings = BlueprintView->GetNumBindings();
+	for (int32 Index = 0; Index < NumBindings; ++Index)
+	{
+		FMVVMBlueprintViewBinding* BindingPtr = BlueprintView->GetBindingAt(Index);
+		if (BindingPtr == nullptr)
+		{
+			WidgetBlueprintCompilerContext.MessageLog.Error(*FText::Format(LOCTEXT("InvalidBindingIndex", "Internal error. Invalid binding index given."), Index).ToString());
+			bIsBindingsValid = false;
+			continue;
+		}
+
+		FMVVMBlueprintViewBinding& Binding = *BindingPtr;
+		if (!Binding.bCompile)
+		{
+			continue;
+		}
+
+		auto AddDestination = [Index, BindingPtr, BlueprintView, NewSkeletonClass, this](bool bIsForwardBinding)
+		{
+			FBindingDestinationContext& NewDestinationContext = BindingDestinationContexts.AddDefaulted_GetRef();
+			NewDestinationContext.BindingIndex = Index;
+			NewDestinationContext.bIsForwardBinding = bIsForwardBinding;
+
+			const FMVVMBlueprintPropertyPath& DestinationPath = bIsForwardBinding ? BindingPtr->DestinationPath : BindingPtr->SourcePath;
+			TArray<UE::MVVM::FMVVMConstFieldVariant> SkeletalSetterPath = CreateBindingDestinationPath(BlueprintView, NewSkeletonClass, DestinationPath);
+			if (!IsPropertyPathValid(SkeletalSetterPath))
+			{
+				AddMessageForBinding(*BindingPtr, BlueprintView, FText::Format(LOCTEXT("PropertyPathIsInvalid", "The property path '{0}' is invalid."),
+					PropertyPathToText(BlueprintView, DestinationPath)),
+					EBindingMessageType::Error
+				);
+				bIsBindingsValid = false;
+				return;
+			}
+			
+			// Generate the path with property converted to BP function
+			TValueOrError<TArray<FMVVMConstFieldVariant>, FText> SkeletalGeneratedField = FieldPathHelper::GenerateFieldPathList(SkeletalSetterPath, false);
+			if (SkeletalGeneratedField.HasError())
+			{
+				AddMessageForBinding(*BindingPtr, BlueprintView, FText::Format(Private::CouldNotCreateSourceFieldPathFormat
+					, ::UE::MVVM::FieldPathHelper::ToText(SkeletalSetterPath)
+					, SkeletalGeneratedField.GetError())
+					, EBindingMessageType::Error
+				);
+				bIsBindingsValid = false;
+				return;
+			}
+
+			NewDestinationContext.bCanBeSetInNative = CanBeSetInNative(SkeletalGeneratedField.GetValue());
+
+			// If the destination can't be done cpp, we need to generate a BP function to set the value.
+			if (!NewDestinationContext.bCanBeSetInNative)
+			{
+				const FProperty* SetterProperty = nullptr;
+				if (SkeletalGeneratedField.GetValue().Num() > 0 && ensure(SkeletalGeneratedField.GetValue().Last().IsProperty()))
+				{
+					SetterProperty = SkeletalGeneratedField.GetValue().Last().GetProperty();
+				}
+
+				if (SetterProperty == nullptr)
+				{
+					AddMessageForBinding(*BindingPtr, BlueprintView, FText::Format(LOCTEXT("CantGetSetter", "Internal Error. The setter function was not created. {0}"),
+						PropertyPathToText(BlueprintView, DestinationPath)),
+						EBindingMessageType::Error
+					);
+					bIsBindingsValid = false;
+					return;
+				}
+
+				// create a setter function to be called from native. For now we follow the convention of Setter(Conversion(Getter))
+				UEdGraph* GeneratedSetterGraph = UE::MVVM::FunctionGraphHelper::CreateIntermediateFunctionGraph(WidgetBlueprintCompilerContext, FString::Printf(TEXT("__Setter_%s"), *SetterProperty->GetName()), EFunctionFlags::FUNC_None, TEXT("AutogeneratedSetter"), false);
+				if (GeneratedSetterGraph == nullptr)
+				{
+					AddMessageForBinding(*BindingPtr, BlueprintView, FText::Format(LOCTEXT("CantCreateSetter", "Internal Error. The setter function was not created. {0}"),
+						PropertyPathToText(BlueprintView, DestinationPath)),
+						EBindingMessageType::Error
+					);
+					bIsBindingsValid = false;
+					return;
+				}
+
+				UE::MVVM::FunctionGraphHelper::AddFunctionArgument(GeneratedSetterGraph, SetterProperty, "NewValue");
+
+				// set GeneratedSetterFunction. Use the SkeletalSetterPath here to use the setter will be generated when the function is generated.
+				if (!UE::MVVM::FunctionGraphHelper::GenerateIntermediateSetter(WidgetBlueprintCompilerContext, GeneratedSetterGraph, SkeletalSetterPath))
+				{
+					AddMessageForBinding(*BindingPtr, BlueprintView, FText::Format(LOCTEXT("CantGeneratedSetter", "Internal Error. The setter function was not generated. {0}"),
+						PropertyPathToText(BlueprintView, DestinationPath)),
+						EBindingMessageType::Error
+					);
+					bIsBindingsValid = false;
+					return;
+				}
+
+				// the new path can only be set later, once the function is compiled.
+				NewDestinationContext.GeneratedFunctionName = GeneratedSetterGraph->GetFName();
+			}
+		};
+
+		if (IsForwardBinding(Binding.BindingType))
+		{
+			AddDestination(true);
+		}
+		if (IsBackwardBinding(Binding.BindingType))
+		{
+			AddDestination(false);
+		}
+	}
+}
+
+
 bool FMVVMViewBlueprintCompiler::PreCompile(UWidgetBlueprintGeneratedClass* Class, UMVVMBlueprintView* BlueprintView)
 {
 	if (!bAreSourcesCreatorValid || !bAreSourceContextsValid || !bIsBindingsValid)
@@ -1151,10 +1299,42 @@ bool FMVVMViewBlueprintCompiler::PreCompileBindings(UWidgetBlueprintGeneratedCla
 
 		TArray<UE::MVVM::FMVVMConstFieldVariant> SetterPath;
 		{
-			const FMVVMBlueprintPropertyPath& DestinationPath = BindingSourceContext.bIsForwardBinding ? Binding.DestinationPath : Binding.SourcePath;
-			SetterPath = CreateBindingDestinationPath(BlueprintView, Class, DestinationPath);
+			const FBindingDestinationContext* DestinationContext = BindingDestinationContexts.FindByPredicate([BindingIndex = BindingSourceContext.BindingIndex, bIsForwardBinding = BindingSourceContext.bIsForwardBinding](const FBindingDestinationContext& Other)
+				{
+					return Other.BindingIndex == BindingIndex && Other.bIsForwardBinding == bIsForwardBinding;
+				});
+
+			if (DestinationContext == nullptr)
+			{
+				AddMessageForBinding(Binding, BlueprintView, LOCTEXT("CouldNotFindDestination", "Could not find the pre compiled destination."), EBindingMessageType::Error);
+				bIsBindingsValid = false;
+				continue;
+			}
+
+			if (!DestinationContext->bCanBeSetInNative)
+			{
+				check(!DestinationContext->GeneratedFunctionName.IsNone());
+				UFunction* FoundFunction = Class->FindFunctionByName(DestinationContext->GeneratedFunctionName);
+				if (FoundFunction == nullptr)
+				{
+					AddMessageForBinding(Binding, BlueprintView, LOCTEXT("CouldNotFindDestinationBPFunction", "Could not find the generated destination function."), EBindingMessageType::Error);
+					bIsBindingsValid = false;
+					continue;
+				}
+				FMVVMBlueprintPropertyPath DestinationPPropertyPath;
+				DestinationPPropertyPath.SetWidgetName(Class->ClassGeneratedBy->GetFName());
+				DestinationPPropertyPath.SetBasePropertyPath(UE::MVVM::FMVVMConstFieldVariant(FoundFunction));
+				SetterPath = CreateBindingDestinationPath(BlueprintView, Class, DestinationPPropertyPath);
+			}
+			else
+			{
+				const FMVVMBlueprintPropertyPath& DestinationPath = BindingSourceContext.bIsForwardBinding ? Binding.DestinationPath : Binding.SourcePath;
+				SetterPath = CreateBindingDestinationPath(BlueprintView, Class, DestinationPath);
+			}
+
 			if (!IsPropertyPathValid(SetterPath))
 			{
+				const FMVVMBlueprintPropertyPath& DestinationPath = BindingSourceContext.bIsForwardBinding ? Binding.DestinationPath : Binding.SourcePath;
 				AddMessageForBinding(Binding, BlueprintView, FText::Format(LOCTEXT("PropertyPathIsInvalid", "The property path '{0}' is invalid."),
 					PropertyPathToText(BlueprintView, DestinationPath)),
 					EBindingMessageType::Error
@@ -1394,19 +1574,17 @@ TValueOrError<FMVVMViewBlueprintCompiler::FCompiledBinding, FText> FMVVMViewBlue
 		// Sanity check
 		if (bIsComplexBinding && !BindingHelper::IsValidForComplexRuntimeConversion(ConversionFunction))
 		{
-			return MakeError(LOCTEXT("ConversionFunctionIsNotComplex", "Internal Error. The complex conversion function does not respect the prerequist."));
+			return MakeError(LOCTEXT("ConversionFunctionIsNotComplex", "Internal Error. The complex conversion function does not respect the prerequisite."));
 		}
 	}
 
 	if (!Result.bIsConversionFunctionComplex)
 	{
-		static const FText CouldNotCreateSourceFieldPathFormat = LOCTEXT("CouldNotCreateSourceFieldPath", "Couldn't create the source field path '{0}'. {1}");
-
 		// Generate a path to read the value at runtime
 		TValueOrError<TArray<FMVVMConstFieldVariant>, FText> GeneratedField = FieldPathHelper::GenerateFieldPathList(GetterFields, true);
 		if (GeneratedField.HasError())
 		{
-			return MakeError(FText::Format(CouldNotCreateSourceFieldPathFormat
+			return MakeError(FText::Format(Private::CouldNotCreateSourceFieldPathFormat
 				, ::UE::MVVM::FieldPathHelper::ToText(GetterFields)
 				, GeneratedField.GetError()));
 		}
@@ -1414,7 +1592,7 @@ TValueOrError<FMVVMViewBlueprintCompiler::FCompiledBinding, FText> FMVVMViewBlue
 		TValueOrError<FCompiledBindingLibraryCompiler::FFieldPathHandle, FText> FieldPathResult = BindingLibraryCompiler.AddFieldPath(GeneratedField.GetValue(), true);
 		if (FieldPathResult.HasError())
 		{
-			return MakeError(FText::Format(CouldNotCreateSourceFieldPathFormat
+			return MakeError(FText::Format(Private::CouldNotCreateSourceFieldPathFormat
 				, ::UE::MVVM::FieldPathHelper::ToText(GetterFields)
 				, FieldPathResult.GetError()));
 		}
@@ -1702,7 +1880,7 @@ TArray<FMVVMConstFieldVariant> FMVVMViewBlueprintCompiler::CreateBindingDestinat
 	return TArray<FMVVMConstFieldVariant>();
 }
 
-TArray<FMVVMConstFieldVariant> FMVVMViewBlueprintCompiler::CreatePropertyPath(const UClass* Class, FName PropertyName, TArray<FMVVMConstFieldVariant> Properties) const
+TArray<FMVVMConstFieldVariant> FMVVMViewBlueprintCompiler::CreatePropertyPath(const UClass* Class, FName PropertyName, TArray<FMVVMConstFieldVariant> Properties)
 {
 	if (PropertyName.IsNone())
 	{
@@ -1716,7 +1894,7 @@ TArray<FMVVMConstFieldVariant> FMVVMViewBlueprintCompiler::CreatePropertyPath(co
 }
 
 
-bool FMVVMViewBlueprintCompiler::IsPropertyPathValid(TArrayView<const FMVVMConstFieldVariant> PropertyPath) const
+bool FMVVMViewBlueprintCompiler::IsPropertyPathValid(TArrayView<const FMVVMConstFieldVariant> PropertyPath)
 {
 	for (const FMVVMConstFieldVariant& Field : PropertyPath)
 	{
@@ -1736,6 +1914,89 @@ bool FMVVMViewBlueprintCompiler::IsPropertyPathValid(TArrayView<const FMVVMConst
 	return true;
 }
 
+bool FMVVMViewBlueprintCompiler::CanBeSetInNative(TArrayView<const FMVVMConstFieldVariant> PropertyPath)
+{
+	check(IsPropertyPathValid(PropertyPath));
+	
+	for (int32 Index = PropertyPath.Num() - 1; Index >= 0; --Index)
+	{
+		const FMVVMConstFieldVariant& Variant = PropertyPath[Index];
+		// Stop the algo if the path is already a function.
+		if (Variant.IsFunction())
+		{
+			return true;
+		}
+		// If the BP is defined in BP and has Net flags or FieldNotify flag, then the VaraibleSet K2Node need to be used to generate the proper byte-code.
+		if (Variant.IsProperty())
+		{
+			// If it's an object then the path before the object doesn't matter.
+			if (const FObjectPropertyBase* PropertyBase = CastField<const FObjectPropertyBase>(Variant.GetProperty()))
+			{
+				bool bLastPath = Index >= PropertyPath.Num() - 1;
+				if (!bLastPath)
+				{
+					return true;
+				}
+			}
+
+			if (Cast<UBlueprintGeneratedClass>(Variant.GetProperty()->GetOwnerStruct()))
+			{
+				if (Variant.GetProperty()->HasMetaData(FName("FieldNotify")) || Variant.GetProperty()->HasAnyPropertyFlags(CPF_Net))
+				{
+					return false;
+				}
+			}
+		}
+	}
+	return true;
+}
+
+void FMVVMViewBlueprintCompiler::TestGenerateSetter(FStringView ObjectName, FStringView FieldPath, FStringView FunctionName)
+{
+#if UE_BUILD_DEVELOPMENT || UE_BUILD_DEBUG
+	UWidgetBlueprint* WidgetBlueprint = nullptr;
+	{
+		UObject* FoundObject = FindObject<UObject>(nullptr, ObjectName.GetData(), false);
+		WidgetBlueprint = Cast<UWidgetBlueprint>(FoundObject);
+	}
+
+	if (WidgetBlueprint == nullptr)
+	{
+		return;
+	}
+
+	UWidgetBlueprintGeneratedClass* NewSkeletonClass = Cast<UWidgetBlueprintGeneratedClass>(WidgetBlueprint->SkeletonGeneratedClass);
+	TValueOrError<TArray<FMVVMConstFieldVariant>, FText> SkeletalSetterPathResult = FieldPathHelper::GenerateFieldPathList(NewSkeletonClass, FieldPath, false);
+	if (SkeletalSetterPathResult.HasError())
+	{
+		return;
+	}
+	if (!IsPropertyPathValid(SkeletalSetterPathResult.GetValue()))
+	{
+		return;
+	}
+
+	UEdGraph* GeneratedSetterGraph = UE::MVVM::FunctionGraphHelper::CreateFunctionGraph(WidgetBlueprint, FunctionName, EFunctionFlags::FUNC_None, TEXT(""), false);
+	if (GeneratedSetterGraph == nullptr)
+	{
+		return;
+	}
+
+	const FProperty* SetterProperty = nullptr;
+	if (SkeletalSetterPathResult.GetValue().Num() > 0 && ensure(SkeletalSetterPathResult.GetValue().Last().IsProperty()))
+	{
+		SetterProperty = SkeletalSetterPathResult.GetValue().Last().GetProperty();
+	}
+	if (SetterProperty == nullptr)
+	{
+		return;
+	}
+
+	UE::MVVM::FunctionGraphHelper::AddFunctionArgument(GeneratedSetterGraph, SetterProperty, "NewValue");
+	UE::MVVM::FunctionGraphHelper::GenerateSetter(WidgetBlueprint, GeneratedSetterGraph, SkeletalSetterPathResult.GetValue());
+#endif
+}
+
 } //namespace
 
 #undef LOCTEXT_NAMESPACE
diff --git a/Engine/Plugins/Runtime/ModelViewViewModel/Source/ModelViewViewModelBlueprint/Private/MVVMViewBlueprintCompiler.h b/Engine/Plugins/Runtime/ModelViewViewModel/Source/ModelViewViewModelBlueprint/Private/MVVMViewBlueprintCompiler.h
--- a/Engine/Plugins/Runtime/ModelViewViewModel/Source/ModelViewViewModelBlueprint/Private/MVVMViewBlueprintCompiler.h
+++ b/Engine/Plugins/Runtime/ModelViewViewModel/Source/ModelViewViewModelBlueprint/Private/MVVMViewBlueprintCompiler.h
@@ -32,6 +32,7 @@ private:
 	struct FCompiledBinding;
 	struct FCompilerBinding;
 	struct FBindingSourceContext;
+	struct FBindingDestinationContext;
 
 public:
 	FMVVMViewBlueprintCompiler(FWidgetBlueprintCompilerContext& InCreationContext)
@@ -58,10 +59,14 @@ public:
 	/** Compile the library and fill the view and viewclass */
 	bool Compile(UWidgetBlueprintGeneratedClass* Class, UMVVMBlueprintView* BlueprintView, UMVVMViewClass* ViewExtension);
 
+
+	static void TestGenerateSetter(FStringView ObjectName, FStringView FieldPath, FStringView FunctionName);
+
 private:
 	void CreateWidgetMap(const FWidgetBlueprintCompilerContext::FCreateVariableContext& Context, UMVVMBlueprintView* BlueprintView);
 	void CreateSourceLists(const FWidgetBlueprintCompilerContext::FCreateVariableContext& Context, UMVVMBlueprintView* BlueprintView);
 	void CreateFunctionsDeclaration(const FWidgetBlueprintCompilerContext::FCreateVariableContext& Context, UMVVMBlueprintView* BlueprintView);
+	void CreateBindingDestinationContexts(UMVVMBlueprintView* BlueprintView);
 
 	bool PreCompileBindingSources(UWidgetBlueprintGeneratedClass* Class, UMVVMBlueprintView* BlueprintView);
 	bool CompileBindingSources(const FCompiledBindingLibraryCompiler::FCompileResult& CompileResult, UWidgetBlueprintGeneratedClass* Class, UMVVMBlueprintView* BlueprintView, UMVVMViewClass* ViewExtension);
@@ -81,9 +86,9 @@ private:
 	TArray<FMVVMConstFieldVariant> CreateBindingDestinationPath(const UMVVMBlueprintView* BlueprintView, const UWidgetBlueprintGeneratedClass* Class, const FMVVMBlueprintPropertyPath& PropertyPath) const;
 	TValueOrError<FCompiledBinding, FText> CreateCompiledBinding(const UWidgetBlueprintGeneratedClass* Class, TArrayView<const UE::MVVM::FMVVMConstFieldVariant> GetterFields, TArrayView<const UE::MVVM::FMVVMConstFieldVariant> SetterFields, const UFunction* ConversionFunction, bool bIsComplexBinding);
 
-
-	TArray<FMVVMConstFieldVariant> CreatePropertyPath(const UClass* Class, FName PropertyName, TArray<FMVVMConstFieldVariant> Properties) const;
-	bool IsPropertyPathValid(TArrayView<const FMVVMConstFieldVariant> PropertyPath) const;
+	static TArray<FMVVMConstFieldVariant> CreatePropertyPath(const UClass* Class, FName PropertyName, TArray<FMVVMConstFieldVariant> Properties);
+	static bool IsPropertyPathValid(TArrayView<const FMVVMConstFieldVariant> PropertyPath);
+	static bool CanBeSetInNative(TArrayView<const FMVVMConstFieldVariant> PropertyPath);
 
 private:
 	/**
@@ -143,6 +148,7 @@ private:
 		FCompiledBindingLibraryCompiler::FFieldPathHandle SourceRead;
 		FCompiledBindingLibraryCompiler::FFieldPathHandle DestinationWrite;
 		FCompiledBindingLibraryCompiler::FFieldPathHandle ConversionFunction;
+		FCompiledBindingLibraryCompiler::FFieldPathHandle ExecutionFunction;
 
 		bool bIsConversionFunctionComplex = false;
 	};
@@ -200,6 +206,21 @@ private:
 	};
 	TArray<FBindingSourceContext> BindingSourceContexts;
 
+	/**
+	 * Destination path of a binding.
+	 * The info needs to be validate before we generate the functions list.
+	 */
+	struct FBindingDestinationContext
+	{
+		int32 BindingIndex = INDEX_NONE;
+		bool bIsForwardBinding = false;
+
+		bool bCanBeSetInNative = true;
+		FName GeneratedFunctionName;
+
+	};
+	TArray<FBindingDestinationContext> BindingDestinationContexts;
+
 	TMap<FName, UWidget*> WidgetNameToWidgetPointerMap;
 	FWidgetBlueprintCompilerContext& WidgetBlueprintCompilerContext;
 	FCompiledBindingLibraryCompiler BindingLibraryCompiler;
