MVVM: Add in menu extension to viewmodel editor. The Bindings panel is now "MVVM.ViewBindings.Toolbar" and the viewmodel panel is "MVVM.Viewmodels.Toolbar".
#prefligth 647f2c2e6158e3d3209defb5

diff --git a/Engine/Plugins/Runtime/ModelViewViewModel/Source/ModelViewViewModelEditor/Private/ModelViewViewModelEditorModule.cpp b/Engine/Plugins/Runtime/ModelViewViewModel/Source/ModelViewViewModelEditor/Private/ModelViewViewModelEditorModule.cpp
--- a/Engine/Plugins/Runtime/ModelViewViewModel/Source/ModelViewViewModelEditor/Private/ModelViewViewModelEditorModule.cpp
+++ b/Engine/Plugins/Runtime/ModelViewViewModel/Source/ModelViewViewModelEditor/Private/ModelViewViewModelEditorModule.cpp
@@ -20,9 +20,12 @@
 #include "Styling/MVVMEditorStyle.h"
 #include "Tabs/MVVMBindingSummoner.h"
 #include "Tabs/MVVMViewModelSummoner.h"
+#include "ToolMenus.h"
 #include "UMGEditorModule.h"
 #include "WidgetBlueprintEditor.h"
 #include "WidgetDrawerConfig.h"
+#include "Widgets/SMVVMViewBindingPanel.h"
+#include "Widgets/SMVVMViewModelPanel.h"
 
 #define LOCTEXT_NAMESPACE "ModelViewViewModelModule"
 
@@ -58,11 +61,15 @@ void FModelViewViewModelEditorModule::StartupModule()
 	FMVVMEditorCommands::Register();
 	FWidgetBlueprintDelegates::GetAssetTags.AddRaw(this, &FModelViewViewModelEditorModule::HandleWidgetBlueprintAssetTags);
 	FWidgetBlueprintGeneratedClassDelegates::GetAssetTags.AddRaw(this, &FModelViewViewModelEditorModule::HandleClassBlueprintAssetTags);
+
+	UToolMenus::RegisterStartupCallback(FSimpleMulticastDelegate::FDelegate::CreateRaw(this, &FModelViewViewModelEditorModule::HandleRegisterMenus));
 }
 
 
 void FModelViewViewModelEditorModule::ShutdownModule()
 {
+	UnregisterMenus();
+
 	FWidgetBlueprintGeneratedClassDelegates::GetAssetTags.RemoveAll(this);
 	FWidgetBlueprintDelegates::GetAssetTags.RemoveAll(this);
 	if (FMessageLogModule* MessageLogModule = FModuleManager::GetModulePtr<FMessageLogModule>("MessageLog"))
@@ -220,6 +227,23 @@ void FModelViewViewModelEditorModule::HandleClassBlueprintAssetTags(const UWidge
 	}
 }
 
+
+void FModelViewViewModelEditorModule::HandleRegisterMenus()
+{
+	// Allow cleanup when module unloads
+	FToolMenuOwnerScoped OwnerScoped(this);
+
+	UE::MVVM::SMVVMViewModelPanel::RegisterMenu();
+	UE::MVVM::SBindingsPanel::RegisterSettingsMenu();
+}
+
+
+void FModelViewViewModelEditorModule::UnregisterMenus()
+{
+	UToolMenus::UnRegisterStartupCallback(this);
+	UToolMenus::UnregisterOwner(this);
+}
+
 IMPLEMENT_MODULE(FModelViewViewModelEditorModule, ModelViewViewModelEditor);
 
 #undef LOCTEXT_NAMESPACE
diff --git a/Engine/Plugins/Runtime/ModelViewViewModel/Source/ModelViewViewModelEditor/Private/ModelViewViewModelEditorModule.h b/Engine/Plugins/Runtime/ModelViewViewModel/Source/ModelViewViewModelEditor/Private/ModelViewViewModelEditorModule.h
--- a/Engine/Plugins/Runtime/ModelViewViewModel/Source/ModelViewViewModelEditor/Private/ModelViewViewModelEditorModule.h
+++ b/Engine/Plugins/Runtime/ModelViewViewModel/Source/ModelViewViewModelEditor/Private/ModelViewViewModelEditorModule.h
@@ -38,6 +38,8 @@ private:
 	void HandleActivateMode(FWidgetBlueprintApplicationMode& InDesignerMode);
 	void HandleWidgetBlueprintAssetTags(const UWidgetBlueprint* Widget, TArray<UObject::FAssetRegistryTag>& OutTags);
 	void HandleClassBlueprintAssetTags(const UWidgetBlueprintGeneratedClass* GeneratedClass, TArray<UObject::FAssetRegistryTag>& OutTags);
+	void HandleRegisterMenus();
+	void UnregisterMenus();
 
 private:
 	TSharedPtr<FMVVMPropertyBindingExtension> PropertyBindingExtension;
diff --git a/Engine/Plugins/Runtime/ModelViewViewModel/Source/ModelViewViewModelEditor/Private/Tabs/MVVMBindingSummoner.cpp b/Engine/Plugins/Runtime/ModelViewViewModel/Source/ModelViewViewModelEditor/Private/Tabs/MVVMBindingSummoner.cpp
--- a/Engine/Plugins/Runtime/ModelViewViewModel/Source/ModelViewViewModelEditor/Private/Tabs/MVVMBindingSummoner.cpp
+++ b/Engine/Plugins/Runtime/ModelViewViewModel/Source/ModelViewViewModelEditor/Private/Tabs/MVVMBindingSummoner.cpp
@@ -17,13 +17,13 @@ FMVVMBindingSummoner::FMVVMBindingSummoner(TSharedPtr<FWidgetBlueprintEditor> Bl
 	, WeakWidgetBlueprintEditor(BlueprintEditor)
 	, bIsDrawerTab(bInIsDrawerTab)
 {
-	TabLabel = LOCTEXT("ViewBinding_ViewMenu_Label", "View Binding");
+	TabLabel = LOCTEXT("ViewBinding_ViewMenu_Label", "View Bindings");
 	TabIcon = FSlateIcon(FMVVMEditorStyle::Get().GetStyleSetName(), "BlueprintView.TabIcon");
 
 	bIsSingleton = true;
 
-	ViewMenuDescription = LOCTEXT("ViewBinding_ViewMenu_Desc", "View Binding");
-	ViewMenuTooltip = LOCTEXT("ViewBinding_ViewMenu_ToolTip", "Show the View Binding tab");
+	ViewMenuDescription = LOCTEXT("ViewBinding_ViewMenu_Desc", "View Bindings");
+	ViewMenuTooltip = LOCTEXT("ViewBinding_ViewMenu_ToolTip", "Show the View Bindings tab");
 }
 
 TSharedRef<SWidget> FMVVMBindingSummoner::CreateTabBody(const FWorkflowTabSpawnInfo& Info) const
diff --git a/Engine/Plugins/Runtime/ModelViewViewModel/Source/ModelViewViewModelEditor/Private/Widgets/SMVVMViewBindingPanel.cpp b/Engine/Plugins/Runtime/ModelViewViewModel/Source/ModelViewViewModelEditor/Private/Widgets/SMVVMViewBindingPanel.cpp
--- a/Engine/Plugins/Runtime/ModelViewViewModel/Source/ModelViewViewModelEditor/Private/Widgets/SMVVMViewBindingPanel.cpp
+++ b/Engine/Plugins/Runtime/ModelViewViewModel/Source/ModelViewViewModelEditor/Private/Widgets/SMVVMViewBindingPanel.cpp
@@ -432,13 +432,13 @@ TSharedRef<SWidget> SBindingsPanel::GenerateSettingsMenu()
 	WidgetBlueprintMenuContext->WidgetBlueprintEditor = WeakBlueprintEditor;
 
 	FToolMenuContext MenuContext(WidgetBlueprintMenuContext);
-	return UToolMenus::Get()->GenerateWidget("MVVMEditor.Panel.Toolbar.Settings", MenuContext);
+	return UToolMenus::Get()->GenerateWidget("MVVM.ViewBindings.Toolbar", MenuContext);
 }
 
 
 void SBindingsPanel::RegisterSettingsMenu()
 {
-	UToolMenu* Menu = UToolMenus::Get()->RegisterMenu("MVVMEditor.Panel.Toolbar.Settings");
+	UToolMenu* Menu = UToolMenus::Get()->RegisterMenu("MVVM.ViewBindings.Toolbar");
 }
 
 
@@ -447,35 +447,6 @@ TSharedRef<SWidget> SBindingsPanel::GenerateEditViewWidget()
 	FSlateIcon EmptyIcon(FAppStyle::GetAppStyleSetName(), "Icon.Empty");
 
 	FSlimHorizontalToolBarBuilder ToolbarBuilderGlobal(TSharedPtr<const FUICommandList>(), FMultiBoxCustomization::None);
-	//ToolbarBuilderGlobal.SetStyle(&FAppStyle::Get(), "SlimToolBar");
-	//{
-	//	ToolbarBuilderGlobal.BeginSection("Binding");
-	//	ToolbarBuilderGlobal.AddWidget(
-	//		SNew(SPositiveActionButton)
-	//		.OnClicked(this, &SBindingsPanel::AddDefaultBinding)
-	//		.Text(this, &SBindingsPanel::GetAddBindingText)
-	//		.ToolTipText(this, &SBindingsPanel::GetAddBindingToolTip)
-	//		.IsEnabled(this, &SBindingsPanel::CanAddBinding)
-	//	);
-
-	//	const FToolBarStyle& ToolBarStyle = ToolbarBuilderGlobal.GetStyleSet()->GetWidgetStyle<FToolBarStyle>(ToolbarBuilderGlobal.GetStyleName());
-	//	ToolbarBuilderGlobal.AddWidget(
-	//		SNew(SComboButton)
-	//		.IsEnabled(this, &SBindingsPanel::CanAddBinding)
-	//		.ContentPadding(0.f)
-	//		.ComboButtonStyle(&ToolBarStyle.SettingsComboButton)
-	//		.ButtonStyle(&ToolBarStyle.ButtonStyle)
-	//		.ToolTipText(this, &SBindingsPanel::GetAddBindingToolTip)
-	//		.ForegroundColor(FSlateColor::UseStyle())
-	//		.OnGetMenuContent(this, &SBindingsPanel::HandleAddDefaultBindingContextMenu)
-	//		.ButtonContent()
-	//		[
-	//			SNullWidget::NullWidget
-	//		]
-	//	);
-	//	ToolbarBuilderGlobal.EndSection();
-	//}
-
 	ToolbarBuilderGlobal.BeginSection("Picking");
 	{
 		ToolbarBuilderGlobal.AddToolBarButton(
diff --git a/Engine/Plugins/Runtime/ModelViewViewModel/Source/ModelViewViewModelEditor/Private/Widgets/SMVVMViewModelPanel.cpp b/Engine/Plugins/Runtime/ModelViewViewModel/Source/ModelViewViewModelEditor/Private/Widgets/SMVVMViewModelPanel.cpp
--- a/Engine/Plugins/Runtime/ModelViewViewModel/Source/ModelViewViewModelEditor/Private/Widgets/SMVVMViewModelPanel.cpp
+++ b/Engine/Plugins/Runtime/ModelViewViewModel/Source/ModelViewViewModelEditor/Private/Widgets/SMVVMViewModelPanel.cpp
@@ -8,9 +8,11 @@
 #include "MVVMBlueprintView.h"
 #include "MVVMEditorSubsystem.h"
 #include "PropertyEditorModule.h"
+#include "ToolMenus.h"
 #include "View/MVVMViewModelContextResolver.h"
 #include "WidgetBlueprint.h"
 #include "WidgetBlueprintEditor.h"
+#include "WidgetBlueprintToolMenuContext.h"
 
 #include "Framework/Commands/GenericCommands.h"
 #include "Framework/MultiBox/MultiBoxBuilder.h"
@@ -46,6 +48,38 @@ void UMVVMBlueprintViewModelContextWrapper::PostEditChangeProperty(struct FPrope
 namespace UE::MVVM
 {
 
+void SMVVMViewModelPanel::RegisterMenu()
+{
+	UToolMenu* Menu = UToolMenus::Get()->RegisterMenu("MVVM.Viewmodels.Toolbar");
+	Menu->MenuType = EMultiBoxType::SlimHorizontalToolBar;
+	FToolMenuSection& Section = Menu->FindOrAddSection("Left");
+	Section.AddDynamicEntry("AddViewmodel", FNewToolMenuSectionDelegate::CreateLambda([](FToolMenuSection& InSection)
+		{
+			if (const UMVVMViewModelPanelToolMenuContext* Context = InSection.FindContext<UMVVMViewModelPanelToolMenuContext>())
+			{
+				if (TSharedPtr<SMVVMViewModelPanel> ViewModelPanel = Context->ViewModelPanel.Pin())
+				{
+					ViewModelPanel->BuildContextMenu(InSection);
+				}
+			}
+		}));
+}
+
+
+void SMVVMViewModelPanel::BuildContextMenu(FToolMenuSection& InSection)
+{
+	if (!AddMenuButton.IsValid())
+	{
+		SAssignNew(AddMenuButton, SPositiveActionButton)
+			.OnGetMenuContent(this, &SMVVMViewModelPanel::MakeAddMenu)
+			.Text(LOCTEXT("Viewmodel", "Viewmodel"))
+			.IsEnabled(this, &SMVVMViewModelPanel::HandleCanEditViewmodelList);
+	}
+	
+	InSection.AddEntry(FToolMenuEntry::InitWidget("AddViewmodel", AddMenuButton.ToSharedRef(), FText()));
+}
+
+
 void SMVVMViewModelPanel::Construct(const FArguments& InArgs, TSharedPtr<FWidgetBlueprintEditor> WidgetBlueprintEditor)
 {
 	UWidgetBlueprint* WidgetBlueprint = WidgetBlueprintEditor->GetWidgetBlueprintObj();
@@ -64,9 +98,24 @@ void SMVVMViewModelPanel::Construct(const FArguments& InArgs, TSharedPtr<FWidget
 		// Listen to when the viewmodel are modified
 		ViewModelsUpdatedHandle = CurrentBlueprintView->OnViewModelsUpdated.AddSP(this, &SMVVMViewModelPanel::HandleViewModelsUpdated);
 	}
+	else
+	{
+		ExtensionAddeddHandle = WidgetBlueprint->OnExtensionAdded.AddSP(this, &SMVVMViewModelPanel::HandleViewUpdated);
+	}
 
 	CreateCommandList();
 
+	FToolMenuContext GenerateWidgetContext;
+	{
+		UWidgetBlueprintToolMenuContext* WidgetBlueprintMenuContext = NewObject<UWidgetBlueprintToolMenuContext>();
+		WidgetBlueprintMenuContext->WidgetBlueprintEditor = WeakBlueprintEditor;
+		GenerateWidgetContext.AddObject(WidgetBlueprintMenuContext);
+
+		UMVVMViewModelPanelToolMenuContext* ViewModelPanelToolMenuContext = NewObject<UMVVMViewModelPanelToolMenuContext>();
+		ViewModelPanelToolMenuContext->ViewModelPanel = SharedThis(this);
+		GenerateWidgetContext.AddObject(ViewModelPanelToolMenuContext);
+	}
+
 	ViewModelTreeView = SNew(UE::PropertyViewer::SPropertyViewer)
 		.PropertyVisibility(UE::PropertyViewer::SPropertyViewer::EPropertyVisibility::Visible)
 		.bShowSearchBox(true)
@@ -80,10 +129,7 @@ void SMVVMViewModelPanel::Construct(const FArguments& InArgs, TSharedPtr<FWidget
 		.OnGenerateContainer(this, &SMVVMViewModelPanel::HandleGenerateContainer)
 		.SearchBoxPreSlot()
 		[
-			SAssignNew(AddMenuButton, SPositiveActionButton)
-			.OnGetMenuContent(this, &SMVVMViewModelPanel::MakeAddMenu)
-			.Text(LOCTEXT("Viewmodel", "Viewmodel"))
-			.IsEnabled(this, &SMVVMViewModelPanel::HandleCanEditViewmodelList)
+			UToolMenus::Get()->GenerateWidget("MVVM.Viewmodels.Toolbar", GenerateWidgetContext)
 		];
 
 	FillViewModel();
@@ -111,19 +157,24 @@ SMVVMViewModelPanel::SMVVMViewModelPanel() = default;
 
 SMVVMViewModelPanel::~SMVVMViewModelPanel()
 {
-	if (TSharedPtr<FWidgetBlueprintEditor> WidgetBlueprintEditor = WeakBlueprintEditor.Pin())
+	if (ExtensionAddeddHandle.IsValid())
 	{
-		if (UWidgetBlueprint* WidgetBlueprint = WidgetBlueprintEditor->GetWidgetBlueprintObj())
+		if (TSharedPtr<FWidgetBlueprintEditor> WidgetBlueprintEditor = WeakBlueprintEditor.Pin())
 		{
-			WidgetBlueprint->OnExtensionAdded.RemoveAll(this);
-			WidgetBlueprint->OnExtensionRemoved.RemoveAll(this);
+			if (UWidgetBlueprint* WidgetBlueprint = WidgetBlueprintEditor->GetWidgetBlueprintObj())
+			{
+				WidgetBlueprint->OnExtensionAdded.Remove(ExtensionAddeddHandle);
+			}
 		}
 	}
 
-	if (UMVVMBlueprintView* CurrentBlueprintView = WeakBlueprintView.Get())
+	if (ViewModelsUpdatedHandle.IsValid())
 	{
-		// bind to check if the view is enabled
-		CurrentBlueprintView->OnViewModelsUpdated.Remove(ViewModelsUpdatedHandle);
+		if (UMVVMBlueprintView* CurrentBlueprintView = WeakBlueprintView.Get())
+		{
+			// bind to check if the view is enabled
+			CurrentBlueprintView->OnViewModelsUpdated.Remove(ViewModelsUpdatedHandle);
+		}
 	}
 }
 
@@ -178,6 +229,9 @@ void SMVVMViewModelPanel::HandleViewUpdated(UBlueprintExtension*)
 				{
 					ViewModelsUpdatedHandle = CurrentBlueprintView->OnViewModelsUpdated.AddSP(this, &SMVVMViewModelPanel::HandleViewModelsUpdated);
 					bViewUpdated = true;
+
+					WidgetBlueprint->OnExtensionAdded.Remove(ExtensionAddeddHandle);
+					ExtensionAddeddHandle.Reset();
 				}
 			}
 		}
@@ -243,15 +297,7 @@ void SMVVMViewModelPanel::HandleAddMenuViewModel(const UClass* SelectedClass)
 				{
 					UMVVMEditorSubsystem* EditorSubsystem = GEditor->GetEditorSubsystem<UMVVMEditorSubsystem>();
 					check(EditorSubsystem);
-
-					UMVVMBlueprintView* CurrentBlueprintView = WeakBlueprintView.Get();
-					if (!CurrentBlueprintView)
-					{
-						CurrentBlueprintView = EditorSubsystem->RequestView(WidgetBlueprint);
-						WeakBlueprintView = CurrentBlueprintView;
-						ViewModelsUpdatedHandle = CurrentBlueprintView->OnViewModelsUpdated.AddSP(this, &SMVVMViewModelPanel::HandleViewModelsUpdated);
-					}
-
+					EditorSubsystem->RequestView(WidgetBlueprint);
 					EditorSubsystem->AddViewModel(WidgetBlueprint, SelectedClass);
 				}
 			}
diff --git a/Engine/Plugins/Runtime/ModelViewViewModel/Source/ModelViewViewModelEditor/Private/Widgets/SMVVMViewModelPanel.h b/Engine/Plugins/Runtime/ModelViewViewModel/Source/ModelViewViewModelEditor/Private/Widgets/SMVVMViewModelPanel.h
--- a/Engine/Plugins/Runtime/ModelViewViewModel/Source/ModelViewViewModelEditor/Private/Widgets/SMVVMViewModelPanel.h
+++ b/Engine/Plugins/Runtime/ModelViewViewModel/Source/ModelViewViewModelEditor/Private/Widgets/SMVVMViewModelPanel.h
@@ -12,6 +12,7 @@
 
 namespace ETextCommit { enum Type : int; }
 namespace UE::MVVM { class FFieldIterator_Bindable; }
+namespace UE::MVVM { class SMVVMViewModelPanel; }
 namespace UE::PropertyViewer { class FFieldExpander_Default; }
 
 class INotifyFieldValueChanged;
@@ -22,6 +23,7 @@ class UBlueprintExtension;
 class UMVVMBlueprintView;
 class FUICommandList;
 class IDetailsView;
+struct FToolMenuSection;
 
 UCLASS()
 class UMVVMBlueprintViewModelContextWrapper : public UObject
@@ -38,6 +40,16 @@ public:
 	FGuid ViewModelId;
 };
 
+UCLASS()
+class UMVVMViewModelPanelToolMenuContext : public UObject
+{
+	GENERATED_BODY()
+
+public:
+	TWeakPtr<UE::MVVM::SMVVMViewModelPanel> ViewModelPanel;
+};
+
+
 namespace UE::MVVM
 {
 
@@ -56,9 +68,14 @@ public:
 
 	void OpenAddViewModelMenu();
 
+	static void RegisterMenu();
+
 private:
+	void BuildContextMenu(FToolMenuSection& InSection);
+
 	void HandleViewUpdated(UBlueprintExtension* Extension);
 	void HandleViewModelsUpdated();
+
 	TSharedRef<SWidget> MakeAddMenu();
 	void HandleCancelAddMenu();
 	void HandleAddMenuViewModel(const UClass* SelectedClass);
@@ -99,6 +116,7 @@ private:
 	TWeakPtr<FWidgetBlueprintEditor> WeakBlueprintEditor;
 	TWeakObjectPtr<UMVVMBlueprintView> WeakBlueprintView;
 	FDelegateHandle ViewModelsUpdatedHandle;
+	FDelegateHandle ExtensionAddeddHandle;
 
 	TStrongObjectPtr<UMVVMBlueprintViewModelContextWrapper> ModelContextWrapper;
 
