Drag and Drop support for viewmodel fields and widget properties in UMG details panel.

diff --git a/Engine/Plugins/Runtime/ModelViewViewModel/Source/ModelViewViewModelEditor/Private/MVVMEditorSubsystem.cpp b/Engine/Plugins/Runtime/ModelViewViewModel/Source/ModelViewViewModelEditor/Private/MVVMEditorSubsystem.cpp
--- a/Engine/Plugins/Runtime/ModelViewViewModel/Source/ModelViewViewModelEditor/Private/MVVMEditorSubsystem.cpp
+++ b/Engine/Plugins/Runtime/ModelViewViewModel/Source/ModelViewViewModelEditor/Private/MVVMEditorSubsystem.cpp
@@ -466,7 +466,7 @@ void UMVVMEditorSubsystem::SetDestinationToSourceConversionFunction(UWidgetBluep
 	}
 }
 
-void UMVVMEditorSubsystem::SetWidgetPropertyForBinding(UWidgetBlueprint* WidgetBlueprint, FMVVMBlueprintViewBinding& Binding, FMVVMBlueprintPropertyPath Field)
+void UMVVMEditorSubsystem::SetDestinationPathForBinding(UWidgetBlueprint* WidgetBlueprint, FMVVMBlueprintViewBinding& Binding, FMVVMBlueprintPropertyPath Field)
 {
 	if (UMVVMBlueprintView* View = GetView(WidgetBlueprint))
 	{
@@ -480,7 +480,7 @@ void UMVVMEditorSubsystem::SetWidgetPropertyForBinding(UWidgetBlueprint* WidgetB
 	}
 }
 
-void UMVVMEditorSubsystem::SetViewModelPropertyForBinding(UWidgetBlueprint* WidgetBlueprint, FMVVMBlueprintViewBinding& Binding, FMVVMBlueprintPropertyPath Field)
+void UMVVMEditorSubsystem::SetSourcePathForBinding(UWidgetBlueprint* WidgetBlueprint, FMVVMBlueprintViewBinding& Binding, FMVVMBlueprintPropertyPath Field)
 {
 	if (UMVVMBlueprintView* View = GetView(WidgetBlueprint))
 	{
diff --git a/Engine/Plugins/Runtime/ModelViewViewModel/Source/ModelViewViewModelEditor/Private/Widgets/SMVVMFieldSelector.cpp b/Engine/Plugins/Runtime/ModelViewViewModel/Source/ModelViewViewModelEditor/Private/Widgets/SMVVMFieldSelector.cpp
--- a/Engine/Plugins/Runtime/ModelViewViewModel/Source/ModelViewViewModelEditor/Private/Widgets/SMVVMFieldSelector.cpp
+++ b/Engine/Plugins/Runtime/ModelViewViewModel/Source/ModelViewViewModelEditor/Private/Widgets/SMVVMFieldSelector.cpp
@@ -4,6 +4,7 @@
 
 #include "WidgetBlueprint.h"
 
+#include "DragAndDrop/DecoratedDragDropOp.h"
 #include "Styling/MVVMEditorStyle.h"
 
 #include "Widgets/Input/SComboButton.h"
@@ -36,6 +37,8 @@ void SFieldSelector::Construct(const FArguments& InArgs, const UWidgetBlueprint*
 	OnGetConversionFunction = InArgs._OnGetConversionFunction;
 	OnGetSelectionContext = InArgs._OnGetSelectionContext;
 	OnFieldSelectionChanged = InArgs._OnFieldSelectionChanged;
+	OnDragEnterEvent = InArgs._OnDragEnter;
+	OnDropEvent = InArgs._OnDrop;
 
 	ChildSlot
 	[
@@ -160,6 +163,31 @@ void SFieldSelector::HandleFieldSelectionChanged(FMVVMBlueprintPropertyPath Prop
 	}
 }
 
+void SFieldSelector::OnDragEnter(const FGeometry& MyGeometry, const FDragDropEvent& DragDropEvent)
+{
+	if (OnDragEnterEvent.IsBound())
+	{
+		OnDragEnterEvent.Execute(MyGeometry, DragDropEvent);
+	}
+}
+
+void SFieldSelector::OnDragLeave(const FDragDropEvent& DragDropEvent)
+{
+	if (TSharedPtr<FDecoratedDragDropOp> DecoratedDragDropOp = DragDropEvent.GetOperationAs<FDecoratedDragDropOp>())
+	{
+		DecoratedDragDropOp->ResetToDefaultToolTip();
+	}
+}
+
+FReply SFieldSelector::OnDrop(const FGeometry& MyGeometry, const FDragDropEvent& DragDropEvent)
+{
+	if (OnDropEvent.IsBound())
+	{
+		return OnDropEvent.Execute(MyGeometry, DragDropEvent);
+	}
+	return FReply::Unhandled();
+}
+
 void SFieldSelector::HandleMenuClosed()
 {
 	if (ComboButton.IsValid())
diff --git a/Engine/Plugins/Runtime/ModelViewViewModel/Source/ModelViewViewModelEditor/Private/Widgets/SMVVMFieldSelector.h b/Engine/Plugins/Runtime/ModelViewViewModel/Source/ModelViewViewModelEditor/Private/Widgets/SMVVMFieldSelector.h
--- a/Engine/Plugins/Runtime/ModelViewViewModel/Source/ModelViewViewModelEditor/Private/Widgets/SMVVMFieldSelector.h
+++ b/Engine/Plugins/Runtime/ModelViewViewModel/Source/ModelViewViewModelEditor/Private/Widgets/SMVVMFieldSelector.h
@@ -22,6 +22,8 @@ class SFieldSelector : public SCompoundWidget
 {
 public:
 	DECLARE_DELEGATE_RetVal(FFieldSelectionContext, FOnGetSelectionContext);
+	DECLARE_DELEGATE_RetVal_TwoParams(FReply, FOnDrop, const FGeometry&, const FDragDropEvent&);
+	DECLARE_DELEGATE_TwoParams(FOnDragEnter, const FGeometry&, const FDragDropEvent&);
 
 	SLATE_BEGIN_ARGS(SFieldSelector)
 		: _TextStyle(&FCoreStyle::Get().GetWidgetStyle<FTextBlockStyle>("NormalText"))
@@ -33,9 +35,15 @@ public:
 		SLATE_EVENT(FOnGetConversionFunction, OnGetConversionFunction)
 		SLATE_EVENT(FOnFieldSelectionChanged, OnFieldSelectionChanged)
 		SLATE_EVENT(FOnGetSelectionContext, OnGetSelectionContext)
+		SLATE_EVENT(FOnDrop, OnDrop)
+		SLATE_EVENT(FOnDragEnter, OnDragEnter)
 	SLATE_END_ARGS()
 
 	void Construct(const FArguments& InArgs, const UWidgetBlueprint* InWidgetBlueprint);
+	virtual FReply OnDrop(const FGeometry& MyGeometry, const FDragDropEvent& DragDropEvent) override;
+	virtual void OnDragEnter(const FGeometry& MyGeometry, const FDragDropEvent& DragDropEvent) override;
+	virtual void OnDragLeave(const FDragDropEvent& DragDropEvent) override;
+
 
 private:
 	int32 GetCurrentDisplayIndex() const;
@@ -56,6 +64,8 @@ private:
 	FOnGetConversionFunction OnGetConversionFunction;
 	FOnGetSelectionContext OnGetSelectionContext;
 	FOnFieldSelectionChanged OnFieldSelectionChanged;
+	FOnDrop OnDropEvent;
+	FOnDragEnter OnDragEnterEvent;
 }; 
 
 } // namespace UE::MVVM
diff --git a/Engine/Plugins/Runtime/ModelViewViewModel/Source/ModelViewViewModelEditor/Private/Widgets/SMVVMViewBindingListView.cpp b/Engine/Plugins/Runtime/ModelViewViewModel/Source/ModelViewViewModelEditor/Private/Widgets/SMVVMViewBindingListView.cpp
--- a/Engine/Plugins/Runtime/ModelViewViewModel/Source/ModelViewViewModelEditor/Private/Widgets/SMVVMViewBindingListView.cpp
+++ b/Engine/Plugins/Runtime/ModelViewViewModel/Source/ModelViewViewModelEditor/Private/Widgets/SMVVMViewBindingListView.cpp
@@ -4,7 +4,9 @@
 
 #include "Bindings/MVVMBindingHelper.h"
 #include "BlueprintEditor.h"
+#include "Details/WidgetPropertyDragDropOp.h"
 #include "Dialog/SCustomDialog.h"
+#include "DragAndDrop/DecoratedDragDropOp.h"
 #include "Framework/MultiBox/MultiBoxBuilder.h"
 #include "Framework/Views/TableViewMetadata.h"
 #include "Misc/MessageDialog.h"
@@ -17,6 +19,7 @@
 #include "SSimpleButton.h" 
 #include "Styling/MVVMEditorStyle.h"
 #include "Styling/StyleColors.h"
+#include "ViewModelFieldDragDropOp.h"
 #include "Widgets/SMVVMFieldEntry.h"
 #include "Widgets/SMVVMFieldSelector.h"
 #include "Widgets/SMVVMFunctionParameter.h"
@@ -264,7 +267,7 @@ private:
 					FMVVMBlueprintPropertyPath CurrentPath = Binding->DestinationPath;
 					CurrentPath.SetWidgetName(Source.Name);
 
-					EditorSubsystem->SetWidgetPropertyForBinding(WidgetBlueprint, *Binding, CurrentPath);
+					EditorSubsystem->SetDestinationPathForBinding(WidgetBlueprint, *Binding, CurrentPath);
 				}
 			}
 		}
@@ -278,7 +281,7 @@ private:
 			FMVVMBlueprintViewBinding& Binding = EditorSubsystem->AddBinding(WidgetBlueprint);
 			FMVVMBlueprintPropertyPath Path;
 			Path.SetWidgetName(Entry->GetName());
-			EditorSubsystem->SetWidgetPropertyForBinding(WidgetBlueprint, Binding, Path);
+			EditorSubsystem->SetDestinationPathForBinding(WidgetBlueprint, Binding, Path);
 		}
 
 		return FReply::Handled();
@@ -357,6 +360,8 @@ public:
 								.OnGetConversionFunction(this, &SBindingRow::GetSelectedConversionFunction, false)
 								.OnFieldSelectionChanged(this, &SBindingRow::HandleFieldSelectionChanged, false)
 								.OnGetSelectionContext(this, &SBindingRow::GetSelectedSelectionContext, false)
+								.OnDrop(this, &SBindingRow::HandleFieldSelectorDrop, false)
+								.OnDragEnter(this, &SBindingRow::HandleFieldSelectorDragEnter, false)
 								.ShowContext(false)
 							]
 						]
@@ -401,6 +406,8 @@ public:
 								.OnGetConversionFunction(this, &SBindingRow::GetSelectedConversionFunction, true)
 								.OnFieldSelectionChanged(this, &SBindingRow::HandleFieldSelectionChanged, true)
 								.OnGetSelectionContext(this, &SBindingRow::GetSelectedSelectionContext, true)
+								.OnDrop(this, &SBindingRow::HandleFieldSelectorDrop, true)
+								.OnDragEnter(this, &SBindingRow::HandleFieldSelectorDragEnter, true)
 							]
 						]
 
@@ -642,7 +649,7 @@ private:
 				Subsystem->SetSourceToDestinationConversionFunction(WidgetBlueprintPtr, *ViewBinding, Function);
 				if (ViewBinding->SourcePath != SelectedField)
 				{
-					Subsystem->SetViewModelPropertyForBinding(WidgetBlueprintPtr, *ViewBinding, SelectedField);
+					Subsystem->SetSourcePathForBinding(WidgetBlueprintPtr, *ViewBinding, SelectedField);
 				}
 			}
 			else
@@ -650,7 +657,7 @@ private:
 				Subsystem->SetDestinationToSourceConversionFunction(WidgetBlueprintPtr, *ViewBinding, Function);
 				if (ViewBinding->DestinationPath != SelectedField)
 				{
-					Subsystem->SetWidgetPropertyForBinding(WidgetBlueprintPtr, *ViewBinding, SelectedField);
+					Subsystem->SetDestinationPathForBinding(WidgetBlueprintPtr, *ViewBinding, SelectedField);
 				}
 			}
 		}
@@ -711,6 +718,118 @@ private:
 		return Result;
 	}
 
+	FReply HandleFieldSelectorDrop(const FGeometry& MyGeometry, const FDragDropEvent& DragDropEvent, bool bSource)
+	{
+		TSharedPtr<FDecoratedDragDropOp> DragDropOp = DragDropEvent.GetOperationAs<FDecoratedDragDropOp>();
+		if (!DragDropOp.IsValid())
+		{
+			return FReply::Unhandled();
+		}
+
+		// Accept all drag-drop operations that are widget properties, but only accept view model fields when we are dropping into the Source box.
+		if (!DragDropOp->IsOfType<FWidgetPropertyDragDropOp>() && (!DragDropOp->IsOfType<FViewModelFieldDragDropOp>() || !bSource))
+		{
+			return FReply::Unhandled();
+		}
+
+		if (FMVVMBlueprintViewBinding* ViewBinding = GetThisViewBinding())
+		{
+			if (UWidgetBlueprint* WidgetBlueprintPtr = WidgetBlueprintWeak.Get())
+			{
+				TSharedPtr<FViewModelFieldDragDropOp> ViewModelFieldDragDropOp = DragDropEvent.GetOperationAs<FViewModelFieldDragDropOp>();
+				TSharedPtr<FWidgetPropertyDragDropOp> WidgetPropertyDragDropOp = DragDropEvent.GetOperationAs<FWidgetPropertyDragDropOp>();
+				bool bIsViewModelProperty = false;
+
+				if (ViewModelFieldDragDropOp)
+				{
+					bIsViewModelProperty = true;
+				}
+
+				UWidgetBlueprint* DragDropWidgetBP = bIsViewModelProperty ? ViewModelFieldDragDropOp->WidgetBP.Get() : WidgetPropertyDragDropOp->WidgetBP.Get();
+
+				if (WidgetBlueprintPtr == DragDropWidgetBP)
+				{
+					TArray<FFieldVariant> FieldPath = bIsViewModelProperty ? ViewModelFieldDragDropOp->DraggedField : WidgetPropertyDragDropOp->DraggedPropertyPath;
+					FMVVMBlueprintPropertyPath PropertyPath;
+
+					PropertyPath.ResetBasePropertyPath();
+					for (const FFieldVariant& Field : FieldPath)
+					{
+						PropertyPath.AppendBasePropertyPath(FMVVMConstFieldVariant(Field));
+					}
+
+					UMVVMEditorSubsystem* Subsystem = GEditor->GetEditorSubsystem<UMVVMEditorSubsystem>();
+					if (bIsViewModelProperty)
+					{
+						if (ViewModelFieldDragDropOp->ViewModelId.IsValid())
+						{
+							PropertyPath.SetViewModelId(ViewModelFieldDragDropOp->ViewModelId);
+						}
+						else
+						{
+							return FReply::Unhandled();
+						}
+					}
+					else
+					{
+						if (UWidget* OwnerWidgetPtr = WidgetPropertyDragDropOp->OwnerWidget.Get())
+						{
+							PropertyPath.SetWidgetName(OwnerWidgetPtr->GetFName());
+						}
+					}
+
+					if (bSource)
+					{
+						Subsystem->SetSourceToDestinationConversionFunction(WidgetBlueprintPtr, *ViewBinding, nullptr);
+						Subsystem->SetSourcePathForBinding(WidgetBlueprintPtr, *ViewBinding, PropertyPath);
+					}
+					else
+					{
+						Subsystem->SetDestinationToSourceConversionFunction(WidgetBlueprintPtr, *ViewBinding, nullptr);
+						Subsystem->SetDestinationPathForBinding(WidgetBlueprintPtr, *ViewBinding, PropertyPath);
+					}
+					return FReply::Handled();
+				}
+			}
+		}		
+		return FReply::Unhandled();
+	}
+
+	void HandleFieldSelectorDragEnter(const FGeometry& MyGeometry, const FDragDropEvent& DragDropEvent, bool bSource)
+	{
+		TSharedPtr<FDecoratedDragDropOp> DragDropOp = DragDropEvent.GetOperationAs<FDecoratedDragDropOp>();
+		if (DragDropOp.IsValid())
+		{
+			// Accept all drag-drop operations that are widget properties, but only accept view model fields when we are dropping into the Source box.
+			if (DragDropOp->IsOfType<FWidgetPropertyDragDropOp>() || (DragDropOp->IsOfType<FViewModelFieldDragDropOp>() && bSource))
+			{
+				if (UWidgetBlueprint* WidgetBlueprintPtr = WidgetBlueprintWeak.Get())
+				{
+					TSharedPtr<FViewModelFieldDragDropOp> ViewModelFieldDragDropOp = DragDropEvent.GetOperationAs<FViewModelFieldDragDropOp>();
+					TSharedPtr<FWidgetPropertyDragDropOp> WidgetPropertyDragDropOp = DragDropEvent.GetOperationAs<FWidgetPropertyDragDropOp>();
+					bool IsViewModelProperty = false;
+
+					if (ViewModelFieldDragDropOp)
+					{
+						IsViewModelProperty = true;
+					}
+
+					UWidgetBlueprint* DragDropWidgetBP = IsViewModelProperty ? ViewModelFieldDragDropOp->WidgetBP.Get() : WidgetPropertyDragDropOp->WidgetBP.Get();
+
+					if (DragDropWidgetBP && DragDropWidgetBP == WidgetBlueprintPtr)
+					{
+						DragDropOp->CurrentIconBrush = FAppStyle::GetBrush(TEXT("Graph.ConnectorFeedback.OK"));
+						return;
+					}
+				}
+			}	
+			else
+			{
+				DragDropOp->CurrentIconBrush = FAppStyle::GetBrush(TEXT("Graph.ConnectorFeedback.Error"));
+			}
+		}
+	}
+
 	ECheckBoxState IsExecutionModeOverrideChecked() const
 	{
 		if (FMVVMBlueprintViewBinding* ViewBinding = GetThisViewBinding())
diff --git a/Engine/Plugins/Runtime/ModelViewViewModel/Source/ModelViewViewModelEditor/Private/Widgets/SMVVMViewBindingPanel.cpp b/Engine/Plugins/Runtime/ModelViewViewModel/Source/ModelViewViewModelEditor/Private/Widgets/SMVVMViewBindingPanel.cpp
--- a/Engine/Plugins/Runtime/ModelViewViewModel/Source/ModelViewViewModelEditor/Private/Widgets/SMVVMViewBindingPanel.cpp
+++ b/Engine/Plugins/Runtime/ModelViewViewModel/Source/ModelViewViewModelEditor/Private/Widgets/SMVVMViewBindingPanel.cpp
@@ -224,7 +224,7 @@ void SBindingsPanel::AddDefaultBinding()
 						FMVVMBlueprintViewBinding& Binding = EditorSubsystem->AddBinding(MVVMExtensionPtr->GetWidgetBlueprint());
 						FMVVMBlueprintPropertyPath Path;
 						Path.SetWidgetName(WidgetReference.GetTemplate()->GetFName());
-						EditorSubsystem->SetWidgetPropertyForBinding(MVVMExtensionPtr->GetWidgetBlueprint(), Binding, Path);
+						EditorSubsystem->SetDestinationPathForBinding(MVVMExtensionPtr->GetWidgetBlueprint(), Binding, Path);
 						AddedBindingId = Binding.BindingId;
 
 						bBindingAdded = true;
diff --git a/Engine/Plugins/Runtime/ModelViewViewModel/Source/ModelViewViewModelEditor/Private/Widgets/SMVVMViewModelPanel.cpp b/Engine/Plugins/Runtime/ModelViewViewModel/Source/ModelViewViewModelEditor/Private/Widgets/SMVVMViewModelPanel.cpp
--- a/Engine/Plugins/Runtime/ModelViewViewModel/Source/ModelViewViewModelEditor/Private/Widgets/SMVVMViewModelPanel.cpp
+++ b/Engine/Plugins/Runtime/ModelViewViewModel/Source/ModelViewViewModelEditor/Private/Widgets/SMVVMViewModelPanel.cpp
@@ -10,6 +10,7 @@
 #include "PropertyEditorModule.h"
 #include "ToolMenus.h"
 #include "View/MVVMViewModelContextResolver.h"
+#include "ViewModelFieldDragDropOp.h"
 #include "WidgetBlueprint.h"
 #include "WidgetBlueprintEditor.h"
 #include "WidgetBlueprintToolMenuContext.h"
@@ -127,6 +128,7 @@ void SMVVMViewModelPanel::Construct(const FArguments& InArgs, TSharedPtr<FWidget
 		.OnContextMenuOpening(this, &SMVVMViewModelPanel::HandleContextMenuOpening)
 		.OnSelectionChanged(this, &SMVVMViewModelPanel::HandleSelectionChanged)
 		.OnGenerateContainer(this, &SMVVMViewModelPanel::HandleGenerateContainer)
+		.OnDragDetected(this, &SMVVMViewModelPanel::HandleDragDetected)
 		.SearchBoxPreSlot()
 		[
 			UToolMenus::Get()->GenerateWidget("MVVM.Viewmodels.Toolbar", GenerateWidgetContext)
@@ -406,6 +408,38 @@ void SMVVMViewModelPanel::HandleNameTextCommited(const FText& InText, ETextCommi
 	}
 }
 
+FReply SMVVMViewModelPanel::HandleDragDetected(const FGeometry& MyGeometry, const FPointerEvent& MouseEvent, UE::PropertyViewer::SPropertyViewer::FHandle ContainerHandle, TArrayView<const FFieldVariant> Fields) const
+{
+	TSharedPtr<FWidgetBlueprintEditor> WidgetBlueprintEditor = WeakBlueprintEditor.Pin();
+	if (WidgetBlueprintEditor == nullptr)
+	{
+		return FReply::Unhandled();
+	}
+
+	UWidgetBlueprint* WidgetBP = WidgetBlueprintEditor->GetWidgetBlueprintObj();
+	if (WidgetBP == nullptr)
+	{
+		return FReply::Unhandled();
+	}
+
+	if (ViewModelTreeView.IsValid() && MouseEvent.IsMouseButtonDown(EKeys::LeftMouseButton))
+	{
+		if (Fields.Num() > 0)
+		{
+			if (const FGuid* Id = PropertyViewerHandles.Find(ContainerHandle))
+			{
+				TArray<FFieldVariant> FieldsArray;
+				for (const FFieldVariant Field : Fields)
+				{
+					FieldsArray.Add(Field);
+				}
+				
+				return FReply::Handled().BeginDragDrop(FViewModelFieldDragDropOp::New(FieldsArray, *Id, WidgetBP));
+			}
+		}
+	}
+	return FReply::Unhandled();
+}
 
 void SMVVMViewModelPanel::CreateCommandList()
 {
diff --git a/Engine/Plugins/Runtime/ModelViewViewModel/Source/ModelViewViewModelEditor/Private/Widgets/SMVVMViewModelPanel.h b/Engine/Plugins/Runtime/ModelViewViewModel/Source/ModelViewViewModelEditor/Private/Widgets/SMVVMViewModelPanel.h
--- a/Engine/Plugins/Runtime/ModelViewViewModel/Source/ModelViewViewModelEditor/Private/Widgets/SMVVMViewModelPanel.h
+++ b/Engine/Plugins/Runtime/ModelViewViewModel/Source/ModelViewViewModelEditor/Private/Widgets/SMVVMViewModelPanel.h
@@ -88,6 +88,7 @@ private:
 	void HandleEditorSelectionChanged();
 	bool HandleVerifyNameTextChanged(const FText& InText, FText& OutErrorMessage, FGuid ViewModelGuid);
 	void HandleNameTextCommited(const FText& InText, ETextCommit::Type CommitInfo, FGuid ViewModelGuid);
+	FReply HandleDragDetected(const FGeometry& MyGeometry, const FPointerEvent& MouseEvent, UE::PropertyViewer::SPropertyViewer::FHandle ContainerHandle, TArrayView<const FFieldVariant> Fields) const;
 
 	void HandleDeleteViewModel();
 	bool HandleCanDeleteViewModel() const;
diff --git a/Engine/Plugins/Runtime/ModelViewViewModel/Source/ModelViewViewModelEditor/Private/Widgets/ViewModelFieldDragDropOp.cpp b/Engine/Plugins/Runtime/ModelViewViewModel/Source/ModelViewViewModelEditor/Private/Widgets/ViewModelFieldDragDropOp.cpp
new file mode 100644
--- /dev/null
+++ b/Engine/Plugins/Runtime/ModelViewViewModel/Source/ModelViewViewModelEditor/Private/Widgets/ViewModelFieldDragDropOp.cpp
@@ -0,0 +1,49 @@
+// Copyright Epic Games, Inc. All Rights Reserved.
+
+#include "Widgets/ViewModelFieldDragDropOp.h"
+#include "CoreMinimal.h"
+#include "UObject/Field.h"
+#include "DragAndDrop/DecoratedDragDropOp.h"
+#include "Misc/Guid.h"
+#include "ScopedTransaction.h"
+#include "WidgetBlueprint.h"
+
+#define LOCTEXT_NAMESPACE "ViewModelFieldDragDropOp"
+
+namespace UE::MVVM
+{
+FViewModelFieldDragDropOp::FViewModelFieldDragDropOp()
+	: WidgetBP(nullptr)
+	, Transaction(nullptr)
+{
+}
+
+void FViewModelFieldDragDropOp::OnDrop(bool bDropWasHandled, const FPointerEvent& MouseEvent)
+{
+	if (!bDropWasHandled)
+	{
+		if (Transaction)
+		{
+			Transaction->Cancel();
+		}
+		Transaction = nullptr;
+	}
+}
+
+/** Constructs a new drag/drop operation */
+TSharedRef<FViewModelFieldDragDropOp> FViewModelFieldDragDropOp::New(TArray<FFieldVariant> InField, const FGuid& ViewModelId, UWidgetBlueprint* InWidgetBP)
+{
+	TSharedRef<FViewModelFieldDragDropOp> Operation = MakeShared<FViewModelFieldDragDropOp>();
+
+	// Set the display text and the transaction name based on whether we're dragging a single or multiple widgets
+	Operation->CurrentHoverText = Operation->DefaultHoverText = FText::FromName(InField[0].GetFName());
+	Operation->Transaction = MakeUnique<FScopedTransaction>(LOCTEXT("Designer_DragViewModelField", "Drag ViewModel Field"));
+	Operation->DraggedField = InField;
+	Operation->ViewModelId = ViewModelId;
+	Operation->WidgetBP = InWidgetBP;
+
+	Operation->Construct();
+	return Operation;
+}
+} // namespace
+#undef LOCTEXT_NAMESPACE
\ No newline at end of file
diff --git a/Engine/Plugins/Runtime/ModelViewViewModel/Source/ModelViewViewModelEditor/Private/Widgets/ViewModelFieldDragDropOp.h b/Engine/Plugins/Runtime/ModelViewViewModel/Source/ModelViewViewModelEditor/Private/Widgets/ViewModelFieldDragDropOp.h
new file mode 100644
--- /dev/null
+++ b/Engine/Plugins/Runtime/ModelViewViewModel/Source/ModelViewViewModelEditor/Private/Widgets/ViewModelFieldDragDropOp.h
@@ -0,0 +1,31 @@
+// Copyright Epic Games, Inc. All Rights Reserved.
+
+#pragma once
+
+#include "CoreMinimal.h"
+#include "UObject/Field.h"
+#include "DragAndDrop/DecoratedDragDropOp.h"
+#include "Misc/Guid.h"
+
+class FScopedTransaction;
+class UWidgetBlueprint;
+
+namespace UE::MVVM
+{
+
+class FViewModelFieldDragDropOp : public FDecoratedDragDropOp
+{
+public:
+	DRAG_DROP_OPERATOR_TYPE(FViewModelFieldDragDropOp, FDecoratedDragDropOp)
+
+	FViewModelFieldDragDropOp();
+	virtual void OnDrop(bool bDropWasHandled, const FPointerEvent& MouseEvent) override;
+	static TSharedRef<FViewModelFieldDragDropOp> New(TArray<FFieldVariant> InField, const FGuid& ViewModelId, UWidgetBlueprint* InWidgetBP);
+
+	TArray<FFieldVariant> DraggedField;
+	FGuid ViewModelId;
+	TWeakObjectPtr<UWidgetBlueprint> WidgetBP;
+	TUniquePtr<FScopedTransaction> Transaction;
+};
+
+} // namespace
\ No newline at end of file
diff --git a/Engine/Plugins/Runtime/ModelViewViewModel/Source/ModelViewViewModelEditor/Public/MVVMEditorSubsystem.h b/Engine/Plugins/Runtime/ModelViewViewModel/Source/ModelViewViewModelEditor/Public/MVVMEditorSubsystem.h
--- a/Engine/Plugins/Runtime/ModelViewViewModel/Source/ModelViewViewModelEditor/Public/MVVMEditorSubsystem.h
+++ b/Engine/Plugins/Runtime/ModelViewViewModel/Source/ModelViewViewModelEditor/Public/MVVMEditorSubsystem.h
@@ -58,8 +58,8 @@ public:
 
 	void SetSourceToDestinationConversionFunction(UWidgetBlueprint* WidgetBlueprint, FMVVMBlueprintViewBinding& Binding, const UFunction* ConversionFunction);
 	void SetDestinationToSourceConversionFunction(UWidgetBlueprint* WidgetBlueprint, FMVVMBlueprintViewBinding& Binding, const UFunction* ConversionFunction);
-	void SetWidgetPropertyForBinding(UWidgetBlueprint* WidgetBlueprint, FMVVMBlueprintViewBinding& Binding, FMVVMBlueprintPropertyPath Field);
-	void SetViewModelPropertyForBinding(UWidgetBlueprint* WidgetBlueprint, FMVVMBlueprintViewBinding& Binding, FMVVMBlueprintPropertyPath Field);
+	void SetDestinationPathForBinding(UWidgetBlueprint* WidgetBlueprint, FMVVMBlueprintViewBinding& Binding, FMVVMBlueprintPropertyPath Field);
+	void SetSourcePathForBinding(UWidgetBlueprint* WidgetBlueprint, FMVVMBlueprintViewBinding& Binding, FMVVMBlueprintPropertyPath Field);
 	void OverrideExecutionModeForBinding(UWidgetBlueprint* WidgetBlueprint, FMVVMBlueprintViewBinding& Binding, EMVVMExecutionMode Mode);
 	void ResetExecutionModeForBinding(UWidgetBlueprint* WidgetBlueprint, FMVVMBlueprintViewBinding& Binding);
 	void SetBindingTypeForBinding(UWidgetBlueprint* WidgetBlueprint, FMVVMBlueprintViewBinding& Binding, EMVVMBindingMode Type);
diff --git a/Engine/Source/Editor/UMGEditor/Private/Details/DetailWidgetExtensionHandler.cpp b/Engine/Source/Editor/UMGEditor/Private/Details/DetailWidgetExtensionHandler.cpp
--- a/Engine/Source/Editor/UMGEditor/Private/Details/DetailWidgetExtensionHandler.cpp
+++ b/Engine/Source/Editor/UMGEditor/Private/Details/DetailWidgetExtensionHandler.cpp
@@ -3,6 +3,7 @@
 #include "Details/DetailWidgetExtensionHandler.h"
 #include "Binding/WidgetBinding.h"
 #include "Customizations/UMGDetailCustomizations.h"
+#include "Details/WidgetPropertyDragDropHandler.h"
 #include "Engine/Blueprint.h"
 #include "UMGEditorModule.h"
 #include "WidgetBlueprint.h"
@@ -137,7 +138,7 @@ void FDetailWidgetExtensionHandler::ExtendWidgetRow(
 			return !FBlueprintWidgetCustomization::HasPropertyBindings(BlueprintEditorPtr, PropertyHandleRef);
 		}));
 
-	const UWidgetBlueprint* WidgetBlueprint = BlueprintEditor.Pin()->GetWidgetBlueprintObj();
+	UWidgetBlueprint* WidgetBlueprint = BlueprintEditor.Pin()->GetWidgetBlueprintObj();
 
 	bool bShouldShowOldBindingWidget = Private::ShouldShowOldBindingWidget(WidgetBlueprint, PropertyHandleRef.Get());
 
@@ -177,4 +178,12 @@ void FDetailWidgetExtensionHandler::ExtendWidgetRow(
 	[
 		FBlueprintWidgetCustomization::MakePropertyBindingWidget(BlueprintEditor, SignatureFunction, InPropertyHandle.ToSharedRef(), true, bShouldShowOldBindingWidget)
 	];
+
+	if (Objects.Num() > 0)
+	{
+		if (UWidget* Widget = Cast<UWidget>(Objects[0]))
+		{
+			InWidgetRow.DragDropHandler(MakeShared<FWidgetPropertyDragDropHandler>(Widget, InPropertyHandle, WidgetBlueprint));
+		}
+	}
 }
diff --git a/Engine/Source/Editor/UMGEditor/Private/Details/WidgetPropertyDragDropHandler.cpp b/Engine/Source/Editor/UMGEditor/Private/Details/WidgetPropertyDragDropHandler.cpp
new file mode 100644
--- /dev/null
+++ b/Engine/Source/Editor/UMGEditor/Private/Details/WidgetPropertyDragDropHandler.cpp
@@ -0,0 +1,67 @@
+// Copyright Epic Games, Inc. All Rights Reserved.
+
+#include "Details/WidgetPropertyDragDropHandler.h"
+
+#include "Binding/WidgetBinding.h"
+#include "Customizations/UMGDetailCustomizations.h"
+#include "Details/WidgetPropertyDragDropOp.h"
+#include "Engine/Blueprint.h"
+#include "PropertyHandle.h"
+#include "PropertyPathHelpers.h"
+#include "UMGEditorModule.h"
+#include "WidgetBlueprint.h"
+#include "WidgetBlueprintEditor.h"
+
+FWidgetPropertyDragDropHandler::FWidgetPropertyDragDropHandler(UWidget* InWidget, TSharedPtr<IPropertyHandle> InPropertyHandle, UWidgetBlueprint* InWidgetBP)
+{
+	OwnerWidget = InWidget;
+	PropertyHandle = InPropertyHandle;
+	WidgetBP = InWidgetBP;
+}
+
+TSharedPtr<FDragDropOperation> FWidgetPropertyDragDropHandler::CreateDragDropOperation() const
+{
+	if (!PropertyHandle.IsValid())
+	{
+		return MakeShared<FDragDropOperation>();
+	}
+
+	FProperty* WidgetProperty = PropertyHandle->GetProperty();
+	if (WidgetProperty == nullptr)
+	{
+		return MakeShared<FDragDropOperation>();
+	}
+	if (UWidget* WidgetPtr = OwnerWidget.Get())
+	{
+		if (UWidgetBlueprint* WidgetBPPtr = WidgetBP.Get())
+		{
+			FCachedPropertyPath CachedPropertyPath(PropertyHandle->GeneratePathToProperty());
+			CachedPropertyPath.Resolve(WidgetPtr);
+			TArray<FFieldVariant> FieldPath;
+
+			for (int32 SegNum = 0; SegNum < CachedPropertyPath.GetNumSegments(); SegNum++)
+			{
+				FieldPath.Add(CachedPropertyPath.GetSegment(SegNum).GetField());
+			}
+			TSharedPtr<FWidgetPropertyDragDropOp> DragOp = FWidgetPropertyDragDropOp::New(WidgetPtr, WidgetProperty->GetFName(), FieldPath, WidgetBPPtr);
+			return DragOp;
+		}
+	}
+	return MakeShared<FDragDropOperation>();
+}
+
+int32 FWidgetPropertyDragDropHandler::ComputeNewIndex(int32 OriginalIndex, int32 DropOntoIndex, EItemDropZone DropZone)
+{
+	return 0;
+}
+
+bool FWidgetPropertyDragDropHandler::AcceptDrop(const FDragDropEvent& DragDropEvent, EItemDropZone DropZone) const
+{
+	return false;
+}
+
+TOptional<EItemDropZone> FWidgetPropertyDragDropHandler::CanAcceptDrop(const FDragDropEvent& DragDropEvent, EItemDropZone DropZone) const
+{
+	return TOptional<EItemDropZone>();
+}
+
diff --git a/Engine/Source/Editor/UMGEditor/Private/Details/WidgetPropertyDragDropHandler.h b/Engine/Source/Editor/UMGEditor/Private/Details/WidgetPropertyDragDropHandler.h
new file mode 100644
--- /dev/null
+++ b/Engine/Source/Editor/UMGEditor/Private/Details/WidgetPropertyDragDropHandler.h
@@ -0,0 +1,33 @@
+// Copyright Epic Games, Inc. All Rights Reserved.
+
+#pragma once
+
+#include "CoreMinimal.h"
+#include "DragAndDrop/DecoratedDragDropOp.h"
+#include "IDetailDragDropHandler.h"
+#include "Input/DragAndDrop.h"
+
+class IPropertyHandle;
+class UWidget;
+class UWidgetBlueprint;
+
+/** Handler for customizing the drag-and-drop behavior for function entry/result pins, allowing parameters to be reordered */
+class FWidgetPropertyDragDropHandler : public IDetailDragDropHandler
+{
+public:
+	FWidgetPropertyDragDropHandler(UWidget* InWidget, TSharedPtr<IPropertyHandle> InPropertyHandle, UWidgetBlueprint* InWidgetBP);
+
+	virtual TSharedPtr<FDragDropOperation> CreateDragDropOperation() const override;
+
+	static int32 ComputeNewIndex(int32 OriginalIndex, int32 DropOntoIndex, EItemDropZone DropZone);
+
+	virtual bool AcceptDrop(const FDragDropEvent& DragDropEvent, EItemDropZone DropZone) const override;
+
+	virtual TOptional<EItemDropZone> CanAcceptDrop(const FDragDropEvent& DragDropEvent, EItemDropZone DropZone) const override;
+
+private:
+	TSharedPtr<IPropertyHandle> PropertyHandle;
+	TWeakObjectPtr<UWidget> OwnerWidget;
+	TWeakObjectPtr<UWidgetBlueprint> WidgetBP;
+};
+
diff --git a/Engine/Source/Editor/UMGEditor/Private/Details/WidgetPropertyDragDropOp.cpp b/Engine/Source/Editor/UMGEditor/Private/Details/WidgetPropertyDragDropOp.cpp
new file mode 100644
--- /dev/null
+++ b/Engine/Source/Editor/UMGEditor/Private/Details/WidgetPropertyDragDropOp.cpp
@@ -0,0 +1,42 @@
+// Copyright Epic Games, Inc. All Rights Reserved.
+
+#include "Details/WidgetPropertyDragDropOp.h"
+#include "Components/Widget.h"
+#include "CoreMinimal.h"
+#include "ScopedTransaction.h"
+#include "WidgetBlueprint.h"
+
+#define LOCTEXT_NAMESPACE "WidgetPropertyDragDropOp"
+
+FWidgetPropertyDragDropOp::FWidgetPropertyDragDropOp()
+	: OwnerWidget(nullptr)
+	, WidgetBP(nullptr)
+	, Transaction(nullptr)
+{
+}
+
+void FWidgetPropertyDragDropOp::OnDrop(bool bDropWasHandled, const FPointerEvent& MouseEvent)
+{
+	if (!bDropWasHandled)
+	{
+		check(Transaction);
+		Transaction->Cancel();
+	}
+}
+
+/** Constructs a new drag/drop operation */
+TSharedRef<FWidgetPropertyDragDropOp> FWidgetPropertyDragDropOp::New(UWidget* InOwnerWidget, FName InPropertyName, TArray<FFieldVariant> InPropertyPath, UWidgetBlueprint* InWidgetBP)
+{
+	TSharedRef<FWidgetPropertyDragDropOp> Operation = MakeShared<FWidgetPropertyDragDropOp>();
+
+	// Set the display text and the transaction name based on whether we're dragging a single or multiple widgets
+	Operation->CurrentHoverText = Operation->DefaultHoverText = FText::FromName(InPropertyName);
+	Operation->Transaction = MakeUnique<FScopedTransaction>(LOCTEXT("Designer_DragWidgetProperty", "Drag Widget Property"));
+	Operation->OwnerWidget = InOwnerWidget;
+	Operation->DraggedPropertyPath = InPropertyPath;
+	Operation->WidgetBP = InWidgetBP;
+
+	Operation->Construct();
+	return Operation;
+}
+#undef LOCTEXT_NAMESPACE
\ No newline at end of file
diff --git a/Engine/Source/Editor/UMGEditor/Public/Details/WidgetPropertyDragDropOp.h b/Engine/Source/Editor/UMGEditor/Public/Details/WidgetPropertyDragDropOp.h
new file mode 100644
--- /dev/null
+++ b/Engine/Source/Editor/UMGEditor/Public/Details/WidgetPropertyDragDropOp.h
@@ -0,0 +1,31 @@
+// Copyright Epic Games, Inc. All Rights Reserved.
+
+#pragma once
+
+#include "CoreMinimal.h"
+#include "DragAndDrop/DecoratedDragDropOp.h"
+#include "Input/DragAndDrop.h"
+#include "UObject/Field.h"
+
+class FScopedTransaction;
+class UWidget;
+class UWidgetBlueprint;
+
+#define LOCTEXT_NAMESPACE "WidgetPropertyDragDropOp"
+/** Drag-and-drop operation that stores data about the a widget property being dragged on. */
+class FWidgetPropertyDragDropOp : public FDecoratedDragDropOp
+{
+public:
+	DRAG_DROP_OPERATOR_TYPE(FWidgetPropertyDragDropOp, FDecoratedDragDropOp);
+
+	FWidgetPropertyDragDropOp();
+	virtual void OnDrop(bool bDropWasHandled, const FPointerEvent& MouseEvent) override;
+	static TSharedRef<FWidgetPropertyDragDropOp> New(UWidget* InOwnerWidget, FName InPropertyName, TArray<FFieldVariant> InPropertyPath, UWidgetBlueprint* InWidgetBP);
+
+	TWeakObjectPtr<UWidget> OwnerWidget;
+	TArray<FFieldVariant> DraggedPropertyPath;
+	TWeakObjectPtr<UWidgetBlueprint> WidgetBP;
+	TUniquePtr<FScopedTransaction> Transaction;
+};
+#undef LOCTEXT_NAMESPACE
+
diff --git a/Engine/Source/Runtime/AdvancedWidgets/Private/Widgets/PropertyViewer/SPropertyViewerImpl.cpp b/Engine/Source/Runtime/AdvancedWidgets/Private/Widgets/PropertyViewer/SPropertyViewerImpl.cpp
--- a/Engine/Source/Runtime/AdvancedWidgets/Private/Widgets/PropertyViewer/SPropertyViewerImpl.cpp
+++ b/Engine/Source/Runtime/AdvancedWidgets/Private/Widgets/PropertyViewer/SPropertyViewerImpl.cpp
@@ -412,6 +412,7 @@ FPropertyViewerImpl::FPropertyViewerImpl(const SPropertyViewer::FArguments& InAr
 	OnContextMenuOpening = InArgs._OnContextMenuOpening;
 	OnSelectionChanged = InArgs._OnSelectionChanged;
 	OnDoubleClicked = InArgs._OnDoubleClicked;
+	OnDragDetected = InArgs._OnDragDetected;
 	OnGenerateContainer = InArgs._OnGenerateContainer;
 	PropertyVisibility = InArgs._PropertyVisibility;
 	bSanitizeName = InArgs._bSanitizeName;
@@ -1087,6 +1088,7 @@ TSharedRef<ITableRow> FPropertyViewerImpl::HandleGenerateRow(TSharedPtr<FTreeNod
 	if (bUseRows)
 	{
 		return SNew(SMultiRowType, AsShared(), OwnerTable, Item.ToSharedRef(), FieldWidget)
+			.OnDragDetected(this, &FPropertyViewerImpl::HandleDragDetected, Item)
 			.Padding(0.0f);
 	}
 
@@ -1183,6 +1185,23 @@ void FPropertyViewerImpl::HandleDoubleClick(TSharedPtr<FTreeNode> Item)
 }
 
 
+FReply FPropertyViewerImpl::HandleDragDetected(const FGeometry& MyGeometry, const FPointerEvent& MouseEvent, const TSharedPtr<FTreeNode> Item)
+{
+	if (OnDragDetected.IsBound())
+	{
+		if (Item.IsValid())
+		{
+			TSharedPtr<FContainer> OwnerContainer = Item->GetOwnerContainer();
+			return OnDragDetected.Execute(MyGeometry, MouseEvent, OwnerContainer.IsValid() ? OwnerContainer->GetIdentifier() : SPropertyViewer::FHandle(), Item->GetFieldPath());
+		}
+		else
+		{
+			return OnDragDetected.Execute(MyGeometry, MouseEvent, SPropertyViewer::FHandle(), TArray<FFieldVariant>());
+		}
+	}
+	return FReply::Unhandled();
+}
+
 void FPropertyViewerImpl::HandleSearchChanged(const FText& InFilterText)
 {
 	if (SearchBoxWidget)
diff --git a/Engine/Source/Runtime/AdvancedWidgets/Private/Widgets/PropertyViewer/SPropertyViewerImpl.h b/Engine/Source/Runtime/AdvancedWidgets/Private/Widgets/PropertyViewer/SPropertyViewerImpl.h
--- a/Engine/Source/Runtime/AdvancedWidgets/Private/Widgets/PropertyViewer/SPropertyViewerImpl.h
+++ b/Engine/Source/Runtime/AdvancedWidgets/Private/Widgets/PropertyViewer/SPropertyViewerImpl.h
@@ -221,6 +221,7 @@ private:
 	SPropertyViewer::FOnContextMenuOpening OnContextMenuOpening;
 	SPropertyViewer::FOnSelectionChanged OnSelectionChanged;
 	SPropertyViewer::FOnDoubleClicked OnDoubleClicked;
+	SPropertyViewer::FOnDragDetected OnDragDetected;
 	SPropertyViewer::FOnGenerateContainer OnGenerateContainer;
 	SPropertyViewer::EPropertyVisibility PropertyVisibility;
 	bool bSanitizeName = false;
@@ -264,6 +265,7 @@ private:
 	TSharedPtr<SWidget> HandleContextMenuOpening();
 	void HandleSelectionChanged(TSharedPtr<FTreeNode> Item, ESelectInfo::Type SelectionType);
 	void HandleDoubleClick(TSharedPtr<FTreeNode> Item);
+	FReply HandleDragDetected(const FGeometry& MyGeometry, const FPointerEvent& MouseEvent, const TSharedPtr<FTreeNode> Item);
 
 	TSharedPtr<FTreeNode> FindExistingChild(const TSharedPtr<FTreeNode>& ContainerNode, TArrayView<const FFieldVariant> FieldPath) const;
 
diff --git a/Engine/Source/Runtime/AdvancedWidgets/Public/Widgets/PropertyViewer/SPropertyViewer.h b/Engine/Source/Runtime/AdvancedWidgets/Public/Widgets/PropertyViewer/SPropertyViewer.h
--- a/Engine/Source/Runtime/AdvancedWidgets/Public/Widgets/PropertyViewer/SPropertyViewer.h
+++ b/Engine/Source/Runtime/AdvancedWidgets/Public/Widgets/PropertyViewer/SPropertyViewer.h
@@ -71,6 +71,7 @@ public:
 	DECLARE_DELEGATE_RetVal_TwoParams(TSharedPtr<SWidget>, FOnContextMenuOpening, FHandle, TArrayView<const FFieldVariant>);
 	DECLARE_DELEGATE_ThreeParams(FOnSelectionChanged, FHandle, TArrayView<const FFieldVariant>, ESelectInfo::Type);
 	DECLARE_DELEGATE_TwoParams(FOnDoubleClicked, FHandle, TArrayView<const FFieldVariant>);
+	DECLARE_DELEGATE_RetVal_FourParams(FReply, FOnDragDetected, const FGeometry&, const FPointerEvent&, FHandle, TArrayView<const FFieldVariant>);
 	DECLARE_DELEGATE_RetVal_TwoParams(TSharedRef<SWidget>, FOnGenerateContainer, FHandle, TOptional<FText> DisplayName);
 
 	SLATE_BEGIN_ARGS(SPropertyViewer)
@@ -109,6 +110,8 @@ public:
 		SLATE_EVENT(FOnDoubleClicked, OnDoubleClicked);
 		/** Delegate to invoke when we generate the entry for an added container. */
 		SLATE_EVENT(FOnGenerateContainer, OnGenerateContainer)
+		/** Delegate to invoke when an item is dragged. */
+		SLATE_EVENT(FOnDragDetected, OnDragDetected);
 	SLATE_END_ARGS()
 
 	void Construct(const FArguments& InArgs);
