MVVM: drag from viewmodel to Bind widget in the details view to create bindings.

diff --git a/Engine/Plugins/Runtime/ModelViewViewModel/Source/ModelViewViewModelEditor/Private/Customizations/MVVMPropertyBindingExtension.cpp b/Engine/Plugins/Runtime/ModelViewViewModel/Source/ModelViewViewModelEditor/Private/Customizations/MVVMPropertyBindingExtension.cpp
--- a/Engine/Plugins/Runtime/ModelViewViewModel/Source/ModelViewViewModelEditor/Private/Customizations/MVVMPropertyBindingExtension.cpp
+++ b/Engine/Plugins/Runtime/ModelViewViewModel/Source/ModelViewViewModelEditor/Private/Customizations/MVVMPropertyBindingExtension.cpp
@@ -8,14 +8,20 @@
 #include "Kismet2/BlueprintEditorUtils.h"
 #include "MVVMBlueprintView.h"
 #include "MVVMDeveloperProjectSettings.h"
+#include "MVVMEditorSubsystem.h"
 #include "MVVMWidgetBlueprintExtension_View.h"
+#include "PropertyHandle.h"
+#include "PropertyPathHelpers.h"
 #include "Styling/StyleColors.h"
+#include "Widgets/ViewModelFieldDragDropOp.h"
 #include "Widgets/Images/SImage.h"
 #include "Widgets/Text/STextBlock.h"
 
 #define LOCTEXT_NAMESPACE "MVVMPropertyBindingExtension"
 
-static void ExtendBindingsMenu(FMenuBuilder& MenuBuilder, const UWidgetBlueprint* WidgetBlueprint, const UWidget* Widget, const FProperty* Property)
+namespace UE::MVVM
+{
+static void ExtendBindingsMenu(FMenuBuilder& MenuBuilder, const UWidgetBlueprint* WidgetBlueprint, UWidget* Widget, TSharedPtr<IPropertyHandle> WidgetPropertyHandle)
 {
 	MenuBuilder.BeginSection("ViewModels", LOCTEXT("ViewModels", "View Models"));
 
@@ -75,15 +81,29 @@ static void ExtendBindingsMenu(FMenuBuilder& MenuBuilder, const UWidgetBlueprint
 		return HorizontalBox;
 	};
 
-	auto CreateBinding = [MVVMBlueprintView](const UWidget* Widget, const FProperty* WidgetProperty, FGuid ViewModelId, const FProperty* ViewModelProperty)
+	auto CreateBinding = [MVVMBlueprintView](UWidget* Widget, TSharedPtr<IPropertyHandle> WidgetPropertyHandle, FGuid ViewModelId, const FProperty* ViewModelProperty)
 	{
 		FMVVMBlueprintViewBinding& NewBinding = MVVMBlueprintView->AddDefaultBinding();
 
 		NewBinding.SourcePath.SetViewModelId(ViewModelId);
 		NewBinding.SourcePath.SetBasePropertyPath(UE::MVVM::FMVVMConstFieldVariant(ViewModelProperty));
 
-		NewBinding.DestinationPath.SetWidgetName(Widget->GetFName());
-		NewBinding.DestinationPath.SetBasePropertyPath(UE::MVVM::FMVVMConstFieldVariant(WidgetProperty));
+		// Generate the destination path from the widget property that we are dropping on.
+		FCachedPropertyPath CachedPropertyPath(WidgetPropertyHandle->GeneratePathToProperty());
+		CachedPropertyPath.Resolve(Widget);
+
+		// Set the destination path.
+		FMVVMBlueprintPropertyPath DestinationPropertyPath;
+		DestinationPropertyPath.ResetBasePropertyPath();
+
+		for (int32 SegNum = 0; SegNum < CachedPropertyPath.GetNumSegments(); SegNum++)
+		{
+			FFieldVariant Field = CachedPropertyPath.GetSegment(SegNum).GetField();
+			DestinationPropertyPath.AppendBasePropertyPath(UE::MVVM::FMVVMConstFieldVariant(Field));
+		}
+
+		DestinationPropertyPath.SetWidgetName(Widget->GetFName());
+		NewBinding.DestinationPath = DestinationPropertyPath;
 
 		NewBinding.BindingType = EMVVMBindingMode::OneWayToDestination;
 	};
@@ -97,9 +117,10 @@ static void ExtendBindingsMenu(FMenuBuilder& MenuBuilder, const UWidgetBlueprint
 		}
 
 		MenuBuilder.AddSubMenu(ViewModel.GetDisplayName(), ViewModel.GetDisplayName(),
-			FNewMenuDelegate::CreateLambda([ViewModel, Widget, Property, Schema, CreatePropertyWidget, CreateBinding](FMenuBuilder& MenuBuilder)
+			FNewMenuDelegate::CreateLambda([ViewModel, Widget, WidgetPropertyHandle, Schema, CreatePropertyWidget, CreateBinding](FMenuBuilder& MenuBuilder)
 				{
 					const UClass* ViewModelClass = ViewModel.GetViewModelClass();
+					const FProperty* Property = WidgetPropertyHandle->GetProperty();
 
 					MenuBuilder.BeginSection("ValidDataTypes", LOCTEXT("ValidDataTypes", "Valid Data Types"));
 
@@ -112,7 +133,7 @@ static void ExtendBindingsMenu(FMenuBuilder& MenuBuilder, const UWidgetBlueprint
 						}
 
 						FUIAction UIAction;
-						UIAction.ExecuteAction = FExecuteAction::CreateLambda(CreateBinding, Widget, Property, ViewModel.GetViewModelId(), VMProperty);
+						UIAction.ExecuteAction = FExecuteAction::CreateLambda(CreateBinding, Widget, WidgetPropertyHandle, ViewModel.GetViewModelId(), VMProperty);
 						MenuBuilder.AddMenuEntry(
 							UIAction,
 							CreatePropertyWidget(VMProperty, false)
@@ -132,7 +153,7 @@ static void ExtendBindingsMenu(FMenuBuilder& MenuBuilder, const UWidgetBlueprint
 						}
 
 						FUIAction UIAction;
-						UIAction.ExecuteAction = FExecuteAction::CreateLambda(CreateBinding, Widget, Property, ViewModel.GetViewModelId(), VMProperty);
+						UIAction.ExecuteAction = FExecuteAction::CreateLambda(CreateBinding, Widget, WidgetPropertyHandle, ViewModel.GetViewModelId(), VMProperty);
 						MenuBuilder.AddMenuEntry(
 							UIAction,
 							CreatePropertyWidget(VMProperty, true)
@@ -217,10 +238,10 @@ void FMVVMPropertyBindingExtension::ClearCurrentValue(const UWidgetBlueprint* Wi
 	}
 }
 
-TSharedPtr<FExtender> FMVVMPropertyBindingExtension::CreateMenuExtender(const UWidgetBlueprint* WidgetBlueprint, const UWidget* Widget, const FProperty* Property)
+TSharedPtr<FExtender> FMVVMPropertyBindingExtension::CreateMenuExtender(const UWidgetBlueprint* WidgetBlueprint, UWidget* Widget, TSharedPtr<IPropertyHandle> WidgetPropertyHandle)
 {
 	TSharedPtr<FExtender> Extender = MakeShared<FExtender>();
-	Extender->AddMenuExtension("BindingActions", EExtensionHook::After, nullptr, FMenuExtensionDelegate::CreateStatic(&ExtendBindingsMenu, WidgetBlueprint, Widget, Property));
+	Extender->AddMenuExtension("BindingActions", EExtensionHook::After, nullptr, FMenuExtensionDelegate::CreateStatic(&ExtendBindingsMenu, WidgetBlueprint, Widget, WidgetPropertyHandle));
 	return Extender;
 }
 
@@ -240,4 +261,66 @@ bool FMVVMPropertyBindingExtension::CanExtend(const UWidgetBlueprint* WidgetBlue
 	return MVVMBlueprintView ? MVVMBlueprintView->GetViewModels().Num() > 0 : false;
 }
 
+IPropertyBindingExtension::EDropResult FMVVMPropertyBindingExtension::OnDrop(const FGeometry& Geometry, const FDragDropEvent& DragDropEvent, UWidgetBlueprint* WidgetBlueprint, UWidget* Widget, TSharedPtr<IPropertyHandle> WidgetPropertyHandle)
+{
+	UMVVMWidgetBlueprintExtension_View* MVVMExtensionPtr = UMVVMWidgetBlueprintExtension_View::GetExtension<UMVVMWidgetBlueprintExtension_View>(WidgetBlueprint);
+	if (MVVMExtensionPtr == nullptr)
+	{
+		return EDropResult::Unhandled;
+	}
+	UMVVMBlueprintView* MVVMBlueprintView = MVVMExtensionPtr->GetBlueprintView();
+	if (MVVMBlueprintView == nullptr)
+	{
+		return EDropResult::Unhandled;
+	}
+
+	if (TSharedPtr<UE::MVVM::FViewModelFieldDragDropOp> ViewModelFieldDragDropOp = DragDropEvent.GetOperationAs<UE::MVVM::FViewModelFieldDragDropOp>())
+	{
+		UMVVMEditorSubsystem* EditorSubsystem = GEditor->GetEditorSubsystem<UMVVMEditorSubsystem>();
+		FMVVMBlueprintViewBinding& NewBinding = EditorSubsystem->AddBinding(WidgetBlueprint);
+
+		TArray<FFieldVariant> SourceFieldPath = ViewModelFieldDragDropOp->DraggedField;
+
+		// Set the source path (view model property from the drop event).
+		FMVVMBlueprintPropertyPath SourcePropertyPath;
+		SourcePropertyPath.ResetBasePropertyPath();
+		for (const FFieldVariant& Field : SourceFieldPath)
+		{
+			SourcePropertyPath.AppendBasePropertyPath(UE::MVVM::FMVVMConstFieldVariant(Field));
+		}
+		if (ViewModelFieldDragDropOp->ViewModelId.IsValid())
+		{
+			SourcePropertyPath.SetViewModelId(ViewModelFieldDragDropOp->ViewModelId);
+		}
+
+		NewBinding.SourcePath = SourcePropertyPath;
+		EditorSubsystem->SetSourcePathForBinding(WidgetBlueprint, NewBinding, SourcePropertyPath);
+
+		// Generate the destination path from the widget property that we are dropping on.
+		FCachedPropertyPath CachedPropertyPath(WidgetPropertyHandle->GeneratePathToProperty());
+		CachedPropertyPath.Resolve(Widget);
+
+		// Set the destination path.
+		FMVVMBlueprintPropertyPath DestinationPropertyPath;
+		DestinationPropertyPath.ResetBasePropertyPath();
+
+		for (int32 SegNum = 0; SegNum < CachedPropertyPath.GetNumSegments(); SegNum++)
+		{
+			FFieldVariant Field = CachedPropertyPath.GetSegment(SegNum).GetField();
+			DestinationPropertyPath.AppendBasePropertyPath(UE::MVVM::FMVVMConstFieldVariant(Field));
+		}
+
+		DestinationPropertyPath.SetWidgetName(Widget->GetFName());
+		EditorSubsystem->SetDestinationPathForBinding(WidgetBlueprint, NewBinding, DestinationPropertyPath);
+
+		return EDropResult::HandledContinue;
+	}
+	return EDropResult::Unhandled;
+}
+
+TSharedPtr<FExtender> FMVVMPropertyBindingExtension::CreateMenuExtender(const UWidgetBlueprint* WidgetBlueprint, const UWidget* Widget, const FProperty* Property)
+{
+	return MakeShared<FExtender>();
+}
+}
 #undef LOCTEXT_NAMESPACE
diff --git a/Engine/Plugins/Runtime/ModelViewViewModel/Source/ModelViewViewModelEditor/Private/Customizations/MVVMPropertyBindingExtension.h b/Engine/Plugins/Runtime/ModelViewViewModel/Source/ModelViewViewModelEditor/Private/Customizations/MVVMPropertyBindingExtension.h
--- a/Engine/Plugins/Runtime/ModelViewViewModel/Source/ModelViewViewModelEditor/Private/Customizations/MVVMPropertyBindingExtension.h
+++ b/Engine/Plugins/Runtime/ModelViewViewModel/Source/ModelViewViewModelEditor/Private/Customizations/MVVMPropertyBindingExtension.h
@@ -7,6 +7,9 @@
 class UWidgetBlueprint;
 class UWidget;
 
+namespace UE::MVVM
+{
+
 class FMVVMPropertyBindingExtension
 	: public IPropertyBindingExtension
 {
@@ -14,6 +17,10 @@ class FMVVMPropertyBindingExtension
 	virtual const FSlateBrush* GetCurrentIcon(const UWidgetBlueprint* WidgetBlueprint, const UWidget* Widget, const FProperty* Property) const override;
 
 	virtual void ClearCurrentValue(const UWidgetBlueprint* WidgetBlueprint, const UWidget* Widget, const FProperty* Property) override;
-	virtual TSharedPtr<FExtender> CreateMenuExtender(const UWidgetBlueprint* WidgetBlueprint, const UWidget* Widget, const FProperty* Property) override;
+	virtual TSharedPtr<FExtender> CreateMenuExtender(const UWidgetBlueprint* WidgetBlueprint, UWidget* Widget, TSharedPtr<IPropertyHandle> WidgetPropertyHandle) override;
 	virtual bool CanExtend(const UWidgetBlueprint* WidgetBlueprint, const UWidget* Widget, const FProperty* Property) const;
+	virtual EDropResult OnDrop(const FGeometry& Geometry, const FDragDropEvent& DragDropEvent, UWidgetBlueprint* WidgetBlueprint, UWidget* Widget, TSharedPtr<IPropertyHandle> WidgetPropertyHandle) override;
+	virtual TSharedPtr<FExtender> CreateMenuExtender(const UWidgetBlueprint* WidgetBlueprint, const UWidget* Widget, const FProperty* Property) override;
 };
+
+}
diff --git a/Engine/Plugins/Runtime/ModelViewViewModel/Source/ModelViewViewModelEditor/Private/ModelViewViewModelEditorModule.cpp b/Engine/Plugins/Runtime/ModelViewViewModel/Source/ModelViewViewModelEditor/Private/ModelViewViewModelEditorModule.cpp
--- a/Engine/Plugins/Runtime/ModelViewViewModel/Source/ModelViewViewModelEditor/Private/ModelViewViewModelEditorModule.cpp
+++ b/Engine/Plugins/Runtime/ModelViewViewModel/Source/ModelViewViewModelEditor/Private/ModelViewViewModelEditorModule.cpp
@@ -39,7 +39,7 @@ void FModelViewViewModelEditorModule::StartupModule()
 	IUMGEditorModule& UMGEditorModule = FModuleManager::LoadModuleChecked<IUMGEditorModule>("UMGEditor");
 	UMGEditorModule.OnRegisterTabsForEditor().AddRaw(this, &FModelViewViewModelEditorModule::HandleRegisterBlueprintEditorTab);
 
-	PropertyBindingExtension = MakeShared<FMVVMPropertyBindingExtension>();
+	PropertyBindingExtension = MakeShared<UE::MVVM::FMVVMPropertyBindingExtension>();
 	UMGEditorModule.GetPropertyBindingExtensibilityManager()->AddExtension(PropertyBindingExtension.ToSharedRef());
 
 	UMGEditorModule.RegisterInstancedCustomPropertyTypeLayout(
diff --git a/Engine/Plugins/Runtime/ModelViewViewModel/Source/ModelViewViewModelEditor/Private/ModelViewViewModelEditorModule.h b/Engine/Plugins/Runtime/ModelViewViewModel/Source/ModelViewViewModelEditor/Private/ModelViewViewModelEditorModule.h
--- a/Engine/Plugins/Runtime/ModelViewViewModel/Source/ModelViewViewModelEditor/Private/ModelViewViewModelEditorModule.h
+++ b/Engine/Plugins/Runtime/ModelViewViewModel/Source/ModelViewViewModelEditor/Private/ModelViewViewModelEditorModule.h
@@ -8,10 +8,10 @@
 namespace UE::MVVM
 {
 	class FAssetTypeActions_ViewModelBlueprint;
+	class FMVVMPropertyBindingExtension;
 }
 
 class FMVVMBindPropertiesDetailView;
-class FMVVMPropertyBindingExtension;
 class FWidgetBlueprintApplicationMode;
 class FWorkflowAllowedTabSet;
 class UBlueprint;
@@ -42,6 +42,6 @@ private:
 	void UnregisterMenus();
 
 private:
-	TSharedPtr<FMVVMPropertyBindingExtension> PropertyBindingExtension;
+	TSharedPtr<UE::MVVM::FMVVMPropertyBindingExtension> PropertyBindingExtension;
 	TSharedPtr<UE::MVVM::FAssetTypeActions_ViewModelBlueprint> ViewModelBlueprintActions;
 };
diff --git a/Engine/Plugins/Runtime/PropertyAccess/Source/PropertyAccessEditor/Private/SPropertyBinding.cpp b/Engine/Plugins/Runtime/PropertyAccess/Source/PropertyAccessEditor/Private/SPropertyBinding.cpp
--- a/Engine/Plugins/Runtime/PropertyAccess/Source/PropertyAccessEditor/Private/SPropertyBinding.cpp
+++ b/Engine/Plugins/Runtime/PropertyAccess/Source/PropertyAccessEditor/Private/SPropertyBinding.cpp
@@ -1037,4 +1037,14 @@ FReply SPropertyBinding::HandleGotoBindingClicked()
 	return FReply::Unhandled();
 }
 
+FReply SPropertyBinding::OnDrop(const FGeometry& MyGeometry, const FDragDropEvent& DragDropEvent)
+{
+	if (Args.OnDrop.IsBound())
+	{
+		return Args.OnDrop.Execute(MyGeometry, DragDropEvent);
+	}
+
+	return FReply::Unhandled();
+}
+
 #undef LOCTEXT_NAMESPACE
diff --git a/Engine/Plugins/Runtime/PropertyAccess/Source/PropertyAccessEditor/Private/SPropertyBinding.h b/Engine/Plugins/Runtime/PropertyAccess/Source/PropertyAccessEditor/Private/SPropertyBinding.h
--- a/Engine/Plugins/Runtime/PropertyAccess/Source/PropertyAccessEditor/Private/SPropertyBinding.h
+++ b/Engine/Plugins/Runtime/PropertyAccess/Source/PropertyAccessEditor/Private/SPropertyBinding.h
@@ -66,6 +66,8 @@ protected:
 
 	FReply HandleGotoBindingClicked();
 
+	virtual FReply OnDrop(const FGeometry& MyGeometry, const FDragDropEvent& DragDropEvent) override;
+
 private:
 	bool IsClassDenied(UClass* OwnerClass) const;
 	bool IsFieldFromDeniedClass(FFieldVariant Field) const;
diff --git a/Engine/Source/Editor/UMGEditor/Private/Customizations/UMGDetailCustomizations.cpp b/Engine/Source/Editor/UMGEditor/Private/Customizations/UMGDetailCustomizations.cpp
--- a/Engine/Source/Editor/UMGEditor/Private/Customizations/UMGDetailCustomizations.cpp
+++ b/Engine/Source/Editor/UMGEditor/Private/Customizations/UMGDetailCustomizations.cpp
@@ -140,7 +140,7 @@ TSharedRef<SWidget> FBlueprintWidgetCustomization::MakePropertyBindingWidget(TWe
 	{
 		if (Extension->CanExtend(WidgetBlueprint, Widget, InPropertyHandle->GetProperty()))
 		{
-			MenuExtenders.Add(Extension->CreateMenuExtender(WidgetBlueprint, Widget, InPropertyHandle->GetProperty()));
+			MenuExtenders.Add(Extension->CreateMenuExtender(WidgetBlueprint, Widget, InPropertyHandle));
 			ActiveExtensions.Add(Extension);
 		}
 	}
@@ -450,6 +450,30 @@ TSharedRef<SWidget> FBlueprintWidgetCustomization::MakePropertyBindingWidget(TWe
 		return false;
 	});
 
+	Args.OnDrop = FOnDrop::CreateLambda([ActiveExtensions, InEditor, Objects, InPropertyHandle](const FGeometry& Geometry, const FDragDropEvent& DragDropEvent)
+	{
+		bool bDropHandled = false;
+		UWidgetBlueprint* WidgetBlueprint = InEditor.Pin()->GetWidgetBlueprintObj();
+		UWidget* Widget = Objects.Num() ? Cast<UWidget>(Objects[0]) : nullptr;
+
+		if (Widget && WidgetBlueprint)
+		{
+			for (const TSharedPtr<IPropertyBindingExtension>& Extension : ActiveExtensions)
+			{
+				IPropertyBindingExtension::EDropResult Result = Extension->OnDrop(Geometry, DragDropEvent, WidgetBlueprint, Widget, InPropertyHandle);
+				if (Result != IPropertyBindingExtension::EDropResult::Unhandled)
+				{
+					bDropHandled = true;
+				}
+				if (Result == IPropertyBindingExtension::EDropResult::HandledBreak)
+				{
+					break;
+				}
+			}
+		}
+		return bDropHandled ? FReply::Handled() : FReply::Unhandled();
+	});
+
 	Args.CurrentBindingText = MakeAttributeLambda([InEditor, Objects, InPropertyHandle, ActiveExtensions]()
 	{
 		UWidgetBlueprint* ThisBlueprint = InEditor.Pin()->GetWidgetBlueprintObj();
diff --git a/Engine/Source/Editor/UMGEditor/Public/IHasPropertyBindingExtensibility.h b/Engine/Source/Editor/UMGEditor/Public/IHasPropertyBindingExtensibility.h
--- a/Engine/Source/Editor/UMGEditor/Public/IHasPropertyBindingExtensibility.h
+++ b/Engine/Source/Editor/UMGEditor/Public/IHasPropertyBindingExtensibility.h
@@ -4,9 +4,11 @@
 
 #include "Containers/Array.h"
 #include "Framework/MultiBox/MultiBoxExtender.h"
+#include "Input/DragAndDrop.h"
 #include "Misc/Optional.h"
 #include "Templates/SharedPointer.h"
 
+class IPropertyHandle;
 class UWidget;
 class UWidgetBlueprint;
 struct FSlateBrush;
@@ -14,12 +16,23 @@ struct FSlateBrush;
 class UMGEDITOR_API IPropertyBindingExtension
 {
 public:
+	enum class EDropResult
+	{
+		HandledContinue,	// The drop event was handled by the OnDrop call and it can still be re-used for OnDrop calls in other extensions.
+		HandledBreak,		// The drop event was handled by the OnDrop call but it is not reusable for OnDrop calls in other extensions (e.g. it was modified).
+		Unhandled,			// The drop event was not handled by this extension.
+	};
+
 	/** Does this extension want to extend this property in the widget? */
 	virtual bool CanExtend(const UWidgetBlueprint* WidgetBlueprint, const UWidget* Widget, const FProperty* Property) const = 0;
-	virtual TSharedPtr<FExtender> CreateMenuExtender(const UWidgetBlueprint* WidgetBlueprint, const UWidget* Widget, const FProperty* Property) = 0;
+	virtual TSharedPtr<FExtender> CreateMenuExtender(const UWidgetBlueprint* WidgetBlueprint, UWidget* Widget, TSharedPtr<IPropertyHandle> WidgetPropertyHandle) = 0;
 	virtual void ClearCurrentValue(const UWidgetBlueprint* WidgetBlueprint, const UWidget* Widget, const FProperty* Property) = 0;
 	virtual TOptional<FName> GetCurrentValue(const UWidgetBlueprint* WidgetBlueprint, const UWidget* Widget, const FProperty* Property) const = 0;
 	virtual const FSlateBrush* GetCurrentIcon(const UWidgetBlueprint* WidgetBlueprint, const UWidget* Widget, const FProperty* Property) const = 0;
+	virtual EDropResult OnDrop(const FGeometry& Geometry, const FDragDropEvent& DragDropEvent, UWidgetBlueprint* WidgetBlueprint, UWidget* Widget, TSharedPtr<IPropertyHandle> WidgetPropertyHandle) = 0;
+
+	UE_DEPRECATED(5.3, "The function CreateMenuExtender with FProperty parameter is deprecated. Please pass the property handler as the parameter instead.")
+	virtual TSharedPtr<FExtender> CreateMenuExtender(const UWidgetBlueprint* WidgetBlueprint, const UWidget* Widget, const FProperty* Property) = 0;
 };
 
 /**
diff --git a/Engine/Source/Editor/UnrealEd/Public/IPropertyAccessEditor.h b/Engine/Source/Editor/UnrealEd/Public/IPropertyAccessEditor.h
--- a/Engine/Source/Editor/UnrealEd/Public/IPropertyAccessEditor.h
+++ b/Engine/Source/Editor/UnrealEd/Public/IPropertyAccessEditor.h
@@ -134,6 +134,9 @@ DECLARE_DELEGATE_RetVal_OneParam(bool, FOnCanRemoveBinding, FName /*InPropertyNa
 /** Delegate called once a new function binding has been created */
 DECLARE_DELEGATE_TwoParams(FOnNewFunctionBindingCreated, UEdGraph* /*InFunctionGraph*/, UFunction* /*InFunction*/);
 
+/** Delegate called once a drag-drop event is dropped on the binding widget */
+DECLARE_DELEGATE_RetVal_TwoParams(FReply, FOnDrop, const FGeometry&, const FDragDropEvent&);
+
 /** Setup arguments structure for a property binding widget */
 struct FPropertyBindingWidgetArgs
 {
@@ -182,6 +185,9 @@ struct FPropertyBindingWidgetArgs
 	/** Delegate called once a new function binding has been created */
 	FOnNewFunctionBindingCreated OnNewFunctionBindingCreated;
 	
+	/** Delegate called when a property is dropped on the property binding widget */
+	FOnDrop OnDrop;
+
 	/** The current binding's text label */
 	TAttribute<FText> CurrentBindingText;
 
