MVVM: Test if the Binding library is loaded before accessing it.

diff --git a/Engine/Plugins/Runtime/ModelViewViewModel/Source/ModelViewViewModel/Private/Bindings/MVVMCompiledBindingLibrary.cpp b/Engine/Plugins/Runtime/ModelViewViewModel/Source/ModelViewViewModel/Private/Bindings/MVVMCompiledBindingLibrary.cpp
--- a/Engine/Plugins/Runtime/ModelViewViewModel/Source/ModelViewViewModel/Private/Bindings/MVVMCompiledBindingLibrary.cpp
+++ b/Engine/Plugins/Runtime/ModelViewViewModel/Source/ModelViewViewModel/Private/Bindings/MVVMCompiledBindingLibrary.cpp
@@ -116,6 +116,12 @@ void FMVVMCompiledBindingLibrary::Load()
 }
 
 
+bool FMVVMCompiledBindingLibrary::IsLoaded() const
+{
+	return LoadedProperties.Num() > 0 || LoadedFunctions.Num() > 0 || LoadedFieldIds.Num() > 0 || CompiledFields.Num() == 0;
+}
+
+
 void FMVVMCompiledBindingLibrary::Unload()
 {
 	LoadedProperties.Empty();
@@ -479,7 +485,6 @@ TValueOrError<FString, FString> FMVVMCompiledBindingLibrary::FieldPathToString(F
 			}
 			else
 			{
-				ensureMsgf(false, TEXT("The library is not loaded"));
 				StringBuilder << TEXT("<Invalid>");
 				bHasError = true;
 			}
@@ -510,7 +515,6 @@ TValueOrError<FString, FString> FMVVMCompiledBindingLibrary::FieldPathToString(F
 			}
 			else
 			{
-				ensureMsgf(false, TEXT("The library is not loaded"));
 				StringBuilder << TEXT("<Invalid>");
 				bHasError = true;
 			}
@@ -540,8 +544,11 @@ TValueOrError<UE::FieldNotification::FFieldId, void> FMVVMCompiledBindingLibrary
 	}
 #endif
 
-	check(LoadedFieldIds.IsValidIndex(InFieldId.FieldIdIndex));
-	return MakeValue(LoadedFieldIds[InFieldId.FieldIdIndex]);
+	if (LoadedFieldIds.IsValidIndex(InFieldId.FieldIdIndex))
+	{
+		return MakeValue(LoadedFieldIds[InFieldId.FieldIdIndex]);
+	}
+	return MakeError<>();
 }
 
 #undef LOCTEXT_NAMESPACE
diff --git a/Engine/Plugins/Runtime/ModelViewViewModel/Source/ModelViewViewModel/Public/Bindings/MVVMCompiledBindingLibrary.h b/Engine/Plugins/Runtime/ModelViewViewModel/Source/ModelViewViewModel/Public/Bindings/MVVMCompiledBindingLibrary.h
--- a/Engine/Plugins/Runtime/ModelViewViewModel/Source/ModelViewViewModel/Public/Bindings/MVVMCompiledBindingLibrary.h
+++ b/Engine/Plugins/Runtime/ModelViewViewModel/Source/ModelViewViewModel/Public/Bindings/MVVMCompiledBindingLibrary.h
@@ -302,6 +302,8 @@ public:
 	void Load();
 	/** Release the acquired FProperty and UFunction. */
 	void Unload();
+	/** FProperty and UFunction are fetched. */
+	bool IsLoaded() const;
 
 	/**
 	 * Execute a binding, in one direction.
diff --git a/Engine/Plugins/Runtime/ModelViewViewModel/Source/ModelViewViewModelDebugger/Private/MVVMDebugSnapshot.cpp b/Engine/Plugins/Runtime/ModelViewViewModel/Source/ModelViewViewModelDebugger/Private/MVVMDebugSnapshot.cpp
--- a/Engine/Plugins/Runtime/ModelViewViewModel/Source/ModelViewViewModelDebugger/Private/MVVMDebugSnapshot.cpp
+++ b/Engine/Plugins/Runtime/ModelViewViewModel/Source/ModelViewViewModelDebugger/Private/MVVMDebugSnapshot.cpp
@@ -36,25 +36,28 @@ FMVVMViewBindingDebugEntry CreateViewBindingDebugEntry(const UMVVMViewClass* Vie
 #if WITH_EDITOR
 	Result.BlueprintViewBindingId = Binding.GetEditorId();
 #endif
-	if (Binding.GetSourceFieldId().IsValid())
+	if (ViewClass->GetBindingLibrary().IsLoaded())
 	{
-		TValueOrError<UE::FieldNotification::FFieldId, void> FieldIdValue = ViewClass->GetBindingLibrary().GetFieldId(Binding.GetSourceFieldId());
-		Result.FieldId = FieldIdValue.HasValue() ? FieldIdValue.GetValue().GetName() : FName();
-	}
-	if (Binding.GetBinding().GetSourceFieldPath().IsValid())
-	{
-		TValueOrError<FString, FString> SourceFieldPathValue = ViewClass->GetBindingLibrary().FieldPathToString(Binding.GetBinding().GetSourceFieldPath(), false);
-		Result.SourceFieldPath = SourceFieldPathValue.HasValue() ? SourceFieldPathValue.StealValue() : SourceFieldPathValue.StealError();
-	}
-	if (Binding.GetBinding().GetDestinationFieldPath().IsValid())
-	{
-		TValueOrError<FString, FString> DestinationFieldPathValue = ViewClass->GetBindingLibrary().FieldPathToString(Binding.GetBinding().GetDestinationFieldPath(), false);
-		Result.DestinationFieldPath = DestinationFieldPathValue.HasValue() ? DestinationFieldPathValue.StealValue() : DestinationFieldPathValue.StealError();
-	}
-	if (Binding.GetBinding().GetConversionFunctionFieldPath().IsValid())
-	{
-		TValueOrError<FString, FString> ConversionFieldPathValue = ViewClass->GetBindingLibrary().FieldPathToString(Binding.GetBinding().GetConversionFunctionFieldPath(), false);
-		Result.ConversionFunctionFieldPath = ConversionFieldPathValue.HasValue() ? ConversionFieldPathValue.StealValue() : ConversionFieldPathValue.StealError();
+		if (Binding.GetSourceFieldId().IsValid())
+		{
+			TValueOrError<UE::FieldNotification::FFieldId, void> FieldIdValue = ViewClass->GetBindingLibrary().GetFieldId(Binding.GetSourceFieldId());
+			Result.FieldId = FieldIdValue.HasValue() ? FieldIdValue.GetValue().GetName() : FName();
+		}
+		if (Binding.GetBinding().GetSourceFieldPath().IsValid())
+		{
+			TValueOrError<FString, FString> SourceFieldPathValue = ViewClass->GetBindingLibrary().FieldPathToString(Binding.GetBinding().GetSourceFieldPath(), false);
+			Result.SourceFieldPath = SourceFieldPathValue.HasValue() ? SourceFieldPathValue.StealValue() : SourceFieldPathValue.StealError();
+		}
+		if (Binding.GetBinding().GetDestinationFieldPath().IsValid())
+		{
+			TValueOrError<FString, FString> DestinationFieldPathValue = ViewClass->GetBindingLibrary().FieldPathToString(Binding.GetBinding().GetDestinationFieldPath(), false);
+			Result.DestinationFieldPath = DestinationFieldPathValue.HasValue() ? DestinationFieldPathValue.StealValue() : DestinationFieldPathValue.StealError();
+		}
+		if (Binding.GetBinding().GetConversionFunctionFieldPath().IsValid())
+		{
+			TValueOrError<FString, FString> ConversionFieldPathValue = ViewClass->GetBindingLibrary().FieldPathToString(Binding.GetBinding().GetConversionFunctionFieldPath(), false);
+			Result.ConversionFunctionFieldPath = ConversionFieldPathValue.HasValue() ? ConversionFieldPathValue.StealValue() : ConversionFieldPathValue.StealError();
+		}
 	}
 	Result.CompiledBindingIndex = Index;
 	return Result;
