MVVM: Fix displayname for viewmodel resolver.

diff --git a/Engine/Plugins/Runtime/ModelViewViewModel/Source/ModelViewViewModel/Public/View/MVVMViewModelContextResolver.h b/Engine/Plugins/Runtime/ModelViewViewModel/Source/ModelViewViewModel/Public/View/MVVMViewModelContextResolver.h
--- a/Engine/Plugins/Runtime/ModelViewViewModel/Source/ModelViewViewModel/Public/View/MVVMViewModelContextResolver.h
+++ b/Engine/Plugins/Runtime/ModelViewViewModel/Source/ModelViewViewModel/Public/View/MVVMViewModelContextResolver.h
@@ -14,7 +14,7 @@ class UUserWidget;
 /**
  * Shared data to find or create a ViewModel at runtime.
  */
-UCLASS(Abstract, Blueprintable, EditInlineNew, DisplayName = "Viewmode Resolver")
+UCLASS(Abstract, Blueprintable, EditInlineNew, DisplayName = "Viewmodel Resolver")
 class MODELVIEWVIEWMODEL_API UMVVMViewModelContextResolver : public UObject
 {
 	GENERATED_BODY()
