MVVM: Change the displayname of some Viewmodel properties.

diff --git a/Engine/Plugins/Runtime/ModelViewViewModel/Source/ModelViewViewModelBlueprint/Public/MVVMBlueprintViewModelContext.h b/Engine/Plugins/Runtime/ModelViewViewModel/Source/ModelViewViewModelBlueprint/Public/MVVMBlueprintViewModelContext.h
--- a/Engine/Plugins/Runtime/ModelViewViewModel/Source/ModelViewViewModelBlueprint/Public/MVVMBlueprintViewModelContext.h
+++ b/Engine/Plugins/Runtime/ModelViewViewModel/Source/ModelViewViewModelBlueprint/Public/MVVMBlueprintViewModelContext.h
@@ -89,7 +89,7 @@ public:
 
 private:
 	/** When the view is spawn, create an instance of the viewmodel. */
-	UPROPERTY(VisibleAnywhere, AdvancedDisplay, Category = "Viewmodel")
+	UPROPERTY(VisibleAnywhere, AdvancedDisplay, Category = "Viewmodel", meta = (DisplayName = "Viewmodel Context Id"))
 	FGuid ViewModelContextId;
 
 	UPROPERTY(VisibleAnywhere, Category = "Viewmodel", NoClear, meta = (AllowedClasses = "/Script/UMG.NotifyFieldValueChanged", DisallowedClasses = "/Script/UMG.Widget"))
@@ -103,7 +103,7 @@ private:
 
 public:
 	/** Property name that will be generated. */
-	UPROPERTY(EditAnywhere, Category = "Viewmodel")
+	UPROPERTY(EditAnywhere, Category = "Viewmodel", meta = (DisplayName = "Viewmodel Name"))
 	FName ViewModelName;
 
 	/** When the view is spawn, create an instance of the viewmodel. */
@@ -111,11 +111,11 @@ public:
 	EMVVMBlueprintViewModelContextCreationType CreationType = EMVVMBlueprintViewModelContextCreationType::CreateInstance;
 
 	/** Identifier of an already registered viewmodel. */
-	UPROPERTY(EditAnywhere, Category = "Viewmodel", AdvancedDisplay, meta = (EditCondition = "CreationType == EMVVMBlueprintViewModelContextCreationType::GlobalViewModelCollection", EditConditionHides))
+	UPROPERTY(EditAnywhere, Category = "Viewmodel", AdvancedDisplay, meta = (DisplayName = "Global Viewmodel Identifier", EditCondition = "CreationType == EMVVMBlueprintViewModelContextCreationType::GlobalViewModelCollection", EditConditionHides))
 	FName GlobalViewModelIdentifier;
 
 	/** The Path to get the viewmodel instance. */
-	UPROPERTY(EditAnywhere, Category = "Viewmodel", AdvancedDisplay, meta = (EditCondition = "CreationType == EMVVMBlueprintViewModelContextCreationType::PropertyPath", EditConditionHides))
+	UPROPERTY(EditAnywhere, Category = "Viewmodel", AdvancedDisplay, meta = (DisplayName = "Viewmodel Property Path", EditCondition = "CreationType == EMVVMBlueprintViewModelContextCreationType::PropertyPath", EditConditionHides))
 	FString ViewModelPropertyPath;
 
 	UPROPERTY(EditAnywhere, Category = "Viewmodel", Instanced, AdvancedDisplay, meta = (EditInline, EditCondition = "CreationType == EMVVMBlueprintViewModelContextCreationType::Resolver", EditConditionHides))
