MVVM: Save the MemberReference with self context. That allow the path to be used in the skeletal compilation phase.

diff --git a/Engine/Plugins/Runtime/ModelViewViewModel/Source/ModelViewViewModel/Private/Bindings/MVVMFieldPathHelper.cpp b/Engine/Plugins/Runtime/ModelViewViewModel/Source/ModelViewViewModel/Private/Bindings/MVVMFieldPathHelper.cpp
--- a/Engine/Plugins/Runtime/ModelViewViewModel/Source/ModelViewViewModel/Private/Bindings/MVVMFieldPathHelper.cpp
+++ b/Engine/Plugins/Runtime/ModelViewViewModel/Source/ModelViewViewModel/Private/Bindings/MVVMFieldPathHelper.cpp
@@ -9,6 +9,10 @@
 #include "Types/MVVMAvailableBinding.h"
 #include "Types/MVVMFieldContext.h"
 
+#if WITH_EDITOR
+#include "Kismet2/BlueprintEditorUtils.h"
+#endif
+
 #define LOCTEXT_NAMESPACE "MVVMFieldPathHelper"
 
 namespace UE::MVVM::FieldPathHelper
@@ -51,8 +55,21 @@ const FProperty* FindProperty(FMVVMConstFieldVariant Field)
 	return nullptr;
 }
 
+const UStruct* GetMostUpToDateStruct(const UStruct* Struct)
+{
+#if WITH_EDITOR
+	if (const UClass* Class = Cast<UClass>(Struct))
+	{
+		return FBlueprintEditorUtils::GetMostUpToDateClass(Class);
+	}
+	return Struct;
+#else
+	return Struct;
+#endif
+}
+
 
-TValueOrError<FMVVMConstFieldVariant, FText> TransformWithAccessor(UStruct* CurrentContainer, FMVVMConstFieldVariant CurrentField, bool bForReading)
+TValueOrError<FMVVMConstFieldVariant, FText> TransformWithAccessor(const UStruct* CurrentContainer, FMVVMConstFieldVariant CurrentField, bool bForReading)
 {
 #if WITH_EDITORONLY_DATA
 	if (bForReading)
@@ -206,7 +223,7 @@ TValueOrError<TArray<FMVVMConstFieldVariant>, FText> GenerateFieldPathList(TArra
 
 	TArray<FMVVMConstFieldVariant> Result;
 	Result.Reserve(InFieldPath.Num());
-	UStruct* CurrentContainer = InFieldPath[0].GetOwner();
+	const UStruct* CurrentContainer = InFieldPath[0].GetOwner();
 
 	for (int32 Index = 0; Index < InFieldPath.Num(); ++Index)
 	{
@@ -223,8 +240,10 @@ TValueOrError<TArray<FMVVMConstFieldVariant>, FText> GenerateFieldPathList(TArra
 			return MakeError(FText::Format(LOCTEXT("FieldDoesNotExist", "The field '{0}' does not exist."), FText::FromName(Field.GetName())));
 		}
 
-		const bool bIsChild = Field.GetOwner()->IsChildOf(CurrentContainer);
-		const bool bIsDownCast = CurrentContainer->IsChildOf(Field.GetOwner());
+		const UStruct* SkeletalOwner = Private::GetMostUpToDateStruct(Field.GetOwner());
+		const UStruct* SkeletalCurrentContainer = Private::GetMostUpToDateStruct(CurrentContainer);
+		const bool bIsChild = SkeletalOwner->IsChildOf(SkeletalCurrentContainer);
+		const bool bIsDownCast = SkeletalCurrentContainer->IsChildOf(SkeletalOwner);
 		if (!(bIsChild || bIsDownCast))
 		{
 			return MakeError(FText::Format(LOCTEXT("FieldDoesNotExistInStruct", "The field '{0}' does not exist in the struct '{1}'."), 
@@ -287,9 +306,13 @@ TValueOrError<FParsedNotifyBindingInfo, FText> GetNotifyBindingInfoFromFieldPath
 		return MakeError(LOCTEXT("FieldPathEmpty", "The field path is empty."));
 	}
 
-	if (!InAccessor->IsChildOf(InFieldPath[0].GetOwner()))
 	{
-		return MakeError(LOCTEXT("FieldPathDoesntStartWithAccessor", "Field path that doesn't start with the accessor is not supported."));
+		const UStruct* SkeletalFieldOwner = Private::GetMostUpToDateStruct(InFieldPath[0].GetOwner());
+		const UStruct* SkeletalAccessor = Private::GetMostUpToDateStruct(InAccessor);
+		if (!SkeletalAccessor->IsChildOf(SkeletalFieldOwner))
+		{
+			return MakeError(LOCTEXT("FieldPathDoesntStartWithAccessor", "Field path that doesn't start with the accessor is not supported."));
+		}
 	}
 
 	FParsedNotifyBindingInfo Result;
@@ -328,8 +351,10 @@ TValueOrError<FParsedNotifyBindingInfo, FText> GetNotifyBindingInfoFromFieldPath
 			return MakeError(FText::Format(LOCTEXT("InvalidFieldParent", "The field '{0}' is not a container."), Index - 1));
 		}
 
-		const bool bIsChild = FieldOwner->IsChildOf(ParentContainerType);
-		const bool bIsDownCast = ParentContainerType->IsChildOf(FieldOwner);
+		const UStruct* SkeletalFieldOwner = Private::GetMostUpToDateStruct(FieldOwner);
+		const UStruct* SkeletalParentContainer = Private::GetMostUpToDateStruct(ParentContainerType);
+		const bool bIsChild = SkeletalFieldOwner->IsChildOf(SkeletalParentContainer);
+		const bool bIsDownCast = SkeletalParentContainer->IsChildOf(SkeletalFieldOwner);
 		if (!(bIsChild || bIsDownCast))
 		{
 			return MakeError(FText::Format(LOCTEXT("FieldDoesNotExistInStruct", "The field '{0}' does not exist in the struct '{1}'."),
@@ -450,7 +475,9 @@ TValueOrError<UObject*, void> EvaluateObjectProperty(const FFieldContext& InSour
 		check(InSource.GetObjectVariant().GetUObject());
 
 		UFunction* Function = InSource.GetFieldVariant().GetFunction();
-		if (!InSource.GetObjectVariant().GetUObject()->GetClass()->IsChildOf(Function->GetOuterUClass()))
+		const UStruct* SkeletalFunctionOwner = Private::GetMostUpToDateStruct(Function->GetOuterUClass());
+		const UStruct* SkeletalObjectClass = Private::GetMostUpToDateStruct(InSource.GetObjectVariant().GetUObject()->GetClass());
+		if (!SkeletalObjectClass->IsChildOf(SkeletalFunctionOwner))
 		{
 			return MakeError();
 		}
diff --git a/Engine/Plugins/Runtime/ModelViewViewModel/Source/ModelViewViewModelBlueprint/Private/Bindings/MVVMCompiledBindingLibraryCompiler.cpp b/Engine/Plugins/Runtime/ModelViewViewModel/Source/ModelViewViewModelBlueprint/Private/Bindings/MVVMCompiledBindingLibraryCompiler.cpp
--- a/Engine/Plugins/Runtime/ModelViewViewModel/Source/ModelViewViewModelBlueprint/Private/Bindings/MVVMCompiledBindingLibraryCompiler.cpp
+++ b/Engine/Plugins/Runtime/ModelViewViewModel/Source/ModelViewViewModelBlueprint/Private/Bindings/MVVMCompiledBindingLibraryCompiler.cpp
@@ -3,7 +3,9 @@
 #include "Bindings/MVVMCompiledBindingLibraryCompiler.h"
 
 #include "Bindings/MVVMBindingHelper.h"
+#include "Engine/Blueprint.h"
 #include "Engine/Engine.h"
+#include "Kismet2/BlueprintEditorUtils.h"
 #include "MVVMDeveloperProjectSettings.h"
 #include "MVVMSubsystem.h"
 #include "Misc/TVariantMeta.h"
@@ -117,6 +119,67 @@ public:
 	}
 };
 
+/** */
+//FMVVMConstFieldVariant GetMostUpToDate(FMVVMConstFieldVariant Field)
+//{
+//	if (Field.IsValid())
+//	{
+//		UClass* Class = Cast<UClass>(Field.GetOwner());
+//		if (!Class)
+//		{
+//			return Field;
+//		}
+//		if (Field.IsProperty())
+//		{
+//			return FMVVMConstFieldVariant(FBlueprintEditorUtils::GetMostUpToDateProperty(Field.GetProperty()));
+//		}
+//		else if (Field.IsFunction())
+//		{
+//			return FMVVMConstFieldVariant(FBlueprintEditorUtils::GetMostUpToDateFunction(Field.GetFunction()));
+//		}
+//	}
+//	return Field;
+//}
+
+
+/** */
+const UStruct* GetSavedGeneratedStruct(FMVVMConstFieldVariant Field)
+{
+	const UStruct* Result = Field.GetOwner();
+	const UClass* Class = Cast<UClass>(Result);
+	if (!Class || Class->HasAnyClassFlags(CLASS_Native) || Class->bCooked || !Field.IsValid())
+	{
+		return Result;
+	}
+
+	const UBlueprint* Blueprint = Cast<const UBlueprint>(Class->ClassGeneratedBy);
+	if (!Blueprint)
+	{
+		return Result;
+	}
+
+	ensure(Blueprint->GeneratedClass);
+	return Blueprint->GeneratedClass;
+}
+
+/** */
+const UClass* GetSavedGeneratedStruct(const UClass* Class)
+{
+	if (!Class || Class->HasAnyClassFlags(CLASS_Native) || Class->bCooked)
+	{
+		return Class;
+	}
+
+	UBlueprint* Blueprint = Cast<UBlueprint>(Class->ClassGeneratedBy);
+	if (!Blueprint)
+	{
+		return Class;
+	}
+
+	ensure(Blueprint->GeneratedClass);
+	return Blueprint->GeneratedClass;
+}
+
 } //namespace
 
 
@@ -139,10 +202,12 @@ FCompiledBindingLibraryCompiler::FCompiledBindingLibraryCompiler()
 
 }
 
-TValueOrError<FCompiledBindingLibraryCompiler::FFieldIdHandle, FText> FCompiledBindingLibraryCompiler::AddFieldId(const UClass* SourceClass, FName FieldId)
+TValueOrError<FCompiledBindingLibraryCompiler::FFieldIdHandle, FText> FCompiledBindingLibraryCompiler::AddFieldId(const UClass* InSourceClass, FName FieldId)
 {
 	Impl->bCompiled = false;
 
+	const UClass* SourceClass = FBlueprintEditorUtils::GetMostUpToDateClass(InSourceClass);
+
 	if (FieldId.IsNone())
 	{
 		return MakeError(LOCTEXT("FieldNotDefined", "The Field is not defined. Is it missing the FieldNotify flag?."));
@@ -238,7 +303,13 @@ TValueOrError<FCompiledBindingLibraryCompiler::FFieldPathHandle, FText> FCompile
 
 	for (int32 Index = 0; Index < InFieldPath.Num(); ++Index)
 	{
+		// Make sure the FieldVariant is not from a skeletalclass
 		FMVVMConstFieldVariant FieldVariant = InFieldPath[Index];
+		if (!FieldVariant.IsValid())
+		{
+			return MakeError(FText::Format(LOCTEXT("FieldDoesNotHaveValidOwner", "The field {0} doesn't have a valid owner."), FText::FromName(InFieldPath[Index].GetName())));
+		}
+
 		const bool bIsLast = Index == InFieldPath.Num() - 1;
 		if (FieldVariant.IsProperty())
 		{
@@ -339,14 +410,16 @@ TValueOrError<FCompiledBindingLibraryCompiler::FFieldPathHandle, FText> FCompile
 		return MakeError(FText::Format(LOCTEXT("FieldDoesNotReturnType", "The field does not return a '{0}'."), ExpectedType->GetDisplayNameText()));
 	}
 
+	UE::MVVM::FMVVMConstFieldVariant Last = FieldPath.Last();
+
 	const FObjectPropertyBase* ObjectPropertyBase = nullptr;
-	if (FieldPath.Last().IsProperty())
+	if (Last.IsProperty())
 	{
-		ObjectPropertyBase = CastField<const FObjectPropertyBase>(FieldPath.Last().GetProperty());
+		ObjectPropertyBase = CastField<const FObjectPropertyBase>(Last.GetProperty());
 	}
-	else if (FieldPath.Last().IsFunction())
+	else if (Last.IsFunction())
 	{
-		ObjectPropertyBase = CastField<const FObjectPropertyBase>(BindingHelper::GetReturnProperty(FieldPath.Last().GetFunction()));
+		ObjectPropertyBase = CastField<const FObjectPropertyBase>(BindingHelper::GetReturnProperty(Last.GetFunction()));
 	}
 
 	if (ObjectPropertyBase == nullptr)
@@ -362,10 +435,13 @@ TValueOrError<FCompiledBindingLibraryCompiler::FFieldPathHandle, FText> FCompile
 }
 
 
-TValueOrError<FCompiledBindingLibraryCompiler::FFieldPathHandle, FText> FCompiledBindingLibraryCompiler::AddConversionFunctionFieldPath(const UClass* SourceClass, const UFunction* Function)
+TValueOrError<FCompiledBindingLibraryCompiler::FFieldPathHandle, FText> FCompiledBindingLibraryCompiler::AddConversionFunctionFieldPath(const UClass* InSourceClass, const UFunction* InFunction)
 {
 	Impl->bCompiled = false;
 
+	const UClass* SourceClass = FBlueprintEditorUtils::GetMostUpToDateClass(InSourceClass);
+	const UFunction* Function = FBlueprintEditorUtils::GetMostUpToDateFunction(InFunction);
+
 	if (SourceClass == nullptr)
 	{
 		return MakeError(LOCTEXT("SourceClassInvalid", "The source class is invalid."));
@@ -575,7 +651,7 @@ TValueOrError<FCompiledBindingLibraryCompiler::FCompileResult, FText> FCompiledB
 		TArray<int32> RawFieldIdIndex;
 	};
 
-	// create the list of UClass
+	// create the list of UClass. Uses the Generated class, not the SkeletalClass.
 	TMap<const UStruct*, FCompiledClassInfo> MapOfFieldInClass;
 	{
 		for (int32 Index = 0; Index < Impl->Fields.Num(); ++Index)
@@ -583,11 +659,11 @@ TValueOrError<FCompiledBindingLibraryCompiler::FCompileResult, FText> FCompiledB
 			const Private::FRawField& RawField = Impl->Fields[Index];
 			check(!RawField.Field.IsEmpty());
 
-			const UStruct* Owner = RawField.Field.IsProperty() ? RawField.Field.GetProperty()->GetOwnerStruct() : RawField.Field.GetFunction()->GetOwnerClass();
-			check(Owner);
-			FCompiledClassInfo& ClassInfo = MapOfFieldInClass.FindOrAdd(Owner);
+			const UStruct* FinalOwner = Private::GetSavedGeneratedStruct(RawField.Field);
+			check(FinalOwner);
+			FCompiledClassInfo& ClassInfo = MapOfFieldInClass.FindOrAdd(FinalOwner);
 
-			// Test if the Field is not there more than one
+			// Test if the Field is there more than one
 			{
 				FMVVMConstFieldVariant FieldToTest = RawField.Field;
 				const TArray<Private::FRawField>& ListOfFields = Impl->Fields;
@@ -608,7 +684,7 @@ TValueOrError<FCompiledBindingLibraryCompiler::FCompileResult, FText> FCompiledB
 			check(RawFieldId.FieldId.IsValid());
 			check(RawFieldId.NotifyFieldValueChangedClass);
 
-			FCompiledClassInfo& ClassInfo = MapOfFieldInClass.FindOrAdd(RawFieldId.NotifyFieldValueChangedClass);
+			FCompiledClassInfo& ClassInfo = MapOfFieldInClass.FindOrAdd(Private::GetSavedGeneratedStruct(RawFieldId.NotifyFieldValueChangedClass));
 
 			// Test if the Field is not there more than one
 			{
@@ -625,9 +701,11 @@ TValueOrError<FCompiledBindingLibraryCompiler::FCompileResult, FText> FCompiledB
 		}
 	}
 
-
 	// Todo optimize that list to group common type. ie UWidget::ToolTip == UProgressBar::ToolTip. We can merge UWidget in UProgressBar.
-	//The difficulty is with type that may be not loaded at runtime and would create runtime issue with type that would be loaded otherwise.
+	//Algo: for each class entry
+		// if: a class is child of another class in the list (ProgressBar is child of Widget all property inside Widget are also in ProgressBar)
+		// then: merge the 2 and restart the algo
+
 
 	FCompileResult Result;
 
@@ -638,7 +716,7 @@ TValueOrError<FCompiledBindingLibraryCompiler::FCompileResult, FText> FCompiledB
 	for (TPair<const UStruct*, FCompiledClassInfo>& StructCompiledFields : MapOfFieldInClass)
 	{
 		FMVVMVCompiledFields CompiledFields;
-		CompiledFields.ClassOrScriptStruct = StructCompiledFields.Key;
+		CompiledFields.ClassOrScriptStruct = StructCompiledFields.Key; // The generated class not the skeletal class
 		check(StructCompiledFields.Key);
 
 		TArray<FName> PropertyNames;
@@ -652,6 +730,7 @@ TValueOrError<FCompiledBindingLibraryCompiler::FCompileResult, FText> FCompiledB
 
 			if (Field.IsProperty())
 			{
+				// N.B. no need to translate from skeletal to generated, we only use the the name or the index.
 				Result.Library.LoadedProperties.Add(const_cast<FProperty*>(Field.GetProperty()));
 				PropertyNames.Add(Field.GetName());
 				RawField.LoadedPropertyOrFunctionIndex = TotalNumberOfProperties;
@@ -660,6 +739,7 @@ TValueOrError<FCompiledBindingLibraryCompiler::FCompileResult, FText> FCompiledB
 			else
 			{
 				check(Field.IsFunction());
+				// N.B. no need to translate from skeletal to generated, we only use the the name or the index.
 				Result.Library.LoadedFunctions.Add(const_cast<UFunction*>(Field.GetFunction()));
 				FunctionNames.Add(Field.GetName());
 				RawField.LoadedPropertyOrFunctionIndex = TotalNumberOfFunctions;
diff --git a/Engine/Plugins/Runtime/ModelViewViewModel/Source/ModelViewViewModelBlueprint/Private/Bindings/MVVMConversionFunctionHelper.cpp b/Engine/Plugins/Runtime/ModelViewViewModel/Source/ModelViewViewModelBlueprint/Private/Bindings/MVVMConversionFunctionHelper.cpp
--- a/Engine/Plugins/Runtime/ModelViewViewModel/Source/ModelViewViewModelBlueprint/Private/Bindings/MVVMConversionFunctionHelper.cpp
+++ b/Engine/Plugins/Runtime/ModelViewViewModel/Source/ModelViewViewModelBlueprint/Private/Bindings/MVVMConversionFunctionHelper.cpp
@@ -16,33 +16,6 @@ static const FName ConversionFunctionMetadataKey = "ConversionFunction";
 
 namespace Private
 {
-
-	void BuildPropertyPath(TArray<FMVVMBlueprintPropertyPath>& Result, FMVVMBlueprintPropertyPath BasePath, UClass* Class, const UK2Node_VariableGet* CurrentNode)
-	{
-		bool bAddPath = false;
-		const UEdGraphPin* OutputPin = CurrentNode->FindPinByPredicate([](UEdGraphPin* Pin) { return Pin->Direction == EGPD_Output; });
-		if (OutputPin)
-		{
-			for (const UEdGraphPin* LinkedTo : OutputPin->LinkedTo)
-			{
-				if (const UK2Node_VariableGet* LinkToVariable = Cast<UK2Node_VariableGet>(LinkedTo->GetOwningNode()))
-				{
-					BasePath.SetBasePropertyPath(UE::MVVM::FMVVMConstFieldVariant(LinkToVariable->VariableReference.ResolveMember<FProperty>(Class)));
-					BuildPropertyPath(Result, BasePath, Class, LinkToVariable);
-				}
-				else
-				{
-					bAddPath = true;
-				}
-			}
-		}
-
-		if (bAddPath)
-		{
-			Result.Add(BasePath);
-		}
-	}
-
 	UEdGraph* FindExistingConversionFunctionWrapper(const UWidgetBlueprint* WidgetBlueprint, FName WrapperName)
 	{
 		const TObjectPtr<UEdGraph>* Result = WidgetBlueprint->FunctionGraphs.FindByPredicate([WrapperName](const UEdGraph* GraphPtr) { return GraphPtr->GetFName() == WrapperName; });
@@ -197,11 +170,11 @@ namespace Private
 		{
 			if (UFunction* Function = MemberReference.ResolveMember<UFunction>(WidgetBlueprint->SkeletonGeneratedClass)) 
 			{
-				ResultPath.AppendBasePropertyPath(UE::MVVM::FMVVMConstFieldVariant(Function));
+				ResultPath.AppendPropertyPath(WidgetBlueprint, UE::MVVM::FMVVMConstFieldVariant(Function));
 			}
 			else if (const FProperty* Property = MemberReference.ResolveMember<FProperty>(WidgetBlueprint->SkeletonGeneratedClass))
 			{
-				ResultPath.AppendBasePropertyPath(UE::MVVM::FMVVMConstFieldVariant(Property));
+				ResultPath.AppendPropertyPath(WidgetBlueprint, UE::MVVM::FMVVMConstFieldVariant(Property));
 			}
 		}
 
diff --git a/Engine/Plugins/Runtime/ModelViewViewModel/Source/ModelViewViewModelBlueprint/Private/MVVMBlueprintView.cpp b/Engine/Plugins/Runtime/ModelViewViewModel/Source/ModelViewViewModelBlueprint/Private/MVVMBlueprintView.cpp
--- a/Engine/Plugins/Runtime/ModelViewViewModel/Source/ModelViewViewModelBlueprint/Private/MVVMBlueprintView.cpp
+++ b/Engine/Plugins/Runtime/ModelViewViewModel/Source/ModelViewViewModelBlueprint/Private/MVVMBlueprintView.cpp
@@ -5,6 +5,8 @@
 #include "Components/Widget.h"
 #include "Kismet2/BlueprintEditorUtils.h"
 #include "MVVMWidgetBlueprintExtension_View.h"
+#include "WidgetBlueprint.h"
+#include "UObject/FortniteMainBranchObjectVersion.h"
 
 #include UE_INLINE_GENERATED_CPP_BY_NAME(MVVMBlueprintView)
 
@@ -99,10 +101,10 @@ const FMVVMBlueprintViewBinding* UMVVMBlueprintView::FindBinding(const UWidget*
 FMVVMBlueprintViewBinding* UMVVMBlueprintView::FindBinding(const UWidget* Widget, const FProperty* Property)
 {
 	FName WidgetName = Widget->GetFName();
-	return Bindings.FindByPredicate([WidgetName, Property](const FMVVMBlueprintViewBinding& Binding)
+	return Bindings.FindByPredicate([WidgetBlueprint = GetOuterUMVVMWidgetBlueprintExtension_View()->GetWidgetBlueprint(), WidgetName, Property](const FMVVMBlueprintViewBinding& Binding)
 		{
 			return Binding.DestinationPath.GetWidgetName() == WidgetName &&
-				Binding.DestinationPath.BasePropertyPathContains(UE::MVVM::FMVVMConstFieldVariant(Property));
+				Binding.DestinationPath.PropertyPathContains(WidgetBlueprint->GeneratedClass, UE::MVVM::FMVVMConstFieldVariant(Property));
 		});
 }
 
@@ -133,7 +135,7 @@ FMVVMBlueprintViewBinding& UMVVMBlueprintView::AddBinding(const UWidget* Widget,
 {
 	FMVVMBlueprintViewBinding& NewBinding = Bindings.AddDefaulted_GetRef();
 	NewBinding.DestinationPath.SetWidgetName(Widget->GetFName());
-	NewBinding.DestinationPath.SetBasePropertyPath(UE::MVVM::FMVVMConstFieldVariant(Property));
+	NewBinding.DestinationPath.SetPropertyPath(GetOuterUMVVMWidgetBlueprintExtension_View()->GetWidgetBlueprint(), UE::MVVM::FMVVMConstFieldVariant(Property));
 	NewBinding.BindingId = FGuid::NewGuid();
 
 	OnBindingsUpdated.Broadcast();
@@ -222,6 +224,12 @@ void UMVVMBlueprintView::ResetBindingMessages()
 	BindingMessages.Reset();
 }
 
+void UMVVMBlueprintView::Serialize(FArchive& Ar)
+{
+	Super::Serialize(Ar);
+	Ar.UsingCustomVersion(FFortniteMainBranchObjectVersion::GUID);	
+}
+
 void UMVVMBlueprintView::PostLoad()
 {
 	Super::PostLoad();
@@ -233,6 +241,28 @@ void UMVVMBlueprintView::PostLoad()
 			Binding.BindingId = FGuid::NewGuid();
 		}
 	}
+
+#if WITH_EDITOR
+	// Make sure all bindings uses the skeletal class
+	if (GetLinkerCustomVersion(FFortniteMainBranchObjectVersion::GUID) < FFortniteMainBranchObjectVersion::MVVMConvertPropertyPathToSkeletalClass)
+	{
+		GetOuterUMVVMWidgetBlueprintExtension_View()->ConditionalPostLoad();
+		GetOuterUMVVMWidgetBlueprintExtension_View()->GetWidgetBlueprint()->ConditionalPostLoad();
+		const UWidgetBlueprint* ThisWidgetBlueprint = GetOuterUMVVMWidgetBlueprintExtension_View()->GetWidgetBlueprint();
+
+		for (FMVVMBlueprintViewBinding& Binding : Bindings)
+		{
+			for (const FMVVMBlueprintFieldPath& FieldPath : Binding.SourcePath.GetFieldPaths())
+			{
+				const_cast<FMVVMBlueprintFieldPath&>(FieldPath).SetDeprecatedSelfReference(ThisWidgetBlueprint);
+			}
+			for (const FMVVMBlueprintFieldPath& FieldPath : Binding.DestinationPath.GetFieldPaths())
+			{
+				const_cast<FMVVMBlueprintFieldPath&>(FieldPath).SetDeprecatedSelfReference(ThisWidgetBlueprint);
+			}
+		}
+	}
+#endif
 }
 
 #if WITH_EDITOR
diff --git a/Engine/Plugins/Runtime/ModelViewViewModel/Source/ModelViewViewModelBlueprint/Private/MVVMBlueprintViewBinding.cpp b/Engine/Plugins/Runtime/ModelViewViewModel/Source/ModelViewViewModelBlueprint/Private/MVVMBlueprintViewBinding.cpp
--- a/Engine/Plugins/Runtime/ModelViewViewModel/Source/ModelViewViewModelBlueprint/Private/MVVMBlueprintViewBinding.cpp
+++ b/Engine/Plugins/Runtime/ModelViewViewModel/Source/ModelViewViewModelBlueprint/Private/MVVMBlueprintViewBinding.cpp
@@ -16,7 +16,7 @@ FName FMVVMBlueprintViewBinding::GetFName() const
 
 namespace UE::MVVM::Private
 {
-	void AppendViewModelPathString(const UMVVMBlueprintView* BlueprintView, const FMVVMBlueprintPropertyPath& ViewModelPath, FStringBuilderBase& PathBuilder)
+	void AppendViewModelPathString(const UWidgetBlueprint* WidgetBlueprint, const UMVVMBlueprintView* BlueprintView, const FMVVMBlueprintPropertyPath& ViewModelPath, FStringBuilderBase& PathBuilder)
 	{
 		if (ViewModelPath.IsEmpty())
 		{
@@ -34,7 +34,7 @@ namespace UE::MVVM::Private
 			PathBuilder << TEXT("<none>");
 		}
 
-		const FString PropertyPath = ViewModelPath.GetBasePropertyPath();
+		const FString PropertyPath = ViewModelPath.GetPropertyPath(WidgetBlueprint->SkeletonGeneratedClass);
 		if (!PropertyPath.IsEmpty())
 		{
 			PathBuilder << TEXT(".");
@@ -74,7 +74,7 @@ namespace UE::MVVM::Private
 					FMVVMBlueprintPropertyPath ArgumentPath = ConversionFunctionHelper::GetPropertyPathForArgument(WidgetBlueprint, CallFunctionNode, Pin->GetFName(), true);
 					if (!ArgumentPath.IsEmpty())
 					{
-						AppendViewModelPathString(BlueprintView, ArgumentPath, NameBuilder);
+						AppendViewModelPathString(WidgetBlueprint, BlueprintView, ArgumentPath, NameBuilder);
 					}
 					else
 					{
@@ -98,7 +98,7 @@ namespace UE::MVVM::Private
 				NameBuilder << SourceToDestFunction->GetName();
 				NameBuilder << TEXT("(");
 			
-				AppendViewModelPathString(BlueprintView, Binding.SourcePath, NameBuilder);
+				AppendViewModelPathString(WidgetBlueprint, BlueprintView, Binding.SourcePath, NameBuilder);
 
 				NameBuilder << TEXT(")");
 			}
@@ -110,7 +110,7 @@ namespace UE::MVVM::Private
 		}
 		else if (!Binding.SourcePath.IsEmpty())
 		{
-			AppendViewModelPathString(BlueprintView, Binding.SourcePath, NameBuilder);
+			AppendViewModelPathString(WidgetBlueprint, BlueprintView, Binding.SourcePath, NameBuilder);
 		}
 
 		FString Name = NameBuilder.ToString();
@@ -137,7 +137,7 @@ namespace UE::MVVM::Private
 			PathBuilder << WidgetPath.GetWidgetName();
 		}
 
-		const FString PropertyPath = WidgetPath.GetBasePropertyPath();
+		const FString PropertyPath = WidgetPath.GetPropertyPath(WidgetBlueprint->SkeletonGeneratedClass);
 		if (!PropertyPath.IsEmpty())
 		{
 			PathBuilder << TEXT(".");
diff --git a/Engine/Plugins/Runtime/ModelViewViewModel/Source/ModelViewViewModelBlueprint/Private/MVVMFunctionGraphHelper.cpp b/Engine/Plugins/Runtime/ModelViewViewModel/Source/ModelViewViewModelBlueprint/Private/MVVMFunctionGraphHelper.cpp
--- a/Engine/Plugins/Runtime/ModelViewViewModel/Source/ModelViewViewModelBlueprint/Private/MVVMFunctionGraphHelper.cpp
+++ b/Engine/Plugins/Runtime/ModelViewViewModel/Source/ModelViewViewModelBlueprint/Private/MVVMFunctionGraphHelper.cpp
@@ -2,24 +2,32 @@
 
 #include "MVVMFunctionGraphHelper.h"
 
+#include "EdGraph/EdGraph.h"
 #include "K2Node_CallFunction.h"
 #include "K2Node_DynamicCast.h"
 #include "K2Node_FunctionEntry.h"
 #include "K2Node_VariableGet.h"
+#include "K2Node_VariableSet.h"
 #include "KismetCompiler.h"
 #include "Kismet2/BlueprintEditorUtils.h"
 
 #include "View/MVVMView.h"
 #include "MVVMViewModelBase.h"
 
+#if WITH_EDITOR
+#include "EdGraphSchema_K2.h"
+#endif
 
 namespace UE::MVVM::FunctionGraphHelper
 {
 
-UEdGraph* CreateFunctionGraph(FKismetCompilerContext& InContext, FStringView InFunctionName, EFunctionFlags ExtraFunctionFlag, const FStringView Category, bool bIsEditable)
+UEdGraph* CreateFunctionGraph(UBlueprint* InBlueprint, FStringView InFunctionName, EFunctionFlags ExtraFunctionFlag, const FStringView Category, bool bIsEditable)
 {
-	FName UniqueFunctionName = FBlueprintEditorUtils::FindUniqueKismetName(InContext.Blueprint, InFunctionName.GetData());
-	UEdGraph* FunctionGraph = FBlueprintEditorUtils::CreateNewGraph(InContext.Blueprint, UniqueFunctionName, UEdGraph::StaticClass(), UEdGraphSchema_K2::StaticClass());
+	FName UniqueFunctionName = FBlueprintEditorUtils::FindUniqueKismetName(InBlueprint, InFunctionName.GetData());
+	UEdGraph* FunctionGraph = FBlueprintEditorUtils::CreateNewGraph(InBlueprint, UniqueFunctionName, UEdGraph::StaticClass(), UEdGraphSchema_K2::StaticClass());
+
+	// This is sometime needed. Needs more investigation.
+	//InBlueprint->FunctionGraphs.Add(FunctionGraph);
 
 	const UEdGraphSchema_K2* Schema = CastChecked<UEdGraphSchema_K2>(FunctionGraph->GetSchema());
 	Schema->MarkFunctionEntryAsEditable(FunctionGraph, bIsEditable);
@@ -59,6 +67,19 @@ UEdGraph* CreateIntermediateFunctionGraph(FKismetCompilerContext& InContext, FSt
 }
 
 
+void SetVariableNodeMember(UK2Node_Variable* InVariableNode, const FProperty* InProperty, UBlueprint* InTargetBlueprint)
+{
+	if (InProperty->GetOwnerStruct() && InTargetBlueprint->GeneratedClass && InTargetBlueprint->GeneratedClass->IsChildOf(InProperty->GetOwnerStruct()))
+	{
+		FGuid Guid = FBlueprintEditorUtils::FindMemberVariableGuidByName(InTargetBlueprint, InProperty->GetFName());
+		InVariableNode->VariableReference.SetSelfMember(InProperty->GetFName(), Guid);
+	}
+	else
+	{
+		UEdGraphSchema_K2::ConfigureVarNode(InVariableNode, InProperty->GetFName(), InProperty->GetOwnerStruct(), InTargetBlueprint);
+	}
+}
+
 bool AddFunctionArgument(UEdGraph* InFunctionGraph, TSubclassOf<UObject> InArgument, FName InArgumentName)
 {
 	for (UEdGraphNode* Node : InFunctionGraph->Nodes)
@@ -80,7 +101,7 @@ bool AddFunctionArgument(UEdGraph* InFunctionGraph, TSubclassOf<UObject> InArgum
 }
 
 
-bool AddFunctionArgument(UEdGraph* InFunctionGraph, FProperty* InArgument, FName InArgumentName)
+bool AddFunctionArgument(UEdGraph* InFunctionGraph, const FProperty* InArgument, FName InArgumentName)
 {
 	for (UEdGraphNode* Node : InFunctionGraph->Nodes)
 	{
@@ -322,4 +343,202 @@ bool GenerateViewModelFielNotifyBroadcast(FKismetCompilerContext& InContext, UEd
 	return bResult;
 }
 
+namespace Private
+{
+
+UK2Node_CallFunction* CreateFunctionNode(UEdGraph* InFunctionGraph, const UFunction* InFunction, UBlueprint* InBlueprint, FKismetCompilerContext* InContext, UK2Node_FunctionEntry* InFunctionEntry)
+{
+	if (InContext)
+	{
+		check(InFunctionEntry);
+		UK2Node_CallFunction* FunctionNode = InContext->SpawnIntermediateNode<UK2Node_CallFunction>(InFunctionEntry, InFunctionGraph);
+		FunctionNode->SetFromFunction(InFunction);
+		FunctionNode->AllocateDefaultPins();
+		return FunctionNode;
+	}
+	else
+	{
+		FGraphNodeCreator<UK2Node_CallFunction> RootCreator(*InFunctionGraph);
+		UK2Node_CallFunction* FunctionNode = RootCreator.CreateNode();
+		FunctionNode->NodePosX = 0;
+		FunctionNode->NodePosY = 0;
+		FunctionNode->SetFromFunction(InFunction);
+		RootCreator.Finalize();
+		return FunctionNode;
+	}
+}
+
+template<typename NodeType>
+NodeType* CreateVariableNode(UEdGraph* InFunctionGraph, const FProperty* InProperty, UBlueprint* InBlueprint, FKismetCompilerContext* InContext, UK2Node_FunctionEntry* InFunctionEntry)
+{
+	if (InContext)
+	{
+		check(InFunctionEntry);
+		NodeType* VariableNode = InContext->SpawnIntermediateNode<NodeType>(InFunctionEntry, InFunctionGraph);
+		SetVariableNodeMember(VariableNode, InProperty, InBlueprint);
+		VariableNode->AllocateDefaultPins();
+		return VariableNode;
+	}
+	else
+	{
+		FGraphNodeCreator<NodeType> RootCreator(*InFunctionGraph);
+		NodeType* VariableNode = RootCreator.CreateNode();
+		VariableNode->NodePosX = 0;
+		VariableNode->NodePosY = 0;
+		SetVariableNodeMember(VariableNode, InProperty, InBlueprint);
+		RootCreator.Finalize();
+		return VariableNode;
+	}
+}
+
+bool GenerateSetter(UEdGraph* InFunctionGraph, TArrayView<UE::MVVM::FMVVMConstFieldVariant> InSetterPath, UBlueprint* InBlueprint, FKismetCompilerContext* InContext)
+{
+	check(InFunctionGraph);
+
+	UK2Node_FunctionEntry* FunctionEntry = nullptr;
+	for (UEdGraphNode* Node : InFunctionGraph->Nodes)
+	{
+		FunctionEntry = Cast<UK2Node_FunctionEntry>(Node);
+		if (FunctionEntry)
+		{
+			break;
+		}
+	}
+
+	if (FunctionEntry == nullptr)
+	{
+		return false;
+	}
+
+	UEdGraphPin* FunctionInputPin = FunctionEntry->FindPinByPredicate([](UEdGraphPin* OtherPin)
+		{
+			return OtherPin->Direction == EGPD_Output
+				&& OtherPin->PinName != UEdGraphSchema_K2::PN_Then;
+		});
+
+	const UEdGraphSchema* Schema = InFunctionGraph->GetSchema();
+	if (Schema == nullptr)
+	{
+		return false;
+	}
+
+	UEdGraphPin* ThenPin = FunctionEntry->FindPin(UEdGraphSchema_K2::PN_Then, EGPD_Output);
+	UEdGraphPin* OutContextPin = nullptr;
+
+	bool bResult = true;
+	for (int32 Index = 0; Index < InSetterPath.Num(); ++Index)
+	{
+		UE::MVVM::FMVVMConstFieldVariant& Field = InSetterPath[Index];
+		const bool bLastPath = Index == InSetterPath.Num() - 1;
+
+		UEdGraphPin* NextThenPin = nullptr;
+		UEdGraphPin* NextOutContextPin = nullptr;
+
+		UEdGraphPin* SetterPin = nullptr;
+		UEdGraphPin* CurrentExecPin = nullptr;
+		UEdGraphPin* CurrentInContextPin = nullptr;
+		if (Field.IsProperty())
+		{
+			check(Field.GetProperty());
+
+			if (Field.GetProperty()->GetOwnerClass() == nullptr)
+			{
+				ensureMsgf(false, TEXT("Property point to a structure member and that is not supported right now."));
+				return false;
+			}
+
+			if (bLastPath)
+			{
+				UK2Node_VariableSet* VariableSetNode = CreateVariableNode<UK2Node_VariableSet>(InFunctionGraph, Field.GetProperty(), InBlueprint, InContext, FunctionEntry);
+
+				NextThenPin = VariableSetNode->FindPin(UEdGraphSchema_K2::PN_Then, EGPD_Output);
+				CurrentExecPin = VariableSetNode->FindPin(UEdGraphSchema_K2::PN_Execute, EGPD_Input);
+				CurrentInContextPin = VariableSetNode->FindPin(UEdGraphSchema_K2::PN_Self, EGPD_Input);
+				SetterPin = VariableSetNode->FindPin(Field.GetProperty()->GetFName(), EGPD_Input);
+			}
+			else
+			{
+				UK2Node_VariableGet* VariableGetNode = CreateVariableNode<UK2Node_VariableGet>(InFunctionGraph, Field.GetProperty(), InBlueprint, InContext, FunctionEntry);
+
+				CurrentInContextPin = VariableGetNode->FindPin(UEdGraphSchema_K2::PN_Self, EGPD_Input);
+				NextOutContextPin = VariableGetNode->FindPin(Field.GetProperty()->GetFName(), EGPD_Output);
+			}
+		}
+		else
+		{
+			ensure(Field.IsFunction());
+			check(Field.GetFunction());
+
+			UK2Node_CallFunction* CallFunctionNode = CreateFunctionNode(InFunctionGraph, Field.GetFunction(), InBlueprint, InContext, FunctionEntry);
+
+			NextThenPin = CallFunctionNode->FindPin(UEdGraphSchema_K2::PN_Then, EGPD_Output);
+			CurrentExecPin = CallFunctionNode->FindPin(UEdGraphSchema_K2::PN_Execute, EGPD_Input);
+			CurrentInContextPin = CallFunctionNode->FindPin(UEdGraphSchema_K2::PN_Self, EGPD_Input);
+
+			if (bLastPath)
+			{
+				// find the input of the setter function
+				SetterPin = CallFunctionNode->FindPinByPredicate([](UEdGraphPin* OtherPin)
+					{
+						return OtherPin->Direction == EGPD_Input
+							&& OtherPin->PinName != UEdGraphSchema_K2::PN_Execute
+							&& OtherPin->PinName != UEdGraphSchema_K2::PN_Self;
+					});
+			}
+			else
+			{
+				// find the output of the getter function
+				NextOutContextPin = CallFunctionNode->FindPinByPredicate([](UEdGraphPin* OtherPin)
+					{
+						return OtherPin->Direction == EGPD_Output
+							&& OtherPin->PinName != UEdGraphSchema_K2::PN_Then;
+					});
+			}
+		}
+
+		if (ThenPin && CurrentExecPin)
+		{
+			ensure(Schema->TryCreateConnection(ThenPin, CurrentExecPin));
+		}
+
+		if (OutContextPin)
+		{
+			if (ensure(CurrentInContextPin))
+			{
+				ensure(Schema->TryCreateConnection(OutContextPin, CurrentInContextPin));
+			}
+		}
+
+		if (bLastPath)
+		{
+			if (ensure(SetterPin && FunctionInputPin))
+			{
+				ensure(Schema->TryCreateConnection(FunctionInputPin, SetterPin));
+			}
+		}
+
+		if (NextThenPin)
+		{
+			ThenPin = NextThenPin;
+		}
+		if (NextOutContextPin)
+		{
+			OutContextPin = NextOutContextPin;
+		}
+	}
+
+	return true;
+}
+} //Private
+
+bool GenerateIntermediateSetter(FKismetCompilerContext& InContext, UEdGraph* InFunctionGraph, TArrayView<UE::MVVM::FMVVMConstFieldVariant> InSetterPath)
+{
+	return Private::GenerateSetter(InFunctionGraph, InSetterPath, InContext.Blueprint, &InContext);
+}
+
+bool GenerateSetter(UBlueprint* InBlueprint, UEdGraph* InFunctionGraph, TArrayView<UE::MVVM::FMVVMConstFieldVariant> InSetterPath)
+{
+	return Private::GenerateSetter(InFunctionGraph, InSetterPath, InBlueprint, nullptr);
+}
+
 } //namespace
diff --git a/Engine/Plugins/Runtime/ModelViewViewModel/Source/ModelViewViewModelBlueprint/Private/MVVMFunctionGraphHelper.h b/Engine/Plugins/Runtime/ModelViewViewModel/Source/ModelViewViewModelBlueprint/Private/MVVMFunctionGraphHelper.h
--- a/Engine/Plugins/Runtime/ModelViewViewModel/Source/ModelViewViewModelBlueprint/Private/MVVMFunctionGraphHelper.h
+++ b/Engine/Plugins/Runtime/ModelViewViewModel/Source/ModelViewViewModelBlueprint/Private/MVVMFunctionGraphHelper.h
@@ -3,6 +3,7 @@
 #pragma once
 
 #include "Containers/StringFwd.h"
+#include "Types/MVVMFieldVariant.h"
 
 class FName;
 class FProperty;
@@ -13,21 +14,25 @@ template <typename T> class TSubclassOf;
 class FKismetCompilerContext;
 class UBlueprint;
 class UEdGraph;
+class UK2Node_Variable;
 
 namespace UE::MVVM::FunctionGraphHelper
 {
 
 	/** Generate a function graph that can be viewed in the BP editor. */
-	MODELVIEWVIEWMODELBLUEPRINT_API UEdGraph* CreateFunctionGraph(FKismetCompilerContext& Context, const FStringView FunctionName, EFunctionFlags ExtraFunctionFlag, const FStringView Category, bool bIsEditable);
+	MODELVIEWVIEWMODELBLUEPRINT_API UEdGraph* CreateFunctionGraph(UBlueprint* Blueprint, const FStringView FunctionName, EFunctionFlags ExtraFunctionFlag, const FStringView Category, bool bIsEditable);
 
 	/** Generate a function graph that cannot be viewed in the BP editor. */
 	UEdGraph* CreateIntermediateFunctionGraph(FKismetCompilerContext& Context, const FStringView FunctionName, EFunctionFlags ExtraFunctionFlag, const FStringView Category, bool bIsEditable);
 
+	/** Set the UK2Node_Variable::VariableReference correctly base on the variable's owner. */
+	MODELVIEWVIEWMODELBLUEPRINT_API void SetVariableNodeMember(UK2Node_Variable* InVarNode, const FProperty* InProperty, UBlueprint* InTargetBlueprint);
+
 	/** Add an input argument to an existing function graph. */
 	bool AddFunctionArgument(UEdGraph* FunctionGraph, TSubclassOf<UObject> Argument, FName ArgumentName);
 
 	/** Add an input argument to an existing function graph. */
-	bool AddFunctionArgument(UEdGraph* FunctionGraph, FProperty* Argument, FName ArgumentName);
+	bool AddFunctionArgument(UEdGraph* FunctionGraph, const FProperty* Argument, FName ArgumentName);
 
 	/** Generate the nodes needed for to call the MVVMView::SetViewModel function. */
 	bool GenerateViewModelSetter(FKismetCompilerContext& Context, UEdGraph* FunctionGraph, FName ViewModelName);
@@ -38,4 +43,9 @@ namespace UE::MVVM::FunctionGraphHelper
 	/** Generate the nodes needed to broadcast field notify function. */
 	bool GenerateViewModelFielNotifyBroadcast(FKismetCompilerContext& Context, UEdGraph* FunctionGraph, FProperty* Property);
 
+	/** Generate the nodes needed for to set a variable/function value. */
+	bool GenerateSetter(UBlueprint* Blueprint, UEdGraph* FunctionGraph, TArrayView<UE::MVVM::FMVVMConstFieldVariant> SetterPath);
+	/** Generate the nodes needed to set a variable/function value. To use while while compiling the Blueprint. */
+	bool GenerateIntermediateSetter(FKismetCompilerContext& Context, UEdGraph* FunctionGraph, TArrayView<UE::MVVM::FMVVMConstFieldVariant> SetterPath);
+
 } //namespace
diff --git a/Engine/Plugins/Runtime/ModelViewViewModel/Source/ModelViewViewModelBlueprint/Private/MVVMPropertyPath.cpp b/Engine/Plugins/Runtime/ModelViewViewModel/Source/ModelViewViewModelBlueprint/Private/MVVMPropertyPath.cpp
--- a/Engine/Plugins/Runtime/ModelViewViewModel/Source/ModelViewViewModelBlueprint/Private/MVVMPropertyPath.cpp
+++ b/Engine/Plugins/Runtime/ModelViewViewModel/Source/ModelViewViewModelBlueprint/Private/MVVMPropertyPath.cpp
@@ -1,7 +1,280 @@
 // Copyright Epic Games, Inc. All Rights Reserved.
 
 #include "MVVMPropertyPath.h"
+#include "BlueprintCompilationManager.h"
+#include "Engine/Blueprint.h"
+#include "Engine/UserDefinedStruct.h"
+#include "Kismet2/BlueprintEditorUtils.h"
+#include "Kismet2/StructureEditorUtils.h"
 
 #include UE_INLINE_GENERATED_CPP_BY_NAME(MVVMPropertyPath)
 
+FMVVMBlueprintFieldPath::FMVVMBlueprintFieldPath(const UBlueprint* InContext, UE::MVVM::FMVVMConstFieldVariant InField)
+{
+	if (!InField.IsValid())
+	{
+		return;
+	}
 
+	UStruct* Owner = InField.GetOwner();
+	ensure(Owner);
+	if (!Owner)
+	{
+		return;
+	}
+
+	// Find the Guid and set the BindingKind
+	UClass* OwnerClass = Cast<UClass>(Owner);
+	FName FieldName = InField.GetName();
+	FGuid MemberGuid;
+	if (InField.IsProperty())
+	{
+		BindingKind = EBindingKind::Property;
+		if (OwnerClass)
+		{
+			UBlueprint::GetGuidFromClassByFieldName<FProperty>(OwnerClass, FieldName, MemberGuid);
+		}
+		else if (UUserDefinedStruct* UserDefinedStruct = Cast<UUserDefinedStruct>(Owner))
+		{
+			MemberGuid = FStructureEditorUtils::GetGuidFromPropertyName(FieldName);
+		}
+	}
+	else if (InField.IsFunction())
+	{
+		BindingKind = EBindingKind::Function;
+		if (OwnerClass)
+		{
+			UBlueprint::GetGuidFromClassByFieldName<UFunction>(OwnerClass, FieldName, MemberGuid);
+		}
+	}
+	else
+	{
+		ensureAlwaysMsgf(false, TEXT("Binding to field of unknown type!"));
+	}
+
+	// Set the member reference
+	bool bIsSelf = (InContext->GeneratedClass && InContext->GeneratedClass->IsChildOf(Owner))
+		|| (InContext->SkeletonGeneratedClass && InContext->SkeletonGeneratedClass->IsChildOf(Owner));
+	if (bIsSelf)
+	{
+		BindingReference.SetSelfMember(FieldName, MemberGuid);
+	}
+	else if (OwnerClass)
+	{
+		FGuid Guid;
+		if (UBlueprint* VariableOwnerBP = Cast<UBlueprint>(OwnerClass->ClassGeneratedBy))
+		{
+			OwnerClass = VariableOwnerBP->SkeletonGeneratedClass;
+		}
+
+		BindingReference.SetExternalMember(FieldName, OwnerClass, MemberGuid);
+	}
+	else if (UScriptStruct* OwnerStruct = Cast<UScriptStruct>(Owner))
+	{
+		struct FMyMemberReference : public FMemberReference
+		{
+			void SetExternalStructMember(FName InMemberName, UStruct* InMemberParentStruct, FGuid InGuid)
+			{
+				MemberName = InMemberName;
+				MemberGuid = InGuid;
+				MemberParent = InMemberParentStruct;
+				MemberScope.Empty();
+				bSelfContext = false;
+				bWasDeprecated = false;
+			}
+		};
+		static_cast<FMyMemberReference&>(BindingReference).SetExternalStructMember(FieldName, Owner, MemberGuid);
+	}
+	else
+	{
+		ensureAlwaysMsgf(false, TEXT("Local member is not supported."));
+	}
+}
+
+
+FName FMVVMBlueprintFieldPath::GetFieldName(const UClass* InContext) const
+{
+	UE::MVVM::FMVVMConstFieldVariant Result = GetField(InContext);
+	if (Result.IsValid())
+	{
+		return Result.GetName();
+	}
+	return FName();
+}
+
+
+UE::MVVM::FMVVMConstFieldVariant FMVVMBlueprintFieldPath::GetField(const UClass* InContext) const
+{
+	if (!FBlueprintCompilationManager::IsGeneratedClassLayoutReady())
+	{
+		if (UClass* SkeletonClass = FBlueprintEditorUtils::GetSkeletonClass(BindingReference.GetMemberParentClass(const_cast<UClass*>(InContext))))
+		{
+			return GetFieldInternal(SkeletonClass);
+		}
+	}
+	return GetFieldInternal(InContext);
+}
+
+
+UE::MVVM::FMVVMConstFieldVariant FMVVMBlueprintFieldPath::GetFieldInternal(const UClass* InContext) const
+{
+	// Resolve any redirectors
+	if (!BindingReference.GetMemberName().IsNone())
+	{
+		if (BindingKind == EBindingKind::Property)
+		{
+			struct FMyMemberReference : public FMemberReference
+			{
+				UStruct* GetMemberParent() const
+				{
+					return Cast<UStruct>(MemberParent);
+				}
+				void SetMemberName(FName InMemberName)
+				{
+					MemberName = InMemberName;
+				}
+			};
+			UStruct* MemberParent = static_cast<const FMyMemberReference&>(BindingReference).GetMemberParent();
+			if (UScriptStruct* OwnerStruct = Cast<UScriptStruct>(MemberParent))
+			{
+				const FProperty* FoundProperty = FindUFieldOrFProperty<FProperty>(OwnerStruct, BindingReference.GetMemberName(), EFieldIterationFlags::IncludeAll);
+				if (!FoundProperty)
+				{
+					// Refresh the name from Guid
+					if (UUserDefinedStruct* UserDefinedStruct = Cast<UUserDefinedStruct>(OwnerStruct))
+					{
+						FoundProperty = FStructureEditorUtils::GetPropertyByGuid(UserDefinedStruct, BindingReference.GetMemberGuid());
+						if (FoundProperty)
+						{
+							const_cast<FMyMemberReference&>(static_cast<const FMyMemberReference&>(BindingReference)).SetMemberName(FoundProperty->GetFName());
+						}
+					}
+				}
+
+				return UE::MVVM::FMVVMConstFieldVariant(FoundProperty);
+			}
+
+			// Previous path. IsLocalScope was previously used to save struct properties.
+			if (BindingReference.IsLocalScope())
+			{
+				if (UPackage* Package = BindingReference.GetMemberParentPackage())
+				{
+					UObjectBase* FoundObject = FindObjectWithOuter(Package, UScriptStruct::StaticClass(), *BindingReference.GetMemberScopeName());
+					if (UScriptStruct* Struct = Cast<UScriptStruct>(static_cast<UObject*>(FoundObject)))
+					{
+						if (const FProperty* FoundProperty = FindUFieldOrFProperty<FProperty>(Struct, BindingReference.GetMemberName(), EFieldIterationFlags::IncludeAll))
+						{
+							return UE::MVVM::FMVVMConstFieldVariant(FoundProperty);
+						}
+					}
+				}
+			}
+
+			return UE::MVVM::FMVVMConstFieldVariant(BindingReference.ResolveMember<FProperty>(const_cast<UClass*>(InContext), false));
+		}
+		else if (BindingKind == EBindingKind::Function)
+		{
+			return UE::MVVM::FMVVMConstFieldVariant(BindingReference.ResolveMember<UFunction>(const_cast<UClass*>(InContext), false));
+		}
+	}
+
+	return UE::MVVM::FMVVMConstFieldVariant();
+}
+
+
+UClass* FMVVMBlueprintFieldPath::GetParentClass(const UClass* InSelfContext) const
+{
+	return BindingReference.GetMemberParentClass(const_cast<UClass*>(InSelfContext));
+}
+
+
+#if WITH_EDITOR
+void FMVVMBlueprintFieldPath::SetDeprecatedBindingReference(const FMemberReference& InBindingReference, EBindingKind InBindingKind)
+{
+	BindingReference = InBindingReference;
+	BindingKind = InBindingKind;
+}
+
+
+void FMVVMBlueprintFieldPath::SetDeprecatedSelfReference(const UBlueprint* InContext)
+{
+	struct FMyMemberReference : public FMemberReference
+	{
+		void SetSelfReference()
+		{
+			MemberParent = nullptr;
+			bSelfContext = true;
+			MemberScope.Empty();
+		}
+	};
+
+	if (UClass* ParentClass = BindingReference.GetMemberParentClass())
+	{
+		ParentClass->ConditionalPostLoad();
+		bool bIsSelf = (InContext->GeneratedClass && InContext->GeneratedClass->IsChildOf(ParentClass))
+			|| (InContext->SkeletonGeneratedClass && InContext->SkeletonGeneratedClass->IsChildOf(ParentClass));
+		if (bIsSelf)
+		{
+			FMyMemberReference& MyFieldPath = static_cast<FMyMemberReference&>(BindingReference);
+			MyFieldPath.SetSelfReference();
+		}
+	}
+}
+#endif
+
+
+TArray<FName> FMVVMBlueprintPropertyPath::GetFieldNames(const UClass* InSelfContext) const
+{
+	TArray<FName> Result;
+	Result.Reserve(Paths.Num());
+
+	for (const FMVVMBlueprintFieldPath& Path : Paths)
+	{
+		Result.Add(Path.GetFieldName(InSelfContext));
+	}
+
+	return Result;
+}
+
+
+TArray<UE::MVVM::FMVVMConstFieldVariant> FMVVMBlueprintPropertyPath::GetFields(const UClass* InSelfContext) const
+{
+	TArray<UE::MVVM::FMVVMConstFieldVariant> Result;
+	Result.Reserve(Paths.Num());
+
+	for (const FMVVMBlueprintFieldPath& Path : Paths)
+	{
+		Result.Add(Path.GetField(InSelfContext));
+	}
+
+	return Result;
+}
+
+
+FString FMVVMBlueprintPropertyPath::GetPropertyPath(const UClass* InSelfContext) const
+{
+	TStringBuilder<512> Result;
+	for (const FMVVMBlueprintFieldPath& Path : Paths)
+	{
+		if (Result.Len() > 0)
+		{
+			Result << TEXT('.');
+		}
+		Result << Path.GetFieldName(InSelfContext);
+	}
+	return Result.ToString();
+}
+
+
+bool FMVVMBlueprintPropertyPath::HasFieldInLocalScope() const
+{
+	for (const FMVVMBlueprintFieldPath& Path : Paths)
+	{
+		if (Path.IsFieldLocalScope())
+		{
+			return true;
+		}
+	}
+
+	return false;
+}
\ No newline at end of file
diff --git a/Engine/Plugins/Runtime/ModelViewViewModel/Source/ModelViewViewModelBlueprint/Private/MVVMViewBlueprintCompiler.cpp b/Engine/Plugins/Runtime/ModelViewViewModel/Source/ModelViewViewModelBlueprint/Private/MVVMViewBlueprintCompiler.cpp
--- a/Engine/Plugins/Runtime/ModelViewViewModel/Source/ModelViewViewModelBlueprint/Private/MVVMViewBlueprintCompiler.cpp
+++ b/Engine/Plugins/Runtime/ModelViewViewModel/Source/ModelViewViewModelBlueprint/Private/MVVMViewBlueprintCompiler.cpp
@@ -12,6 +12,7 @@
 #include "MVVMMessageLog.h"
 #include "PropertyPermissionList.h"
 #include "MVVMFunctionGraphHelper.h"
+#include "MVVMWidgetBlueprintExtension_View.h"
 #include "Templates/ValueOrError.h"
 #include "Types/MVVMBindingName.h"
 #include "UObject/LinkerLoad.h"
@@ -28,7 +29,24 @@ FAutoConsoleVariable CVarLogViewCompliedResult(
 	TEXT("After the view is compiled log the compiled bindings and sources.")
 );
 
-FString PropertyPathToString(const UMVVMBlueprintView* BlueprintView, const FMVVMBlueprintPropertyPath& PropertyPath)
+#if UE_BUILD_DEVELOPMENT || UE_BUILD_DEBUG
+FAutoConsoleCommand CVarTestGenerateSetter(
+	TEXT("MVVM.TestGenerateSetter"),
+	TEXT("Generate a setter function base on the input string."),
+	FConsoleCommandWithArgsDelegate::CreateLambda([](const TArray<FString>& Args)
+	{
+		if (Args.Num() < 3)
+		{
+			return;
+		}
+		FMVVMViewBlueprintCompiler::TestGenerateSetter(Args[0], Args[1], Args[2]);
+	})
+);
+#endif
+
+static const FText CouldNotCreateSourceFieldPathFormat = LOCTEXT("CouldNotCreateSourceFieldPath", "Couldn't create the source field path '{0}'. {1}");
+
+FString PropertyPathToString(const UClass* InSelfContext, const UMVVMBlueprintView* BlueprintView, const FMVVMBlueprintPropertyPath& PropertyPath)
 {
 	if (PropertyPath.IsEmpty())
 	{
@@ -56,7 +74,7 @@ FString PropertyPathToString(const UMVVMBlueprintView* BlueprintView, const FMVV
 		Result << TEXT("<none>");
 	}
 
-	FString BasePropertyPath = PropertyPath.GetBasePropertyPath();
+	FString BasePropertyPath = PropertyPath.GetPropertyPath(InSelfContext);
 	if (BasePropertyPath.Len())
 	{
 		Result << TEXT('.');
@@ -65,9 +83,9 @@ FString PropertyPathToString(const UMVVMBlueprintView* BlueprintView, const FMVV
 	return Result.ToString();
 }
 
-FText PropertyPathToText(const UMVVMBlueprintView* BlueprintView, const FMVVMBlueprintPropertyPath& PropertyPath)
+FText PropertyPathToText(const UClass* InSelfContext, const UMVVMBlueprintView* BlueprintView, const FMVVMBlueprintPropertyPath& PropertyPath)
 {
-	return FText::FromString(PropertyPathToString(BlueprintView, PropertyPath));
+	return FText::FromString(PropertyPathToString(InSelfContext, BlueprintView, PropertyPath));
 }
 
 FText GetViewModelIdText(const FMVVMBlueprintPropertyPath& PropertyPath)
@@ -183,6 +201,9 @@ void FMVVMViewBlueprintCompiler::CreateFunctions(UMVVMBlueprintView* BlueprintVi
 	{
 		return;
 	}
+	
+	// Build the list of BP destinations
+	CreateBindingDestinationContexts(BlueprintView);
 
 	if (!GetDefault<UMVVMDeveloperProjectSettings>()->bAllowGeneratedViewModelSetter)
 	{
@@ -615,6 +636,134 @@ void FMVVMViewBlueprintCompiler::CreateFunctionsDeclaration(const FWidgetBluepri
 }
 
 
+void FMVVMViewBlueprintCompiler::CreateBindingDestinationContexts(UMVVMBlueprintView* BlueprintView)
+{
+	if (!bAreSourcesCreatorValid || !bAreSourceContextsValid || !bIsBindingsValid)
+	{
+		return;
+	}
+
+	BindingDestinationContexts.Empty();
+
+	// Use the Skeleton class. The class bind and not all functions are generated yet
+	UWidgetBlueprintGeneratedClass* NewSkeletonClass = Cast<UWidgetBlueprintGeneratedClass>(WidgetBlueprintCompilerContext.Blueprint->SkeletonGeneratedClass);
+	if (NewSkeletonClass == nullptr)
+	{
+		WidgetBlueprintCompilerContext.MessageLog.Error(*LOCTEXT("InvalidNewClass", "Internal error. The skeleton class is not valid.").ToString());
+		return;
+	}
+
+	const int32 NumBindings = BlueprintView->GetNumBindings();
+	for (int32 Index = 0; Index < NumBindings; ++Index)
+	{
+		FMVVMBlueprintViewBinding* BindingPtr = BlueprintView->GetBindingAt(Index);
+		if (BindingPtr == nullptr)
+		{
+			WidgetBlueprintCompilerContext.MessageLog.Error(*FText::Format(LOCTEXT("InvalidBindingIndex", "Internal error. Invalid binding index given."), Index).ToString());
+			bIsBindingsValid = false;
+			continue;
+		}
+
+		FMVVMBlueprintViewBinding& Binding = *BindingPtr;
+		if (!Binding.bCompile)
+		{
+			continue;
+		}
+
+		auto AddDestination = [Index, BindingPtr, BlueprintView, NewSkeletonClass, this](bool bIsForwardBinding)
+		{
+			FBindingDestinationContext& NewDestinationContext = BindingDestinationContexts.AddDefaulted_GetRef();
+			NewDestinationContext.BindingIndex = Index;
+			NewDestinationContext.bIsForwardBinding = bIsForwardBinding;
+
+			const FMVVMBlueprintPropertyPath& DestinationPath = bIsForwardBinding ? BindingPtr->DestinationPath : BindingPtr->SourcePath;
+			TArray<UE::MVVM::FMVVMConstFieldVariant> SkeletalSetterPath = CreateBindingDestinationPath(BlueprintView, NewSkeletonClass, DestinationPath);
+			if (!IsPropertyPathValid(SkeletalSetterPath))
+			{
+				AddMessageForBinding(*BindingPtr, BlueprintView, FText::Format(LOCTEXT("PropertyPathIsInvalid", "The property path '{0}' is invalid."),
+					PropertyPathToText(NewSkeletonClass, BlueprintView, DestinationPath)),
+					EBindingMessageType::Error
+				);
+				bIsBindingsValid = false;
+				return;
+			}
+			
+			// Generate the path with property converted to BP function
+			TValueOrError<TArray<FMVVMConstFieldVariant>, FText> SkeletalGeneratedField = FieldPathHelper::GenerateFieldPathList(SkeletalSetterPath, false);
+			if (SkeletalGeneratedField.HasError())
+			{
+				AddMessageForBinding(*BindingPtr, BlueprintView, FText::Format(Private::CouldNotCreateSourceFieldPathFormat
+					, ::UE::MVVM::FieldPathHelper::ToText(SkeletalSetterPath)
+					, SkeletalGeneratedField.GetError())
+					, EBindingMessageType::Error
+				);
+				bIsBindingsValid = false;
+				return;
+			}
+
+			NewDestinationContext.bCanBeSetInNative = CanBeSetInNative(SkeletalGeneratedField.GetValue());
+
+			// If the destination can't be done cpp, we need to generate a BP function to set the value.
+			if (!NewDestinationContext.bCanBeSetInNative)
+			{
+				const FProperty* SetterProperty = nullptr;
+				if (SkeletalGeneratedField.GetValue().Num() > 0 && ensure(SkeletalGeneratedField.GetValue().Last().IsProperty()))
+				{
+					SetterProperty = SkeletalGeneratedField.GetValue().Last().GetProperty();
+				}
+
+				if (SetterProperty == nullptr)
+				{
+					AddMessageForBinding(*BindingPtr, BlueprintView, FText::Format(LOCTEXT("CantGetSetter", "Internal Error. The setter function was not created. {0}"),
+						PropertyPathToText(NewSkeletonClass, BlueprintView, DestinationPath)),
+						EBindingMessageType::Error
+					);
+					bIsBindingsValid = false;
+					return;
+				}
+
+				// create a setter function to be called from native. For now we follow the convention of Setter(Conversion(Getter))
+				UEdGraph* GeneratedSetterGraph = UE::MVVM::FunctionGraphHelper::CreateIntermediateFunctionGraph(WidgetBlueprintCompilerContext, FString::Printf(TEXT("__Setter_%s"), *SetterProperty->GetName()), EFunctionFlags::FUNC_None, TEXT("AutogeneratedSetter"), false);
+				if (GeneratedSetterGraph == nullptr)
+				{
+					AddMessageForBinding(*BindingPtr, BlueprintView, FText::Format(LOCTEXT("CantCreateSetter", "Internal Error. The setter function was not created. {0}"),
+						PropertyPathToText(NewSkeletonClass, BlueprintView, DestinationPath)),
+						EBindingMessageType::Error
+					);
+					bIsBindingsValid = false;
+					return;
+				}
+
+				UE::MVVM::FunctionGraphHelper::AddFunctionArgument(GeneratedSetterGraph, SetterProperty, "NewValue");
+
+				// set GeneratedSetterFunction. Use the SkeletalSetterPath here to use the setter will be generated when the function is generated.
+				if (!UE::MVVM::FunctionGraphHelper::GenerateIntermediateSetter(WidgetBlueprintCompilerContext, GeneratedSetterGraph, SkeletalSetterPath))
+				{
+					AddMessageForBinding(*BindingPtr, BlueprintView, FText::Format(LOCTEXT("CantGeneratedSetter", "Internal Error. The setter function was not generated. {0}"),
+						PropertyPathToText(NewSkeletonClass, BlueprintView, DestinationPath)),
+						EBindingMessageType::Error
+					);
+					bIsBindingsValid = false;
+					return;
+				}
+
+				// the new path can only be set later, once the function is compiled.
+				NewDestinationContext.GeneratedFunctionName = GeneratedSetterGraph->GetFName();
+			}
+		};
+
+		if (IsForwardBinding(Binding.BindingType))
+		{
+			AddDestination(true);
+		}
+		if (IsBackwardBinding(Binding.BindingType))
+		{
+			AddDestination(false);
+		}
+	}
+}
+
+
 bool FMVVMViewBlueprintCompiler::PreCompile(UWidgetBlueprintGeneratedClass* Class, UMVVMBlueprintView* BlueprintView)
 {
 	if (!bAreSourcesCreatorValid || !bAreSourceContextsValid || !bIsBindingsValid)
@@ -701,7 +850,7 @@ bool FMVVMViewBlueprintCompiler::PreCompileBindingSources(UWidgetBlueprintGenera
 			{
 				AddMessageForBinding(Binding, BlueprintView,
 					FText::Format(LOCTEXT("PropertyPathInvalidWithReason", "The property path '{0}' is invalid. {1}"),
-						PropertyPathToText(BlueprintView, Binding.SourcePath),
+						PropertyPathToText(Class, BlueprintView, Binding.SourcePath),
 						CreatedBindingSourceContext.GetError()
 					),
 					EBindingMessageType::Error,
@@ -716,7 +865,7 @@ bool FMVVMViewBlueprintCompiler::PreCompileBindingSources(UWidgetBlueprintGenera
 				AddMessageForBinding(Binding,
 					BlueprintView,
 					FText::Format(LOCTEXT("PropertyPathIsInvalid", "The property path '{0}' is invalid."), 
-						PropertyPathToText(BlueprintView, Binding.SourcePath)
+						PropertyPathToText(Class, BlueprintView, Binding.SourcePath)
 					),
 					EBindingMessageType::Error,
 					ArgumentName
@@ -787,7 +936,7 @@ bool FMVVMViewBlueprintCompiler::PreCompileBindingSources(UWidgetBlueprintGenera
 
 			TArrayView<FMVVMBlueprintFieldPath const> MVVMBindingPath = Path.GetFieldPaths();
 			TArray< FDelegateRuntimeBinding > LegacyBindings = Class->Bindings;
-			FName MVVMFieldName = Path.GetPaths().Last();
+			FName MVVMFieldName = Path.GetFieldNames(Class).Last();
 			FName MVVMObjectName = Path.GetWidgetName();
 
 			if (Path.GetFieldPaths().Last().GetBindingKind() == EBindingKind::Function)
@@ -796,14 +945,14 @@ bool FMVVMViewBlueprintCompiler::PreCompileBindingSources(UWidgetBlueprintGenera
 			}
 
 			// If the first field is a UserWidget, we know this property resides in a nested UserWidget.
-			if (MVVMBindingPath[0].GetParentClass() && MVVMBindingPath[0].GetParentClass()->IsChildOf(UUserWidget::StaticClass()) && MVVMBindingPath.Num() > 1)
+			if (MVVMBindingPath[0].GetParentClass(Class) && MVVMBindingPath[0].GetParentClass(Class)->IsChildOf(UUserWidget::StaticClass()) && MVVMBindingPath.Num() > 1)
 			{
-				if (UWidgetBlueprintGeneratedClass* NestedBPGClass = Cast<UWidgetBlueprintGeneratedClass>(MVVMBindingPath[MVVMBindingPath.Num() - 2].GetParentClass()))
+				if (UWidgetBlueprintGeneratedClass* NestedBPGClass = Cast<UWidgetBlueprintGeneratedClass>(MVVMBindingPath[MVVMBindingPath.Num() - 2].GetParentClass(Class)))
 				{
 					LegacyBindings = NestedBPGClass->Bindings;
 
 					// We can't use Path.GetWidgetName() when we are dealing with nested UserWidgets, because it refers to the topmost UserWidget.
-					MVVMObjectName = MVVMBindingPath[MVVMBindingPath.Num() - 2].GetFieldName();
+					MVVMObjectName = MVVMBindingPath[MVVMBindingPath.Num() - 2].GetFieldName(Class);
 				}
 				else
 				{
@@ -1151,12 +1300,44 @@ bool FMVVMViewBlueprintCompiler::PreCompileBindings(UWidgetBlueprintGeneratedCla
 
 		TArray<UE::MVVM::FMVVMConstFieldVariant> SetterPath;
 		{
-			const FMVVMBlueprintPropertyPath& DestinationPath = BindingSourceContext.bIsForwardBinding ? Binding.DestinationPath : Binding.SourcePath;
-			SetterPath = CreateBindingDestinationPath(BlueprintView, Class, DestinationPath);
+			const FBindingDestinationContext* DestinationContext = BindingDestinationContexts.FindByPredicate([BindingIndex = BindingSourceContext.BindingIndex, bIsForwardBinding = BindingSourceContext.bIsForwardBinding](const FBindingDestinationContext& Other)
+				{
+					return Other.BindingIndex == BindingIndex && Other.bIsForwardBinding == bIsForwardBinding;
+				});
+
+			if (DestinationContext == nullptr)
+			{
+				AddMessageForBinding(Binding, BlueprintView, LOCTEXT("CouldNotFindDestination", "Could not find the pre compiled destination."), EBindingMessageType::Error);
+				bIsBindingsValid = false;
+				continue;
+			}
+
+			if (!DestinationContext->bCanBeSetInNative)
+			{
+				check(!DestinationContext->GeneratedFunctionName.IsNone());
+				UFunction* FoundFunction = Class->FindFunctionByName(DestinationContext->GeneratedFunctionName);
+				if (FoundFunction == nullptr)
+				{
+					AddMessageForBinding(Binding, BlueprintView, LOCTEXT("CouldNotFindDestinationBPFunction", "Could not find the generated destination function."), EBindingMessageType::Error);
+					bIsBindingsValid = false;
+					continue;
+				}
+				FMVVMBlueprintPropertyPath DestinationPPropertyPath;
+				DestinationPPropertyPath.SetWidgetName(Class->ClassGeneratedBy->GetFName());
+				DestinationPPropertyPath.SetPropertyPath(BlueprintView->GetOuterUMVVMWidgetBlueprintExtension_View()->GetWidgetBlueprint(), UE::MVVM::FMVVMConstFieldVariant(FoundFunction));
+				SetterPath = CreateBindingDestinationPath(BlueprintView, Class, DestinationPPropertyPath);
+			}
+			else
+			{
+				const FMVVMBlueprintPropertyPath& DestinationPath = BindingSourceContext.bIsForwardBinding ? Binding.DestinationPath : Binding.SourcePath;
+				SetterPath = CreateBindingDestinationPath(BlueprintView, Class, DestinationPath);
+			}
+
 			if (!IsPropertyPathValid(SetterPath))
 			{
+				const FMVVMBlueprintPropertyPath& DestinationPath = BindingSourceContext.bIsForwardBinding ? Binding.DestinationPath : Binding.SourcePath;
 				AddMessageForBinding(Binding, BlueprintView, FText::Format(LOCTEXT("PropertyPathIsInvalid", "The property path '{0}' is invalid."),
-					PropertyPathToText(BlueprintView, DestinationPath)),
+					PropertyPathToText(Class, BlueprintView, DestinationPath)),
 					EBindingMessageType::Error
 				);
 				bIsBindingsValid = false;
@@ -1394,19 +1575,17 @@ TValueOrError<FMVVMViewBlueprintCompiler::FCompiledBinding, FText> FMVVMViewBlue
 		// Sanity check
 		if (bIsComplexBinding && !BindingHelper::IsValidForComplexRuntimeConversion(ConversionFunction))
 		{
-			return MakeError(LOCTEXT("ConversionFunctionIsNotComplex", "Internal Error. The complex conversion function does not respect the prerequist."));
+			return MakeError(LOCTEXT("ConversionFunctionIsNotComplex", "Internal Error. The complex conversion function does not respect the prerequisite."));
 		}
 	}
 
 	if (!Result.bIsConversionFunctionComplex)
 	{
-		static const FText CouldNotCreateSourceFieldPathFormat = LOCTEXT("CouldNotCreateSourceFieldPath", "Couldn't create the source field path '{0}'. {1}");
-
 		// Generate a path to read the value at runtime
 		TValueOrError<TArray<FMVVMConstFieldVariant>, FText> GeneratedField = FieldPathHelper::GenerateFieldPathList(GetterFields, true);
 		if (GeneratedField.HasError())
 		{
-			return MakeError(FText::Format(CouldNotCreateSourceFieldPathFormat
+			return MakeError(FText::Format(Private::CouldNotCreateSourceFieldPathFormat
 				, ::UE::MVVM::FieldPathHelper::ToText(GetterFields)
 				, GeneratedField.GetError()));
 		}
@@ -1414,7 +1593,7 @@ TValueOrError<FMVVMViewBlueprintCompiler::FCompiledBinding, FText> FMVVMViewBlue
 		TValueOrError<FCompiledBindingLibraryCompiler::FFieldPathHandle, FText> FieldPathResult = BindingLibraryCompiler.AddFieldPath(GeneratedField.GetValue(), true);
 		if (FieldPathResult.HasError())
 		{
-			return MakeError(FText::Format(CouldNotCreateSourceFieldPathFormat
+			return MakeError(FText::Format(Private::CouldNotCreateSourceFieldPathFormat
 				, ::UE::MVVM::FieldPathHelper::ToText(GetterFields)
 				, FieldPathResult.GetError()));
 		}
@@ -1485,7 +1664,7 @@ TValueOrError<FMVVMViewBlueprintCompiler::FBindingSourceContext, FText> FMVVMVie
 		check(Result.UserWidgetPropertyContextIndex != INDEX_NONE);
 
 		Result.SourceClass = CompilerUserWidgetPropertyContexts[Result.UserWidgetPropertyContextIndex].Class;
-		Result.PropertyPath = CreatePropertyPath(Class, CompilerUserWidgetPropertyContexts[Result.UserWidgetPropertyContextIndex].PropertyName, PropertyPath.GetFields());
+		Result.PropertyPath = CreatePropertyPath(Class, CompilerUserWidgetPropertyContexts[Result.UserWidgetPropertyContextIndex].PropertyName, PropertyPath.GetFields(Class));
 	}
 	else if (PropertyPath.IsFromWidget())
 	{
@@ -1495,14 +1674,14 @@ TValueOrError<FMVVMViewBlueprintCompiler::FBindingSourceContext, FText> FMVVMVie
 		{
 			Result.UserWidgetPropertyContextIndex = INDEX_NONE;
 			Result.SourceClass = const_cast<UWidgetBlueprintGeneratedClass*>(Class);
-			Result.PropertyPath = CreatePropertyPath(Class, FName(), PropertyPath.GetFields());
+			Result.PropertyPath = CreatePropertyPath(Class, FName(), PropertyPath.GetFields(Class));
 		}
 		else
 		{
 			Result.UserWidgetPropertyContextIndex = CompilerUserWidgetPropertyContexts.IndexOfByPredicate([SourceName](const FCompilerUserWidgetPropertyContext& Other) { return Other.PropertyName == SourceName; });
 			check(Result.UserWidgetPropertyContextIndex != INDEX_NONE);
 			Result.SourceClass = CompilerUserWidgetPropertyContexts[Result.UserWidgetPropertyContextIndex].Class;
-			Result.PropertyPath = CreatePropertyPath(Class, CompilerUserWidgetPropertyContexts[Result.UserWidgetPropertyContextIndex].PropertyName, PropertyPath.GetFields());
+			Result.PropertyPath = CreatePropertyPath(Class, CompilerUserWidgetPropertyContexts[Result.UserWidgetPropertyContextIndex].PropertyName, PropertyPath.GetFields(Class));
 		}
 	}
 	else
@@ -1673,7 +1852,7 @@ TArray<FMVVMConstFieldVariant> FMVVMViewBlueprintCompiler::CreateBindingDestinat
 		const int32 DestinationVariableContextIndex = CompilerUserWidgetPropertyContexts.IndexOfByPredicate([DestinationName](const FCompilerUserWidgetPropertyContext& Other) { return Other.PropertyName == DestinationName; });
 		check(DestinationVariableContextIndex != INDEX_NONE);
 
-		return CreatePropertyPath(Class, CompilerUserWidgetPropertyContexts[DestinationVariableContextIndex].PropertyName, PropertyPath.GetFields());
+		return CreatePropertyPath(Class, CompilerUserWidgetPropertyContexts[DestinationVariableContextIndex].PropertyName, PropertyPath.GetFields(Class));
 	}
 	else if (PropertyPath.IsFromWidget())
 	{
@@ -1682,27 +1861,27 @@ TArray<FMVVMConstFieldVariant> FMVVMViewBlueprintCompiler::CreateBindingDestinat
 		const bool bSourceIsUserWidget = DestinationName == Class->ClassGeneratedBy->GetFName();
 		if (bSourceIsUserWidget)
 		{
-			return CreatePropertyPath(Class, FName(), PropertyPath.GetFields());
+			return CreatePropertyPath(Class, FName(), PropertyPath.GetFields(Class));
 		}
 		else
 		{
 			const int32 DestinationVariableContextIndex = CompilerUserWidgetPropertyContexts.IndexOfByPredicate([DestinationName](const FCompilerUserWidgetPropertyContext& Other) { return Other.PropertyName == DestinationName; });
 			if (ensureAlwaysMsgf(DestinationVariableContextIndex != INDEX_NONE, TEXT("Could not find source context for destination '%s'"), *DestinationName.ToString()))
 			{
-				return CreatePropertyPath(Class, CompilerUserWidgetPropertyContexts[DestinationVariableContextIndex].PropertyName, PropertyPath.GetFields());
+				return CreatePropertyPath(Class, CompilerUserWidgetPropertyContexts[DestinationVariableContextIndex].PropertyName, PropertyPath.GetFields(Class));
 			}
 		}
 	}
 	else
 	{
 		ensureAlwaysMsgf(false, TEXT("Not supported yet."));
-		return CreatePropertyPath(Class, FName(), PropertyPath.GetFields());
+		return CreatePropertyPath(Class, FName(), PropertyPath.GetFields(Class));
 	}
 
 	return TArray<FMVVMConstFieldVariant>();
 }
 
-TArray<FMVVMConstFieldVariant> FMVVMViewBlueprintCompiler::CreatePropertyPath(const UClass* Class, FName PropertyName, TArray<FMVVMConstFieldVariant> Properties) const
+TArray<FMVVMConstFieldVariant> FMVVMViewBlueprintCompiler::CreatePropertyPath(const UClass* Class, FName PropertyName, TArray<FMVVMConstFieldVariant> Properties)
 {
 	if (PropertyName.IsNone())
 	{
@@ -1716,7 +1895,7 @@ TArray<FMVVMConstFieldVariant> FMVVMViewBlueprintCompiler::CreatePropertyPath(co
 }
 
 
-bool FMVVMViewBlueprintCompiler::IsPropertyPathValid(TArrayView<const FMVVMConstFieldVariant> PropertyPath) const
+bool FMVVMViewBlueprintCompiler::IsPropertyPathValid(TArrayView<const FMVVMConstFieldVariant> PropertyPath)
 {
 	for (const FMVVMConstFieldVariant& Field : PropertyPath)
 	{
@@ -1736,6 +1915,89 @@ bool FMVVMViewBlueprintCompiler::IsPropertyPathValid(TArrayView<const FMVVMConst
 	return true;
 }
 
+bool FMVVMViewBlueprintCompiler::CanBeSetInNative(TArrayView<const FMVVMConstFieldVariant> PropertyPath)
+{
+	check(IsPropertyPathValid(PropertyPath));
+	
+	for (int32 Index = PropertyPath.Num() - 1; Index >= 0; --Index)
+	{
+		const FMVVMConstFieldVariant& Variant = PropertyPath[Index];
+		// Stop the algo if the path is already a function.
+		if (Variant.IsFunction())
+		{
+			return true;
+		}
+		// If the BP is defined in BP and has Net flags or FieldNotify flag, then the VaraibleSet K2Node need to be used to generate the proper byte-code.
+		if (Variant.IsProperty())
+		{
+			// If it's an object then the path before the object doesn't matter.
+			if (const FObjectPropertyBase* PropertyBase = CastField<const FObjectPropertyBase>(Variant.GetProperty()))
+			{
+				bool bLastPath = Index >= PropertyPath.Num() - 1;
+				if (!bLastPath)
+				{
+					return true;
+				}
+			}
+
+			if (Cast<UBlueprintGeneratedClass>(Variant.GetProperty()->GetOwnerStruct()))
+			{
+				if (Variant.GetProperty()->HasMetaData(FName("FieldNotify")) || Variant.GetProperty()->HasAnyPropertyFlags(CPF_Net))
+				{
+					return false;
+				}
+			}
+		}
+	}
+	return true;
+}
+
+void FMVVMViewBlueprintCompiler::TestGenerateSetter(FStringView ObjectName, FStringView FieldPath, FStringView FunctionName)
+{
+#if UE_BUILD_DEVELOPMENT || UE_BUILD_DEBUG
+	UWidgetBlueprint* WidgetBlueprint = nullptr;
+	{
+		UObject* FoundObject = FindObject<UObject>(nullptr, ObjectName.GetData(), false);
+		WidgetBlueprint = Cast<UWidgetBlueprint>(FoundObject);
+	}
+
+	if (WidgetBlueprint == nullptr)
+	{
+		return;
+	}
+
+	UWidgetBlueprintGeneratedClass* NewSkeletonClass = Cast<UWidgetBlueprintGeneratedClass>(WidgetBlueprint->SkeletonGeneratedClass);
+	TValueOrError<TArray<FMVVMConstFieldVariant>, FText> SkeletalSetterPathResult = FieldPathHelper::GenerateFieldPathList(NewSkeletonClass, FieldPath, false);
+	if (SkeletalSetterPathResult.HasError())
+	{
+		return;
+	}
+	if (!IsPropertyPathValid(SkeletalSetterPathResult.GetValue()))
+	{
+		return;
+	}
+
+	UEdGraph* GeneratedSetterGraph = UE::MVVM::FunctionGraphHelper::CreateFunctionGraph(WidgetBlueprint, FunctionName, EFunctionFlags::FUNC_None, TEXT(""), false);
+	if (GeneratedSetterGraph == nullptr)
+	{
+		return;
+	}
+
+	const FProperty* SetterProperty = nullptr;
+	if (SkeletalSetterPathResult.GetValue().Num() > 0 && ensure(SkeletalSetterPathResult.GetValue().Last().IsProperty()))
+	{
+		SetterProperty = SkeletalSetterPathResult.GetValue().Last().GetProperty();
+	}
+	if (SetterProperty == nullptr)
+	{
+		return;
+	}
+
+	UE::MVVM::FunctionGraphHelper::AddFunctionArgument(GeneratedSetterGraph, SetterProperty, "NewValue");
+	UE::MVVM::FunctionGraphHelper::GenerateSetter(WidgetBlueprint, GeneratedSetterGraph, SkeletalSetterPathResult.GetValue());
+#endif
+}
+
 } //namespace
 
 #undef LOCTEXT_NAMESPACE
diff --git a/Engine/Plugins/Runtime/ModelViewViewModel/Source/ModelViewViewModelBlueprint/Private/MVVMViewBlueprintCompiler.h b/Engine/Plugins/Runtime/ModelViewViewModel/Source/ModelViewViewModelBlueprint/Private/MVVMViewBlueprintCompiler.h
--- a/Engine/Plugins/Runtime/ModelViewViewModel/Source/ModelViewViewModelBlueprint/Private/MVVMViewBlueprintCompiler.h
+++ b/Engine/Plugins/Runtime/ModelViewViewModel/Source/ModelViewViewModelBlueprint/Private/MVVMViewBlueprintCompiler.h
@@ -32,6 +32,7 @@ private:
 	struct FCompiledBinding;
 	struct FCompilerBinding;
 	struct FBindingSourceContext;
+	struct FBindingDestinationContext;
 
 public:
 	FMVVMViewBlueprintCompiler(FWidgetBlueprintCompilerContext& InCreationContext)
@@ -58,10 +59,14 @@ public:
 	/** Compile the library and fill the view and viewclass */
 	bool Compile(UWidgetBlueprintGeneratedClass* Class, UMVVMBlueprintView* BlueprintView, UMVVMViewClass* ViewExtension);
 
+
+	static void TestGenerateSetter(FStringView ObjectName, FStringView FieldPath, FStringView FunctionName);
+
 private:
 	void CreateWidgetMap(const FWidgetBlueprintCompilerContext::FCreateVariableContext& Context, UMVVMBlueprintView* BlueprintView);
 	void CreateSourceLists(const FWidgetBlueprintCompilerContext::FCreateVariableContext& Context, UMVVMBlueprintView* BlueprintView);
 	void CreateFunctionsDeclaration(const FWidgetBlueprintCompilerContext::FCreateVariableContext& Context, UMVVMBlueprintView* BlueprintView);
+	void CreateBindingDestinationContexts(UMVVMBlueprintView* BlueprintView);
 
 	bool PreCompileBindingSources(UWidgetBlueprintGeneratedClass* Class, UMVVMBlueprintView* BlueprintView);
 	bool CompileBindingSources(const FCompiledBindingLibraryCompiler::FCompileResult& CompileResult, UWidgetBlueprintGeneratedClass* Class, UMVVMBlueprintView* BlueprintView, UMVVMViewClass* ViewExtension);
@@ -81,9 +86,9 @@ private:
 	TArray<FMVVMConstFieldVariant> CreateBindingDestinationPath(const UMVVMBlueprintView* BlueprintView, const UWidgetBlueprintGeneratedClass* Class, const FMVVMBlueprintPropertyPath& PropertyPath) const;
 	TValueOrError<FCompiledBinding, FText> CreateCompiledBinding(const UWidgetBlueprintGeneratedClass* Class, TArrayView<const UE::MVVM::FMVVMConstFieldVariant> GetterFields, TArrayView<const UE::MVVM::FMVVMConstFieldVariant> SetterFields, const UFunction* ConversionFunction, bool bIsComplexBinding);
 
-
-	TArray<FMVVMConstFieldVariant> CreatePropertyPath(const UClass* Class, FName PropertyName, TArray<FMVVMConstFieldVariant> Properties) const;
-	bool IsPropertyPathValid(TArrayView<const FMVVMConstFieldVariant> PropertyPath) const;
+	static TArray<FMVVMConstFieldVariant> CreatePropertyPath(const UClass* Class, FName PropertyName, TArray<FMVVMConstFieldVariant> Properties);
+	static bool IsPropertyPathValid(TArrayView<const FMVVMConstFieldVariant> PropertyPath);
+	static bool CanBeSetInNative(TArrayView<const FMVVMConstFieldVariant> PropertyPath);
 
 private:
 	/**
@@ -143,6 +148,7 @@ private:
 		FCompiledBindingLibraryCompiler::FFieldPathHandle SourceRead;
 		FCompiledBindingLibraryCompiler::FFieldPathHandle DestinationWrite;
 		FCompiledBindingLibraryCompiler::FFieldPathHandle ConversionFunction;
+		FCompiledBindingLibraryCompiler::FFieldPathHandle ExecutionFunction;
 
 		bool bIsConversionFunctionComplex = false;
 	};
@@ -200,6 +206,21 @@ private:
 	};
 	TArray<FBindingSourceContext> BindingSourceContexts;
 
+	/**
+	 * Destination path of a binding.
+	 * The info needs to be validate before we generate the functions list.
+	 */
+	struct FBindingDestinationContext
+	{
+		int32 BindingIndex = INDEX_NONE;
+		bool bIsForwardBinding = false;
+
+		bool bCanBeSetInNative = true;
+		FName GeneratedFunctionName;
+
+	};
+	TArray<FBindingDestinationContext> BindingDestinationContexts;
+
 	TMap<FName, UWidget*> WidgetNameToWidgetPointerMap;
 	FWidgetBlueprintCompilerContext& WidgetBlueprintCompilerContext;
 	FCompiledBindingLibraryCompiler BindingLibraryCompiler;
diff --git a/Engine/Plugins/Runtime/ModelViewViewModel/Source/ModelViewViewModelBlueprint/Public/MVVMBlueprintView.h b/Engine/Plugins/Runtime/ModelViewViewModel/Source/ModelViewViewModelBlueprint/Public/MVVMBlueprintView.h
--- a/Engine/Plugins/Runtime/ModelViewViewModel/Source/ModelViewViewModelBlueprint/Public/MVVMBlueprintView.h
+++ b/Engine/Plugins/Runtime/ModelViewViewModel/Source/ModelViewViewModelBlueprint/Public/MVVMBlueprintView.h
@@ -96,6 +96,7 @@ public:
 #endif
 
 	virtual void PostLoad() override;
+	virtual void Serialize(FArchive& Ar) override;
 
 	DECLARE_EVENT(UMVVMBlueprintView, FOnBindingsUpdated);
 	FOnBindingsUpdated OnBindingsUpdated;
diff --git a/Engine/Plugins/Runtime/ModelViewViewModel/Source/ModelViewViewModelBlueprint/Public/MVVMPropertyPath.h b/Engine/Plugins/Runtime/ModelViewViewModel/Source/ModelViewViewModelBlueprint/Public/MVVMPropertyPath.h
--- a/Engine/Plugins/Runtime/ModelViewViewModel/Source/ModelViewViewModelBlueprint/Public/MVVMPropertyPath.h
+++ b/Engine/Plugins/Runtime/ModelViewViewModel/Source/ModelViewViewModelBlueprint/Public/MVVMPropertyPath.h
@@ -8,6 +8,8 @@
 
 #include "MVVMPropertyPath.generated.h"
 
+class UBlueprint;
+
 /**
  * A single item in a Property Path
  */
@@ -26,136 +28,40 @@ private:
 	EBindingKind BindingKind = EBindingKind::Function;
 
 public:
-	/** Get the binding name, resolves reference deprecation / redirectors / etc before returning */
-	FName GetFieldName() const
-	{
-		// Resolve any redirectors
-		if (!BindingReference.GetMemberName().IsNone())
-		{
-			if (BindingKind == EBindingKind::Property)
-			{
-				if (BindingReference.IsLocalScope())
-				{
-					if (UPackage* Package = BindingReference.GetMemberParentPackage())
-					{
-						UObjectBase* FoundObject = FindObjectWithOuter(Package, UScriptStruct::StaticClass(), *BindingReference.GetMemberScopeName());
-						if (UScriptStruct* Struct = Cast<UScriptStruct>(static_cast<UObject*>(FoundObject)))
-						{
-							if (const FProperty* FoundProperty = FindUFieldOrFProperty<FProperty>(Struct, BindingReference.GetMemberName(), EFieldIterationFlags::IncludeAll))
-							{
-								return FoundProperty->GetFName();
-							}
-						}
-					}
-				}
-				else if (BindingReference.ResolveMember<FProperty>())
-				{
-					return BindingReference.GetMemberName();
-				}
-			}
-			else if (BindingKind == EBindingKind::Function)
-			{
-				if (BindingReference.ResolveMember<UFunction>())
-				{
-					return BindingReference.GetMemberName();
-				}
-			}
-		}
+	FMVVMBlueprintFieldPath() = default;
 
-		return FName();
-	}
+	MODELVIEWVIEWMODELBLUEPRINT_API FMVVMBlueprintFieldPath(const UBlueprint* InContext, UE::MVVM::FMVVMConstFieldVariant InField);
+
+	/** Get the binding name, resolves reference deprecation / redirectors / etc before returning */
+	MODELVIEWVIEWMODELBLUEPRINT_API FName GetFieldName(const UClass* SelfContext) const;
 
 	/** */
-	UE::MVVM::FMVVMConstFieldVariant GetField() const
-	{
-		if (!BindingReference.GetMemberName().IsNone())
-		{
-			if (BindingKind == EBindingKind::Property)
-			{
-				if (BindingReference.IsLocalScope())
-				{
-					if (UPackage* Package = BindingReference.GetMemberParentPackage())
-					{
-						UObjectBase* FoundObject = FindObjectWithOuter(Package, UScriptStruct::StaticClass(), *BindingReference.GetMemberScopeName());
-						if (UScriptStruct* Struct = Cast<UScriptStruct>(static_cast<UObject*>(FoundObject)))
-						{
-							if (const FProperty* FoundProperty = FindUFieldOrFProperty<FProperty>(Struct, BindingReference.GetMemberName(), EFieldIterationFlags::IncludeAll))
-							{
-								return UE::MVVM::FMVVMConstFieldVariant(FoundProperty);
-							}
-						}
-					}
-				}
-				else
-				{
-					return UE::MVVM::FMVVMConstFieldVariant(BindingReference.ResolveMember<FProperty>());
-				}
-			}
-			else if (BindingKind == EBindingKind::Function)
-			{
-				return UE::MVVM::FMVVMConstFieldVariant(BindingReference.ResolveMember<UFunction>());
-			}
-		}
-		return UE::MVVM::FMVVMConstFieldVariant();
-	}
+	MODELVIEWVIEWMODELBLUEPRINT_API UE::MVVM::FMVVMConstFieldVariant GetField(const UClass* SelfContext) const;
 
 	bool IsFieldLocalScope() const
 	{
 		return BindingReference.IsLocalScope();
 	}
 
-	UClass* GetParentClass() const
-	{
-		return BindingReference.GetMemberParentClass();
-	}
+	MODELVIEWVIEWMODELBLUEPRINT_API UClass* GetParentClass(const UClass* SelfContext) const;
 
 	EBindingKind GetBindingKind() const
 	{
 		return BindingKind;
 	}
 
+#if WITH_EDITOR
 	/** */
-	void SetBindingReference(UE::MVVM::FMVVMConstFieldVariant InField)
-	{
-		if (InField.IsEmpty())
-		{
-			Reset();
-			return;
-		}
-
-		BindingReference = CreateMemberReference(InField);
-
-		if (InField.IsProperty())
-		{
-			BindingKind = EBindingKind::Property;
-		}
-		else if (InField.IsFunction())
-		{
-			BindingKind = EBindingKind::Function;
-		}
-		else
-		{
-			ensureAlwaysMsgf(false, TEXT("Binding to field of unknown type!"));
-		}
-	}
-
-	/** */
-	void Reset()
-	{
-		BindingReference = FMemberReference();
-	}
-
+	void SetDeprecatedBindingReference(const FMemberReference& InBindingReference, EBindingKind InBindingKind);
 	/** */
-	void SetDeprecatedBindingReference(const FMemberReference& InBindingReference, EBindingKind InBindingKind)
-	{
-		BindingReference = InBindingReference;
-		BindingKind = InBindingKind;
-	}
+	void SetDeprecatedSelfReference(const UBlueprint* Context);
+#endif
 
 public:
 	bool operator==(const FMVVMBlueprintFieldPath& Other) const
 	{
-		return BindingReference.GetMemberName() == Other.BindingReference.GetMemberName();
+		return BindingReference.IsSameReference(Other.BindingReference)
+			&& BindingKind == Other.BindingKind;
 	}
 	bool operator!=(const FMVVMBlueprintFieldPath& Other) const
 	{
@@ -163,41 +69,9 @@ public:
 	}
 
 private:
-	/**
-	 * Create a serializable member reference from this field.
-	 */
-	static FMemberReference CreateMemberReference(UE::MVVM::FMVVMConstFieldVariant InField)
-	{
-		FMemberReference BindingReference = FMemberReference();
-		if (InField.IsProperty())
-		{
-			const FProperty* Property = InField.GetProperty();
-			if (Property->GetOwnerClass())
-			{
-				BindingReference.SetFromField<FProperty>(Property, false);
-			}
-			else
-			{
-				BindingReference.SetGlobalField(Property->GetFName(), Property->GetOutermost());
-				BindingReference.SetLocalMember(Property->GetFName(), Property->GetOwnerStruct(), FGuid());
-			}
-		}
-		else if (InField.IsFunction())
-		{
-			// Functions should never be self context references
-			const UFunction* Function = InField.GetFunction();
-			bool bSelfContext = false;
-			BindingReference.SetFromField<UFunction>(Function, bSelfContext);
-		}
-		return BindingReference;
-	}
+	UE::MVVM::FMVVMConstFieldVariant GetFieldInternal(const UClass* SelfContext) const;
 };
 
-inline uint32 GetTypeHash(const FMVVMBlueprintFieldPath& FieldPath)
-{
-	return GetTypeHash(FieldPath.GetFieldName());
-}
-
 /**
  * Base path to properties for MVVM view models and widgets.
  * 
@@ -233,72 +107,45 @@ public:
 		return Paths.Num() > 0;
 	}
 
-	/** Get the binding name, resolves reference deprecation / redirectors / etc before returning */
-	TArray<FName> GetPaths() const
-	{
-		TArray<FName> Result;
-		Result.Reserve(Paths.Num());
-
-		for (const FMVVMBlueprintFieldPath& Path : Paths)
-		{
-			Result.Add(Path.GetFieldName());
-		}
-
-		return Result;
-	}
-
 	TArrayView<FMVVMBlueprintFieldPath const> GetFieldPaths() const
 	{
 		return Paths;
 	}
 
-	/** Get the binding name, resolves reference deprecation / redirectors / etc before returning */
-	TArray<UE::MVVM::FMVVMConstFieldVariant> GetFields() const
-	{
-		TArray<UE::MVVM::FMVVMConstFieldVariant> Result;
-		Result.Reserve(Paths.Num());
-
-		for (const FMVVMBlueprintFieldPath& Path : Paths)
-		{
-			Result.Add(Path.GetField());
-		}
+	/** Get the binding names, resolves reference deprecation / redirectors / etc before returning */
+	MODELVIEWVIEWMODELBLUEPRINT_API TArray<FName> GetFieldNames(const UClass* SelfContext) const;
 
-		return Result;
-	}
+	/** Get the binding fields, resolves reference deprecation / redirectors / etc before returning */
+	MODELVIEWVIEWMODELBLUEPRINT_API TArray<UE::MVVM::FMVVMConstFieldVariant> GetFields(const UClass* SelfContext) const;
 
-	bool HasFieldInLocalScope() const
-	{
-		for (const FMVVMBlueprintFieldPath& Path : Paths)
-		{
-			if (Path.IsFieldLocalScope())
-			{
-				return true;
-			}
-		}
+	/**
+	 * Get the full path without the first property name.
+	 * returns Field.SubProperty.SubProperty from ViewModel.Field.SubProeprty.SubProperty
+	 */
+	MODELVIEWVIEWMODELBLUEPRINT_API FString GetPropertyPath(const UClass* SelfContext) const;
 
-		return false;
-	}
+	MODELVIEWVIEWMODELBLUEPRINT_API bool HasFieldInLocalScope() const;
 
-	bool BasePropertyPathContains(UE::MVVM::FMVVMConstFieldVariant Field) const
+	bool PropertyPathContains(const UClass* InSelfContext, UE::MVVM::FMVVMConstFieldVariant Field) const
 	{
-		return Paths.ContainsByPredicate([Field](const FMVVMBlueprintFieldPath& FieldPath) { return FieldPath.GetField() == Field;  });
+		return Paths.ContainsByPredicate([InSelfContext, Field](const FMVVMBlueprintFieldPath& FieldPath) { return FieldPath.GetField(InSelfContext) == Field; });
 	}
 
-	void SetBasePropertyPath(UE::MVVM::FMVVMConstFieldVariant InField)
+	void SetPropertyPath(const UBlueprint* InContext, UE::MVVM::FMVVMConstFieldVariant InField)
 	{
-		ResetBasePropertyPath();
-		AppendBasePropertyPath(InField);
+		ResetPropertyPath();
+		AppendPropertyPath(InContext, InField);
 	}
 
-	void AppendBasePropertyPath(UE::MVVM::FMVVMConstFieldVariant InField)
+	void AppendPropertyPath(const UBlueprint* InContext, UE::MVVM::FMVVMConstFieldVariant InField)
 	{
 		if (!InField.IsEmpty() && !InField.GetName().IsNone())
 		{
-			Paths.AddDefaulted_GetRef().SetBindingReference(InField);
+			Paths.Emplace(InContext, InField);
 		}
 	}
 
-	void ResetBasePropertyPath()
+	void ResetPropertyPath()
 	{
 		Paths.Reset();
 	}
@@ -346,24 +193,6 @@ public:
 		return !IsFromWidget() && !IsFromViewModel() && BindingReference_DEPRECATED.GetMemberName() == FName();
 	}
 
-	/**
-	 * Get the full path without the first property name.
-	 * returns Field.SubProperty.SubProperty from ViewModel.Field.SubProeprty.SubProperty
-	 */
-	FString GetBasePropertyPath() const
-	{
-		TStringBuilder<512> Result;
-		for (const FMVVMBlueprintFieldPath& Path : Paths)
-		{
-			if (Result.Len() > 0)
-			{
-				Result << TEXT('.');
-			}
-			Result << Path.GetFieldName();
-		}
-		return Result.ToString();
-	}
-
 	bool operator==(const FMVVMBlueprintPropertyPath& Other) const
 	{
 		return WidgetName == Other.WidgetName && 
@@ -389,18 +218,6 @@ public:
 	}
 };
 
-inline uint32 GetTypeHash(const FMVVMBlueprintPropertyPath& Path)
-{
-	uint32 Hash = 0;
-	TArray<FName> Paths = Path.GetPaths();
-
-	for (const FName& SubPath : Paths)
-	{
-		Hash = HashCombine(Hash, GetTypeHash(SubPath));
-	}
-	return Hash;
-}
-
 template<>
 struct TStructOpsTypeTraits<FMVVMBlueprintPropertyPath> : public TStructOpsTypeTraitsBase2<FMVVMBlueprintPropertyPath>
 {
diff --git a/Engine/Plugins/Runtime/ModelViewViewModel/Source/ModelViewViewModelEditor/Private/Customizations/MVVMPropertyBindingExtension.cpp b/Engine/Plugins/Runtime/ModelViewViewModel/Source/ModelViewViewModelEditor/Private/Customizations/MVVMPropertyBindingExtension.cpp
--- a/Engine/Plugins/Runtime/ModelViewViewModel/Source/ModelViewViewModelEditor/Private/Customizations/MVVMPropertyBindingExtension.cpp
+++ b/Engine/Plugins/Runtime/ModelViewViewModel/Source/ModelViewViewModelEditor/Private/Customizations/MVVMPropertyBindingExtension.cpp
@@ -81,12 +81,12 @@ static void ExtendBindingsMenu(FMenuBuilder& MenuBuilder, const UWidgetBlueprint
 		return HorizontalBox;
 	};
 
-	auto CreateBinding = [MVVMBlueprintView](UWidget* Widget, TSharedPtr<IPropertyHandle> WidgetPropertyHandle, FGuid ViewModelId, const FProperty* ViewModelProperty)
+	auto CreateBinding = [WidgetBlueprint, MVVMBlueprintView](UWidget* Widget, TSharedPtr<IPropertyHandle> WidgetPropertyHandle, FGuid ViewModelId, const FProperty* ViewModelProperty)
 	{
 		FMVVMBlueprintViewBinding& NewBinding = MVVMBlueprintView->AddDefaultBinding();
 
 		NewBinding.SourcePath.SetViewModelId(ViewModelId);
-		NewBinding.SourcePath.SetBasePropertyPath(UE::MVVM::FMVVMConstFieldVariant(ViewModelProperty));
+		NewBinding.SourcePath.SetPropertyPath(WidgetBlueprint, UE::MVVM::FMVVMConstFieldVariant(ViewModelProperty));
 
 		// Generate the destination path from the widget property that we are dropping on.
 		FCachedPropertyPath CachedPropertyPath(WidgetPropertyHandle->GeneratePathToProperty());
@@ -94,12 +94,12 @@ static void ExtendBindingsMenu(FMenuBuilder& MenuBuilder, const UWidgetBlueprint
 
 		// Set the destination path.
 		FMVVMBlueprintPropertyPath DestinationPropertyPath;
-		DestinationPropertyPath.ResetBasePropertyPath();
+		DestinationPropertyPath.ResetPropertyPath();
 
 		for (int32 SegNum = 0; SegNum < CachedPropertyPath.GetNumSegments(); SegNum++)
 		{
 			FFieldVariant Field = CachedPropertyPath.GetSegment(SegNum).GetField();
-			DestinationPropertyPath.AppendBasePropertyPath(UE::MVVM::FMVVMConstFieldVariant(Field));
+			DestinationPropertyPath.AppendPropertyPath(WidgetBlueprint, UE::MVVM::FMVVMConstFieldVariant(Field));
 		}
 
 		DestinationPropertyPath.SetWidgetName(Widget->GetFName());
@@ -180,7 +180,7 @@ TOptional<FName> FMVVMPropertyBindingExtension::GetCurrentValue(const UWidgetBlu
 		return TOptional<FName>();
 	}
 
-	TArray<FName> Names = Binding->SourcePath.GetPaths();
+	TArray<FName> Names = Binding->SourcePath.GetFieldNames(WidgetBlueprint->SkeletonGeneratedClass);
 	if (Names.Num() > 0)
 	{
 		return Names.Last();
@@ -203,7 +203,7 @@ const FSlateBrush* FMVVMPropertyBindingExtension::GetCurrentIcon(const UWidgetBl
 		return nullptr;
 	}
 
-	TArray<UE::MVVM::FMVVMConstFieldVariant> Fields = Binding->SourcePath.GetFields();
+	TArray<UE::MVVM::FMVVMConstFieldVariant> Fields = Binding->SourcePath.GetFields(WidgetBlueprint->SkeletonGeneratedClass);
 	if (Fields.IsEmpty())
 	{
 		return nullptr;
@@ -232,7 +232,7 @@ void FMVVMPropertyBindingExtension::ClearCurrentValue(const UWidgetBlueprint* Wi
 		{
 			if (FMVVMBlueprintViewBinding* Binding = MVVMBlueprintView->FindBinding(Widget, Property))
 			{
-				Binding->SourcePath.ResetBasePropertyPath();
+				Binding->SourcePath.ResetPropertyPath();
 			}
 		}
 	}
@@ -283,10 +283,10 @@ IPropertyBindingExtension::EDropResult FMVVMPropertyBindingExtension::OnDrop(con
 
 		// Set the source path (view model property from the drop event).
 		FMVVMBlueprintPropertyPath SourcePropertyPath;
-		SourcePropertyPath.ResetBasePropertyPath();
+		SourcePropertyPath.ResetPropertyPath();
 		for (const FFieldVariant& Field : SourceFieldPath)
 		{
-			SourcePropertyPath.AppendBasePropertyPath(UE::MVVM::FMVVMConstFieldVariant(Field));
+			SourcePropertyPath.AppendPropertyPath(WidgetBlueprint, UE::MVVM::FMVVMConstFieldVariant(Field));
 		}
 		if (ViewModelFieldDragDropOp->ViewModelId.IsValid())
 		{
@@ -302,12 +302,12 @@ IPropertyBindingExtension::EDropResult FMVVMPropertyBindingExtension::OnDrop(con
 
 		// Set the destination path.
 		FMVVMBlueprintPropertyPath DestinationPropertyPath;
-		DestinationPropertyPath.ResetBasePropertyPath();
+		DestinationPropertyPath.ResetPropertyPath();
 
 		for (int32 SegNum = 0; SegNum < CachedPropertyPath.GetNumSegments(); SegNum++)
 		{
 			FFieldVariant Field = CachedPropertyPath.GetSegment(SegNum).GetField();
-			DestinationPropertyPath.AppendBasePropertyPath(UE::MVVM::FMVVMConstFieldVariant(Field));
+			DestinationPropertyPath.AppendPropertyPath(WidgetBlueprint, UE::MVVM::FMVVMConstFieldVariant(Field));
 		}
 
 		DestinationPropertyPath.SetWidgetName(Widget->GetFName());
diff --git a/Engine/Plugins/Runtime/ModelViewViewModel/Source/ModelViewViewModelEditor/Private/MVVMEditorSubsystem.cpp b/Engine/Plugins/Runtime/ModelViewViewModel/Source/ModelViewViewModelEditor/Private/MVVMEditorSubsystem.cpp
--- a/Engine/Plugins/Runtime/ModelViewViewModel/Source/ModelViewViewModelEditor/Private/MVVMEditorSubsystem.cpp
+++ b/Engine/Plugins/Runtime/ModelViewViewModel/Source/ModelViewViewModelEditor/Private/MVVMEditorSubsystem.cpp
@@ -402,7 +402,7 @@ void UMVVMEditorSubsystem::SetSourceToDestinationConversionFunction(UWidgetBluep
 {
 	if (UMVVMBlueprintView* View = GetView(WidgetBlueprint))
 	{
-		if (ConversionFunction == nullptr || IsValidConversionFunction(ConversionFunction, Binding.SourcePath, Binding.DestinationPath))
+		if (ConversionFunction == nullptr || IsValidConversionFunction(WidgetBlueprint, ConversionFunction, Binding.SourcePath, Binding.DestinationPath))
 		{
 			FScopedTransaction Transaction(LOCTEXT("SetConversionFunction", "Set Conversion Function"));
 
@@ -436,7 +436,7 @@ void UMVVMEditorSubsystem::SetDestinationToSourceConversionFunction(UWidgetBluep
 {
 	if (UMVVMBlueprintView* View = GetView(WidgetBlueprint))
 	{
-		if (ConversionFunction == nullptr || IsValidConversionFunction(ConversionFunction, Binding.DestinationPath, Binding.SourcePath))
+		if (ConversionFunction == nullptr || IsValidConversionFunction(WidgetBlueprint, ConversionFunction, Binding.DestinationPath, Binding.SourcePath))
 		{
 			FScopedTransaction Transaction(LOCTEXT("SetConversionFunction", "Set Conversion Function"));
 
@@ -580,7 +580,7 @@ void UMVVMEditorSubsystem::SetCompileForBinding(UWidgetBlueprint* WidgetBlueprin
 	}
 }
 
-bool UMVVMEditorSubsystem::IsValidConversionFunction(const UFunction* Function, const FMVVMBlueprintPropertyPath& Source, const FMVVMBlueprintPropertyPath& Destination) const
+bool UMVVMEditorSubsystem::IsValidConversionFunction(const UWidgetBlueprint* WidgetBlueprint, const UFunction* Function, const FMVVMBlueprintPropertyPath& Source, const FMVVMBlueprintPropertyPath& Destination) const
 {
 	TValueOrError<const FProperty*, FText> ReturnResult = UE::MVVM::BindingHelper::TryGetReturnTypeForConversionFunction(Function);
 	if (ReturnResult.HasError())
@@ -602,7 +602,7 @@ bool UMVVMEditorSubsystem::IsValidConversionFunction(const UFunction* Function,
 
 	const FProperty* SourceProperty = nullptr;
 
-	TArray<UE::MVVM::FMVVMConstFieldVariant> SourceFields = Source.GetFields();
+	TArray<UE::MVVM::FMVVMConstFieldVariant> SourceFields = Source.GetFields(WidgetBlueprint->SkeletonGeneratedClass);
 	if (SourceFields.Num() > 0)
 	{
 		SourceProperty = SourceFields.Last().IsProperty() ? SourceFields.Last().GetProperty() : UE::MVVM::BindingHelper::GetReturnProperty(SourceFields.Last().GetFunction());
@@ -610,7 +610,7 @@ bool UMVVMEditorSubsystem::IsValidConversionFunction(const UFunction* Function,
 
 	const FProperty* DestinationProperty = nullptr;
 
-	TArray<UE::MVVM::FMVVMConstFieldVariant> DestFields = Destination.GetFields();
+	TArray<UE::MVVM::FMVVMConstFieldVariant> DestFields = Destination.GetFields(WidgetBlueprint->SkeletonGeneratedClass);
 	if (DestFields.Num() > 0)
 	{
 		DestinationProperty = DestFields.Last().IsProperty() ? DestFields.Last().GetProperty() : UE::MVVM::BindingHelper::GetFirstArgumentProperty(DestFields.Last().GetFunction());
@@ -687,7 +687,7 @@ TArray<UFunction*> UMVVMEditorSubsystem::GetAvailableConversionFunctions(const U
 		bool bFromBlueprintFunctionLibrary = FunctionClass->IsChildOf<UBlueprintFunctionLibrary>() && Function->HasAllFunctionFlags(FUNC_Static | FUNC_BlueprintPure);
 		if (bIsFromWidgetBlueprint || bIsFromSkeletonWidgetBlueprint || bFromBlueprintFunctionLibrary)
 		{
-			if (IsValidConversionFunction(Function, Source, Destination))
+			if (IsValidConversionFunction(WidgetBlueprint, Function, Source, Destination))
 			{
 				ConversionFunctions.Add(const_cast<UFunction*>(Function));
 			}
@@ -866,7 +866,7 @@ void UMVVMEditorSubsystem::SetPathForConversionFunctionArgument(UWidgetBlueprint
 		return;
 	}
 
-	TArray<UE::MVVM::FMVVMConstFieldVariant> Fields = Path.GetFields();
+	TArray<UE::MVVM::FMVVMConstFieldVariant> Fields = Path.GetFields(WidgetBlueprint->SkeletonGeneratedClass);
 
 	float PosX = ConversionNode->NodePosX - 300 * (Fields.Num() + 1);
 	const float PosY = ConversionNode->NodePosY + ArgumentIndex * 100;
diff --git a/Engine/Plugins/Runtime/ModelViewViewModel/Source/ModelViewViewModelEditor/Private/Widgets/SMVVMCachedViewBindingPropertyPath.cpp b/Engine/Plugins/Runtime/ModelViewViewModel/Source/ModelViewViewModelEditor/Private/Widgets/SMVVMCachedViewBindingPropertyPath.cpp
--- a/Engine/Plugins/Runtime/ModelViewViewModel/Source/ModelViewViewModelEditor/Private/Widgets/SMVVMCachedViewBindingPropertyPath.cpp
+++ b/Engine/Plugins/Runtime/ModelViewViewModel/Source/ModelViewViewModelEditor/Private/Widgets/SMVVMCachedViewBindingPropertyPath.cpp
@@ -22,7 +22,7 @@ void SCachedViewBindingPropertyPath::Construct(const FArguments& InArgs, const U
 	OnGetPropertyPath = InArgs._OnGetPropertyPath;
 
 	const FMVVMBlueprintPropertyPath PropertyPath = OnGetPropertyPath.Execute();
-	CachedPropertyPath = PropertyPath.GetFields();
+	CachedPropertyPath = PropertyPath.GetFields(InWidgetBlueprint->SkeletonGeneratedClass);
 
 	ChildSlot
 	[
@@ -39,7 +39,7 @@ void SCachedViewBindingPropertyPath::Tick(const FGeometry& AllottedGeometry, con
 	Super::Tick(AllottedGeometry, InCurrentTime, InDeltaTime);
 
 	const FMVVMBlueprintPropertyPath PropertyPath = OnGetPropertyPath.Execute();
-	TArray<FMVVMConstFieldVariant> NewCachedPropertyPath = PropertyPath.GetFields();
+	TArray<FMVVMConstFieldVariant> NewCachedPropertyPath = PropertyPath.GetFields(WidgetBlueprint.Get()->SkeletonGeneratedClass);
 	if (NewCachedPropertyPath != CachedPropertyPath)
 	{
 		CachedPropertyPath = MoveTemp(NewCachedPropertyPath);
diff --git a/Engine/Plugins/Runtime/ModelViewViewModel/Source/ModelViewViewModelEditor/Private/Widgets/SMVVMPropertyPath.cpp b/Engine/Plugins/Runtime/ModelViewViewModel/Source/ModelViewViewModelEditor/Private/Widgets/SMVVMPropertyPath.cpp
--- a/Engine/Plugins/Runtime/ModelViewViewModel/Source/ModelViewViewModelEditor/Private/Widgets/SMVVMPropertyPath.cpp
+++ b/Engine/Plugins/Runtime/ModelViewViewModel/Source/ModelViewViewModelEditor/Private/Widgets/SMVVMPropertyPath.cpp
@@ -78,7 +78,7 @@ void SPropertyPath::SetPropertyPath(const FMVVMBlueprintPropertyPath& InProperty
 			[
 				SNew(SFieldPaths)
 				.TextStyle(TextStyle)
-				.FieldPaths(InPropertyPath.GetFields())
+				.FieldPaths(InPropertyPath.GetFields(WidgetBlueprint.Get()->SkeletonGeneratedClass))
 				.ShowOnlyLast(bShowOnlyLastPath)
 			];
 	}
diff --git a/Engine/Plugins/Runtime/ModelViewViewModel/Source/ModelViewViewModelEditor/Private/Widgets/SMVVMViewBindingListView.cpp b/Engine/Plugins/Runtime/ModelViewViewModel/Source/ModelViewViewModelEditor/Private/Widgets/SMVVMViewBindingListView.cpp
--- a/Engine/Plugins/Runtime/ModelViewViewModel/Source/ModelViewViewModelEditor/Private/Widgets/SMVVMViewBindingListView.cpp
+++ b/Engine/Plugins/Runtime/ModelViewViewModel/Source/ModelViewViewModelEditor/Private/Widgets/SMVVMViewBindingListView.cpp
@@ -679,7 +679,7 @@ private:
 			Result.BindingMode = ViewBinding->BindingType;
 
 			{
-				TArray<FMVVMConstFieldVariant> Fields = bSource ? ViewBinding->DestinationPath.GetFields() : ViewBinding->SourcePath.GetFields();
+				TArray<FMVVMConstFieldVariant> Fields = bSource ? ViewBinding->DestinationPath.GetFields(WidgetBlueprintPtr->SkeletonGeneratedClass) : ViewBinding->SourcePath.GetFields(WidgetBlueprintPtr->SkeletonGeneratedClass);
 				if (Fields.Num() > 0)
 				{
 					FMVVMConstFieldVariant LastField = Fields.Last();
@@ -753,10 +753,10 @@ private:
 					TArray<FFieldVariant> FieldPath = bIsViewModelProperty ? ViewModelFieldDragDropOp->DraggedField : WidgetPropertyDragDropOp->DraggedPropertyPath;
 					FMVVMBlueprintPropertyPath PropertyPath;
 
-					PropertyPath.ResetBasePropertyPath();
+					PropertyPath.ResetPropertyPath();
 					for (const FFieldVariant& Field : FieldPath)
 					{
-						PropertyPath.AppendBasePropertyPath(FMVVMConstFieldVariant(Field));
+						PropertyPath.AppendPropertyPath(WidgetBlueprintPtr, FMVVMConstFieldVariant(Field));
 					}
 
 					UMVVMEditorSubsystem* Subsystem = GEditor->GetEditorSubsystem<UMVVMEditorSubsystem>();
diff --git a/Engine/Plugins/Runtime/ModelViewViewModel/Source/ModelViewViewModelEditor/Private/Widgets/SMVVMViewModelBindingListWidget.cpp b/Engine/Plugins/Runtime/ModelViewViewModel/Source/ModelViewViewModelEditor/Private/Widgets/SMVVMViewModelBindingListWidget.cpp
--- a/Engine/Plugins/Runtime/ModelViewViewModel/Source/ModelViewViewModelEditor/Private/Widgets/SMVVMViewModelBindingListWidget.cpp
+++ b/Engine/Plugins/Runtime/ModelViewViewModel/Source/ModelViewViewModelEditor/Private/Widgets/SMVVMViewModelBindingListWidget.cpp
@@ -561,11 +561,10 @@ FMVVMBlueprintPropertyPath SSourceBindingList::CreateBlueprintPropertyPath(SProp
 				PropertyPath.SetWidgetName(Source->Key.Name);
 			}
 
-			PropertyPath.ResetBasePropertyPath();
-
+			PropertyPath.ResetPropertyPath();
 			for (const FFieldVariant& Field : FieldPath)
 			{
-				PropertyPath.AppendBasePropertyPath(FMVVMConstFieldVariant(Field));
+				PropertyPath.AppendPropertyPath(WidgetBlueprintPtr, FMVVMConstFieldVariant(Field));
 			}
 		}
 	}
@@ -604,7 +603,7 @@ FMVVMBlueprintPropertyPath SSourceBindingList::CreateBlueprintPropertyPath(SProp
 				{
 					PropertyPath.SetWidgetName(Source->Key.Name);
 				}
-				PropertyPath.ResetBasePropertyPath();
+				PropertyPath.ResetPropertyPath();
 			}
 		}
 	}
@@ -654,6 +653,12 @@ void SSourceBindingList::SetSelectedProperty(const FMVVMBlueprintPropertyPath& P
 		return;
 	}
 
+	const UWidgetBlueprint* WidgetBlueprintPtr = WidgetBlueprint.Get();
+	if (!WidgetBlueprintPtr)
+	{
+		return;
+	}
+
 	SPropertyViewer::FHandle SelectedHandle;
 	for (TPair<FBindingSource, SPropertyViewer::FHandle>& Source : Sources)
 	{
@@ -669,7 +674,7 @@ void SSourceBindingList::SetSelectedProperty(const FMVVMBlueprintPropertyPath& P
 	TArray<FFieldVariant, TMemStackAllocator<>> FieldPath;
 	if (SelectedHandle.IsValid())
 	{
-		TArray<FMVVMConstFieldVariant> FieldVariants = Property.GetFields();
+		TArray<FMVVMConstFieldVariant> FieldVariants = Property.GetFields(WidgetBlueprintPtr->SkeletonGeneratedClass);
 		FieldPath.Reserve(FieldVariants.Num());
 
 		for (const FMVVMConstFieldVariant& Variant : FieldVariants)
diff --git a/Engine/Plugins/Runtime/ModelViewViewModel/Source/ModelViewViewModelEditor/Public/MVVMEditorSubsystem.h b/Engine/Plugins/Runtime/ModelViewViewModel/Source/ModelViewViewModelEditor/Public/MVVMEditorSubsystem.h
--- a/Engine/Plugins/Runtime/ModelViewViewModel/Source/ModelViewViewModelEditor/Public/MVVMEditorSubsystem.h
+++ b/Engine/Plugins/Runtime/ModelViewViewModel/Source/ModelViewViewModelEditor/Public/MVVMEditorSubsystem.h
@@ -67,7 +67,7 @@ public:
 	void SetCompileForBinding(UWidgetBlueprint* WidgetBlueprint, FMVVMBlueprintViewBinding& Binding, bool bCompile);
 
 	UFUNCTION(BlueprintCallable, Category = "Viewmodel")
-	bool IsValidConversionFunction(const UFunction* Function, const FMVVMBlueprintPropertyPath& Source, const FMVVMBlueprintPropertyPath& Destination) const;
+	bool IsValidConversionFunction(const UWidgetBlueprint* WidgeteBlueprint, const UFunction* Function, const FMVVMBlueprintPropertyPath& Source, const FMVVMBlueprintPropertyPath& Destination) const;
 
 	UFUNCTION(BlueprintCallable, Category = "Viewmodel")
 	bool IsSimpleConversionFunction(const UFunction* Function) const;
diff --git a/Engine/Source/Runtime/Core/Public/UObject/FortniteMainBranchObjectVersion.h b/Engine/Source/Runtime/Core/Public/UObject/FortniteMainBranchObjectVersion.h
--- a/Engine/Source/Runtime/Core/Public/UObject/FortniteMainBranchObjectVersion.h
+++ b/Engine/Source/Runtime/Core/Public/UObject/FortniteMainBranchObjectVersion.h
@@ -342,6 +342,9 @@ struct CORE_API FFortniteMainBranchObjectVersion
 		// Set flags on water static meshes to duplicate transient to avoid underlying static mesh duplication issue
 		WaterBodyStaticMeshDuplicateTransient,
 
+		// Update paths to use the SkeletalClass
+		MVVMConvertPropertyPathToSkeletalClass,
+
 		// -----<new versions can be added above this line>-------------------------------------------------
 		VersionPlusOne,
 		LatestVersion = VersionPlusOne - 1
