MVVM: Fix analysis issue. Early returns if we can't evaluate the path.
#rnx

diff --git a/Engine/Plugins/Runtime/ModelViewViewModel/Source/ModelViewViewModelBlueprint/Private/MVVMBlueprintViewBinding.cpp b/Engine/Plugins/Runtime/ModelViewViewModel/Source/ModelViewViewModelBlueprint/Private/MVVMBlueprintViewBinding.cpp
--- a/Engine/Plugins/Runtime/ModelViewViewModel/Source/ModelViewViewModelBlueprint/Private/MVVMBlueprintViewBinding.cpp
+++ b/Engine/Plugins/Runtime/ModelViewViewModel/Source/ModelViewViewModelBlueprint/Private/MVVMBlueprintViewBinding.cpp
@@ -124,7 +124,13 @@ namespace UE::MVVM::Private
 
 	void AppendWidgetPathString(const UWidgetBlueprint* WidgetBlueprint, const FMVVMBlueprintPropertyPath& WidgetPath, FStringBuilderBase& PathBuilder)
 	{
-		if (WidgetBlueprint == nullptr || WidgetPath.GetWidgetName().IsNone())
+		if (WidgetBlueprint == nullptr)
+		{
+			PathBuilder << TEXT("<none>");
+			return;
+		}
+
+		if (WidgetPath.GetWidgetName().IsNone())
 		{
 			PathBuilder << TEXT("<none>");
 		}
