MVVM: Fix viewmodel not displaying when added for the first time.
#rnx

diff --git a/Engine/Plugins/Runtime/ModelViewViewModel/Source/ModelViewViewModelEditor/Private/Widgets/SMVVMViewModelPanel.cpp b/Engine/Plugins/Runtime/ModelViewViewModel/Source/ModelViewViewModelEditor/Private/Widgets/SMVVMViewModelPanel.cpp
--- a/Engine/Plugins/Runtime/ModelViewViewModel/Source/ModelViewViewModelEditor/Private/Widgets/SMVVMViewModelPanel.cpp
+++ b/Engine/Plugins/Runtime/ModelViewViewModel/Source/ModelViewViewModelEditor/Private/Widgets/SMVVMViewModelPanel.cpp
@@ -7,6 +7,7 @@
 #include "IStructureDetailsView.h"
 #include "MVVMBlueprintView.h"
 #include "MVVMEditorSubsystem.h"
+#include "MVVMWidgetBlueprintExtension_View.h"
 #include "PropertyEditorModule.h"
 #include "ToolMenus.h"
 #include "View/MVVMViewModelContextResolver.h"
@@ -225,9 +226,14 @@ void SMVVMViewModelPanel::HandleViewUpdated(UBlueprintExtension*)
 			if (UWidgetBlueprint* WidgetBlueprint = WidgetBlueprintEditor->GetWidgetBlueprintObj())
 			{
 				UMVVMBlueprintView* CurrentBlueprintView = GEditor->GetEditorSubsystem<UMVVMEditorSubsystem>()->GetView(WidgetBlueprint);
+				if (CurrentBlueprintView == nullptr)
+				{
+					UMVVMWidgetBlueprintExtension_View::GetExtension<UMVVMWidgetBlueprintExtension_View>(WidgetBlueprint)->CreateBlueprintViewInstance();
+					CurrentBlueprintView = GEditor->GetEditorSubsystem<UMVVMEditorSubsystem>()->GetView(WidgetBlueprint);
+				}
 				WeakBlueprintView = CurrentBlueprintView;
 
-				if (CurrentBlueprintView)
+				if (ensure(CurrentBlueprintView))
 				{
 					ViewModelsUpdatedHandle = CurrentBlueprintView->OnViewModelsUpdated.AddSP(this, &SMVVMViewModelPanel::HandleViewModelsUpdated);
 					bViewUpdated = true;
