MVVM:
- Add Resolver function called when a VM set by a Resolver is removed from a View

diff --git a/Engine/Plugins/Runtime/ModelViewViewModel/Source/ModelViewViewModel/Private/View/MVVMView.cpp b/Engine/Plugins/Runtime/ModelViewViewModel/Source/ModelViewViewModel/Private/View/MVVMView.cpp
--- a/Engine/Plugins/Runtime/ModelViewViewModel/Source/ModelViewViewModel/Private/View/MVVMView.cpp
+++ b/Engine/Plugins/Runtime/ModelViewViewModel/Source/ModelViewViewModel/Private/View/MVVMView.cpp
@@ -138,8 +138,13 @@ void UMVVMView::DeinitializeBindings()
 	UUserWidget* UserWidget = GetUserWidget();
 	check(UserWidget);
 
-	for (FMVVMViewSource& Source : Sources)
+	const TArrayView<const FMVVMViewClass_SourceCreator> AllViewModelCreators = ClassExtension->GetViewModelCreators();
+	check(AllViewModelCreators.Num() == Sources.Num());
+	for (int32 Index = 0; Index < AllViewModelCreators.Num(); ++Index)
 	{
+		const FMVVMViewClass_SourceCreator& Item = AllViewModelCreators[Index];
+		FMVVMViewSource& Source = Sources[Index];
+
 		if (Source.RegisteredCount > 0 && Source.Source)
 		{
 			TScriptInterface<INotifyFieldValueChanged> SourceAsInterface = Source.Source;
@@ -150,6 +155,8 @@ void UMVVMView::DeinitializeBindings()
 		// For GC release any object used by the view
 		if (!Source.bSetManually)
 		{
+			Item.DestroyInstance(Source.Source, this);
+
 			Source.Source = nullptr;
 			if (Source.bAssignedToUserWidgetProperty)
 			{
@@ -329,6 +336,11 @@ bool UMVVMView::SetSourceInternal(FName ViewModelName, TScriptInterface<INotifyF
 			}
 		}
 
+		if (!ViewSource.bSetManually && ViewSource.Source)
+		{
+			SourceCreator.DestroyInstance(ViewSource.Source, this);
+		}
+
 		ViewSource.Source = NewValue.GetObject();
 		ViewSource.bSetManually = !bForDynamicSource;
 		if (SourceCreator.IsSourceAUserWidgetProperty())
diff --git a/Engine/Plugins/Runtime/ModelViewViewModel/Source/ModelViewViewModel/Private/View/MVVMViewClass.cpp b/Engine/Plugins/Runtime/ModelViewViewModel/Source/ModelViewViewModel/Private/View/MVVMViewClass.cpp
--- a/Engine/Plugins/Runtime/ModelViewViewModel/Source/ModelViewViewModel/Private/View/MVVMViewClass.cpp
+++ b/Engine/Plugins/Runtime/ModelViewViewModel/Source/ModelViewViewModel/Private/View/MVVMViewClass.cpp
@@ -122,6 +122,14 @@ UObject* FMVVMViewClass_SourceCreator::CreateInstance(const UMVVMViewClass* InVi
 	return Result;
 }
 
+void FMVVMViewClass_SourceCreator::DestroyInstance(const UObject* ViewModel, const UMVVMView* View) const
+{
+	if (Resolver)
+	{
+		Resolver->DestroyInstance(ViewModel, View);
+	}
+}
+
 #if UE_WITH_MVVM_DEBUGGING
 FMVVMViewClass_SourceCreator::FToStringArgs FMVVMViewClass_SourceCreator::FToStringArgs::Short()
 {
diff --git a/Engine/Plugins/Runtime/ModelViewViewModel/Source/ModelViewViewModel/Public/View/MVVMViewClass.h b/Engine/Plugins/Runtime/ModelViewViewModel/Source/ModelViewViewModel/Public/View/MVVMViewClass.h
--- a/Engine/Plugins/Runtime/ModelViewViewModel/Source/ModelViewViewModel/Public/View/MVVMViewClass.h
+++ b/Engine/Plugins/Runtime/ModelViewViewModel/Source/ModelViewViewModel/Public/View/MVVMViewClass.h
@@ -37,6 +37,8 @@ struct FMVVMViewClass_SourceCreator
 public:
 	UObject* CreateInstance(const UMVVMViewClass* ViewClass, UMVVMView* View, UUserWidget* UserWidget) const;
 
+	void DestroyInstance(const UObject* ViewModel, const UMVVMView* View) const;
+
 	UClass* GetSourceClass() const
 	{
 		return ExpectedSourceType.Get();
diff --git a/Engine/Plugins/Runtime/ModelViewViewModel/Source/ModelViewViewModel/Public/View/MVVMViewModelContextResolver.h b/Engine/Plugins/Runtime/ModelViewViewModel/Source/ModelViewViewModel/Public/View/MVVMViewModelContextResolver.h
--- a/Engine/Plugins/Runtime/ModelViewViewModel/Source/ModelViewViewModel/Public/View/MVVMViewModelContextResolver.h
+++ b/Engine/Plugins/Runtime/ModelViewViewModel/Source/ModelViewViewModel/Public/View/MVVMViewModelContextResolver.h
@@ -27,4 +27,6 @@ public:
 
 	UFUNCTION(BlueprintImplementableEvent, Category="Viewmodel", DisplayName="Create Instance")
 	TScriptInterface<INotifyFieldValueChanged> K2_CreateInstance(const UClass* ExpectedType, const UUserWidget* UserWidget) const;
+
+	virtual void DestroyInstance(const UObject* ViewModel, const UMVVMView* View) const {}
 };
