Remove inconsistent check causing issues in complex widget scenarios.

diff --git a/Engine/Plugins/Runtime/ModelViewViewModel/Source/ModelViewViewModel/Private/View/MVVMView.cpp b/Engine/Plugins/Runtime/ModelViewViewModel/Source/ModelViewViewModel/Private/View/MVVMView.cpp
--- a/Engine/Plugins/Runtime/ModelViewViewModel/Source/ModelViewViewModel/Private/View/MVVMView.cpp
+++ b/Engine/Plugins/Runtime/ModelViewViewModel/Source/ModelViewViewModel/Private/View/MVVMView.cpp
@@ -139,7 +139,6 @@ void UMVVMView::DeinitializeBindings()
 	check(UserWidget);
 
 	const TArrayView<const FMVVMViewClass_SourceCreator> AllViewModelCreators = ClassExtension->GetViewModelCreators();
-	check(AllViewModelCreators.Num() == Sources.Num());
 	for (int32 Index = 0; Index < AllViewModelCreators.Num(); ++Index)
 	{
 		const FMVVMViewClass_SourceCreator& Item = AllViewModelCreators[Index];
