MVVM: add searchbar for binding view panel

diff --git a/Engine/Plugins/Runtime/ModelViewViewModel/Source/ModelViewViewModelBlueprint/Private/MVVMBlueprintView.cpp b/Engine/Plugins/Runtime/ModelViewViewModel/Source/ModelViewViewModelBlueprint/Private/MVVMBlueprintView.cpp
--- a/Engine/Plugins/Runtime/ModelViewViewModel/Source/ModelViewViewModelBlueprint/Private/MVVMBlueprintView.cpp
+++ b/Engine/Plugins/Runtime/ModelViewViewModel/Source/ModelViewViewModelBlueprint/Private/MVVMBlueprintView.cpp
@@ -138,6 +138,7 @@ FMVVMBlueprintViewBinding& UMVVMBlueprintView::AddBinding(const UWidget* Widget,
 	NewBinding.DestinationPath.SetPropertyPath(GetOuterUMVVMWidgetBlueprintExtension_View()->GetWidgetBlueprint(), UE::MVVM::FMVVMConstFieldVariant(Property));
 	NewBinding.BindingId = FGuid::NewGuid();
 
+	OnBindingsAdded.Broadcast();
 	OnBindingsUpdated.Broadcast();
 	return NewBinding;
 }
@@ -147,6 +148,7 @@ FMVVMBlueprintViewBinding& UMVVMBlueprintView::AddDefaultBinding()
 	FMVVMBlueprintViewBinding& NewBinding = Bindings.AddDefaulted_GetRef();
 	NewBinding.BindingId = FGuid::NewGuid();
 
+	OnBindingsAdded.Broadcast();
 	OnBindingsUpdated.Broadcast();
 	return NewBinding;
 }
diff --git a/Engine/Plugins/Runtime/ModelViewViewModel/Source/ModelViewViewModelBlueprint/Private/MVVMBlueprintViewBinding.cpp b/Engine/Plugins/Runtime/ModelViewViewModel/Source/ModelViewViewModelBlueprint/Private/MVVMBlueprintViewBinding.cpp
--- a/Engine/Plugins/Runtime/ModelViewViewModel/Source/ModelViewViewModelBlueprint/Private/MVVMBlueprintViewBinding.cpp
+++ b/Engine/Plugins/Runtime/ModelViewViewModel/Source/ModelViewViewModelBlueprint/Private/MVVMBlueprintViewBinding.cpp
@@ -2,6 +2,7 @@
 
 #include "MVVMBlueprintViewBinding.h"
 
+#include "Blueprint/WidgetTree.h"
 #include "Bindings/MVVMConversionFunctionHelper.h"
 #include "K2Node_CallFunction.h"
 #include "MVVMBlueprintView.h"
@@ -9,6 +10,8 @@
 
 #include UE_INLINE_GENERATED_CPP_BY_NAME(MVVMBlueprintViewBinding)
 
+#define LOCTEXT_NAMESPACE "MVVMBlueprintViewBinding"
+
 FName FMVVMBlueprintViewBinding::GetFName() const
 {
 	return *BindingId.ToString(EGuidFormats::DigitsWithHyphensLower);
@@ -16,7 +19,36 @@ FName FMVVMBlueprintViewBinding::GetFName() const
 
 namespace UE::MVVM::Private
 {
-	void AppendViewModelPathString(const UWidgetBlueprint* WidgetBlueprint, const UMVVMBlueprintView* BlueprintView, const FMVVMBlueprintPropertyPath& ViewModelPath, FStringBuilderBase& PathBuilder)
+	FText GetDisplayNameForField(const FMVVMConstFieldVariant& Field)
+	{
+		if (!Field.IsEmpty())
+		{
+			if (Field.IsProperty())
+			{
+				return Field.GetProperty()->GetDisplayNameText();
+			}
+			if (Field.IsFunction())
+			{
+				return Field.GetFunction()->GetDisplayNameText();
+			}
+		}
+		return LOCTEXT("None", "<None>");
+	}
+
+	FString GetDisplayNameForWidget(const UWidgetBlueprint* WidgetBlueprint, const FName& WidgetName)
+	{
+		if (UWidgetTree* WidgetTree = WidgetBlueprint->WidgetTree)
+		{
+			if (UWidget* FoundWidget = WidgetTree->FindWidget(WidgetName))
+			{
+				return FoundWidget->GetDisplayLabel().IsEmpty() ? WidgetName.ToString() : FoundWidget->GetDisplayLabel();
+			}
+		}
+
+		return WidgetName.IsNone() ? TEXT("<none>") : WidgetName.ToString();
+	}
+
+	void AppendViewModelPathString(const UWidgetBlueprint* WidgetBlueprint, const UMVVMBlueprintView* BlueprintView, const FMVVMBlueprintPropertyPath& ViewModelPath, FStringBuilderBase& PathBuilder, FStringBuilderBase& FunctionKeywordsBuilder, bool bUseDisplayName = false, bool bAppendFunctionKeywords = false)
 	{
 		if (ViewModelPath.IsEmpty())
 		{
@@ -29,100 +61,48 @@ namespace UE::MVVM::Private
 		{
 			PathBuilder << ViewModel->GetDisplayName().ToString();
 		}
+		else if (ViewModelPath.IsFromWidget())
+		{
+			if (bUseDisplayName)
+			{
+				PathBuilder << GetDisplayNameForWidget(WidgetBlueprint, ViewModelPath.GetWidgetName());
+			}
+			else
+			{
+				PathBuilder << ViewModelPath.GetWidgetName();
+			}
+		}
 		else
 		{
 			PathBuilder << TEXT("<none>");
 		}
 
-		const FString PropertyPath = ViewModelPath.GetPropertyPath(WidgetBlueprint->SkeletonGeneratedClass);
-		if (!PropertyPath.IsEmpty())
+		TArray<UE::MVVM::FMVVMConstFieldVariant> Fields = ViewModelPath.GetFields(WidgetBlueprint->SkeletonGeneratedClass);
+		for (const UE::MVVM::FMVVMConstFieldVariant& Field : Fields)
 		{
 			PathBuilder << TEXT(".");
-			PathBuilder << PropertyPath;
-		}
-	}
-	
-	FString GetBindingViewModelName(const UWidgetBlueprint* WidgetBlueprint, const FMVVMBlueprintViewBinding& Binding)
-	{
-		UMVVMWidgetBlueprintExtension_View* ExtensionView = UMVVMWidgetBlueprintExtension_View::GetExtension<UMVVMWidgetBlueprintExtension_View>(WidgetBlueprint);
-		UMVVMBlueprintView* BlueprintView = ExtensionView->GetBlueprintView();
-
-		TStringBuilder<256> NameBuilder;
-
-		if (UEdGraph* Graph = ConversionFunctionHelper::GetGraph(WidgetBlueprint, Binding, true))
-		{
-			if (const UK2Node_CallFunction* CallFunctionNode = ConversionFunctionHelper::GetFunctionNode(Graph))
+			if (bUseDisplayName)
 			{
-				NameBuilder << CallFunctionNode->GetFunctionName();
-				NameBuilder << TEXT("(");
-
-				bool bFirst = true;
-
-				for (const UEdGraphPin* Pin : CallFunctionNode->Pins)
-				{
-					if (Pin->PinName == UEdGraphSchema_K2::PN_Self || 
-						Pin->Direction != EGPD_Input)
-					{
-						continue;
-					}
-
-					if (!bFirst)
-					{
-						NameBuilder << TEXT(", ");
-					}
-
-					FMVVMBlueprintPropertyPath ArgumentPath = ConversionFunctionHelper::GetPropertyPathForArgument(WidgetBlueprint, CallFunctionNode, Pin->GetFName(), true);
-					if (!ArgumentPath.IsEmpty())
-					{
-						AppendViewModelPathString(WidgetBlueprint, BlueprintView, ArgumentPath, NameBuilder);
-					}
-					else
-					{
-						NameBuilder << Pin->GetDefaultAsString();
-					}
-
-					bFirst = false;
-				}
-
-				NameBuilder << TEXT(")");
+				PathBuilder << GetDisplayNameForField(Field).ToString();
 			}
 			else
 			{
-				NameBuilder << TEXT("<error>");
+				PathBuilder << Field.GetName().ToString();
 			}
-		}
-		else if (!Binding.Conversion.SourceToDestinationFunction.GetMemberName().IsNone())
-		{
-			if (const UFunction* SourceToDestFunction = Binding.Conversion.SourceToDestinationFunction.ResolveMember<UFunction>(WidgetBlueprint->SkeletonGeneratedClass))
-			{
-				NameBuilder << SourceToDestFunction->GetName();
-				NameBuilder << TEXT("(");
-			
-				AppendViewModelPathString(WidgetBlueprint, BlueprintView, Binding.SourcePath, NameBuilder);
 
-				NameBuilder << TEXT(")");
-			}
-			else
+			if (Field.IsFunction() && bAppendFunctionKeywords)
 			{
-				NameBuilder << Binding.Conversion.SourceToDestinationFunction.GetMemberName();
-				NameBuilder << TEXT("()");
+				FString FunctionKeywords = Field.GetFunction()->GetMetaData(FBlueprintMetadata::MD_FunctionKeywords);
+				if (!FunctionKeywords.IsEmpty())
+				{
+					FunctionKeywordsBuilder << TEXT(".");
+					FunctionKeywordsBuilder << FunctionKeywords;
+				}
 			}
 		}
-		else if (!Binding.SourcePath.IsEmpty())
-		{
-			AppendViewModelPathString(WidgetBlueprint, BlueprintView, Binding.SourcePath, NameBuilder);
-		}
-
-		FString Name = NameBuilder.ToString();
-		if (Name.IsEmpty())
-		{
-			Name = TEXT("<none>");
-		}
-
-		return Name;
 	}
-
-	void AppendWidgetPathString(const UWidgetBlueprint* WidgetBlueprint, const FMVVMBlueprintPropertyPath& WidgetPath, FStringBuilderBase& PathBuilder)
+	
+	void AppendWidgetPathString(const UWidgetBlueprint* WidgetBlueprint, const FMVVMBlueprintPropertyPath& WidgetPath, FStringBuilderBase& PathBuilder, FStringBuilderBase& FunctionKeywordsBuilder, bool bUseDisplayName = false, bool bAppendFunctionKeywords = false)
 	{
 		if (WidgetBlueprint == nullptr)
 		{
@@ -130,59 +110,131 @@ namespace UE::MVVM::Private
 			return;
 		}
 
-		if (WidgetPath.GetWidgetName().IsNone())
+		if (WidgetPath.GetWidgetName().IsNone() && !bUseDisplayName)
 		{
 			PathBuilder << TEXT("<none>");
 		}
 		else if (WidgetBlueprint->GetFName() == WidgetPath.GetWidgetName())
 		{
-			PathBuilder << TEXT("self");
+			if (bUseDisplayName)
+			{
+				PathBuilder << TEXT("self");
+			}
+			else
+			{
+				PathBuilder << WidgetBlueprint->GetName();
+			}
 		}
 		else
 		{
-			PathBuilder << WidgetPath.GetWidgetName();
+			if (bUseDisplayName)
+			{
+				PathBuilder << GetDisplayNameForWidget(WidgetBlueprint, WidgetPath.GetWidgetName());
+			}
+			else
+			{
+				PathBuilder << WidgetPath.GetWidgetName();
+			}
 		}
 
-		const FString PropertyPath = WidgetPath.GetPropertyPath(WidgetBlueprint->SkeletonGeneratedClass);
-		if (!PropertyPath.IsEmpty())
+		TArray<UE::MVVM::FMVVMConstFieldVariant> Fields = WidgetPath.GetFields(WidgetBlueprint->SkeletonGeneratedClass);
+		for (const UE::MVVM::FMVVMConstFieldVariant& Field : Fields)
 		{
 			PathBuilder << TEXT(".");
-			PathBuilder << PropertyPath;
+			if (bUseDisplayName)
+			{
+				PathBuilder << GetDisplayNameForField(Field).ToString();
+			}
+			else
+			{
+				PathBuilder << Field.GetName().ToString();
+			}
+
+			if (Field.IsFunction() && bAppendFunctionKeywords)
+			{
+				FString FunctionKeywords = Field.GetFunction()->GetMetaData(FBlueprintMetadata::MD_FunctionKeywords);
+				if (!FunctionKeywords.IsEmpty())
+				{
+					FunctionKeywordsBuilder << TEXT(".");
+					FunctionKeywordsBuilder << FunctionKeywords;
+				}
+			}
 		}
 	}
 
-	FString GetBindingWidgetName(const UWidgetBlueprint* WidgetBlueprint, const FMVVMBlueprintViewBinding& Binding)
+	FString GetBindingPathName(const UWidgetBlueprint* WidgetBlueprint, const FMVVMBlueprintViewBinding& Binding, bool bIsSource, FStringBuilderBase& FunctionKeywordsBuilder, bool bUseDisplayName = false, bool bAppendFunctionKeywords = false)
 	{
+		UMVVMWidgetBlueprintExtension_View* ExtensionView = UMVVMWidgetBlueprintExtension_View::GetExtension<UMVVMWidgetBlueprintExtension_View>(WidgetBlueprint);
+		UMVVMBlueprintView* BlueprintView = ExtensionView->GetBlueprintView();
+
 		TStringBuilder<256> NameBuilder;
 
-		if (UEdGraph* Graph = ConversionFunctionHelper::GetGraph(WidgetBlueprint, Binding, false))
+		if (UEdGraph* Graph = ConversionFunctionHelper::GetGraph(WidgetBlueprint, Binding, bIsSource))
 		{
 			if (const UK2Node_CallFunction* CallFunctionNode = ConversionFunctionHelper::GetFunctionNode(Graph))
 			{
-				NameBuilder << CallFunctionNode->GetFunctionName();
+				if (bUseDisplayName)
+				{
+					if (UFunction* TargetFunction = CallFunctionNode->GetTargetFunction())
+					{
+						NameBuilder << TargetFunction->GetDisplayNameText().ToString();
+					}
+					else
+					{
+						NameBuilder << CallFunctionNode->GetFunctionName();
+					}
+				}
+				else
+				{
+					NameBuilder << CallFunctionNode->GetFunctionName();
+				}
+
+				if (bAppendFunctionKeywords)
+				{
+					if (UFunction* TargetFunction = CallFunctionNode->GetTargetFunction())
+					{
+						FString FunctionKeywords = TargetFunction->GetMetaData(FBlueprintMetadata::MD_FunctionKeywords);
+						if (!FunctionKeywords.IsEmpty())
+						{
+							FunctionKeywordsBuilder << TEXT(".");
+							FunctionKeywordsBuilder << FunctionKeywords;
+						}
+					}
+				}
+
 				NameBuilder << TEXT("(");
 
 				bool bFirst = true;
 
 				for (const UEdGraphPin* Pin : CallFunctionNode->Pins)
 				{
-					if (Pin->Direction == EGPD_Input)
+					if (Pin->PinName == UEdGraphSchema_K2::PN_Self || 
+						Pin->Direction != EGPD_Input)
 					{
-						if (!bFirst)
-						{
-							NameBuilder << TEXT(", ");
-						}
+						continue;
+					}
+
+					if (!bFirst)
+					{
+						NameBuilder << TEXT(", ");
+					}
 
-						FMVVMBlueprintPropertyPath ArgumentPath = ConversionFunctionHelper::GetPropertyPathForArgument(WidgetBlueprint, CallFunctionNode, Pin->GetFName(), true);
-						if (!ArgumentPath.IsEmpty())
+					FMVVMBlueprintPropertyPath ArgumentPath = ConversionFunctionHelper::GetPropertyPathForArgument(WidgetBlueprint, CallFunctionNode, Pin->GetFName(), true);
+					if (!ArgumentPath.IsEmpty())
+					{
+						if (ArgumentPath.IsFromViewModel())
 						{
-							AppendWidgetPathString(WidgetBlueprint, ArgumentPath, NameBuilder);
+							AppendViewModelPathString(WidgetBlueprint, BlueprintView, ArgumentPath, NameBuilder, FunctionKeywordsBuilder, bUseDisplayName, bAppendFunctionKeywords);
 						}
 						else
 						{
-							NameBuilder << Pin->GetDefaultAsString();
+							AppendWidgetPathString(WidgetBlueprint, ArgumentPath, NameBuilder, FunctionKeywordsBuilder, bUseDisplayName, bAppendFunctionKeywords);
 						}
 					}
+					else
+					{
+						NameBuilder << Pin->GetDefaultAsString();
+					}
 
 					bFirst = false;
 				}
@@ -194,26 +246,62 @@ namespace UE::MVVM::Private
 				NameBuilder << TEXT("<error>");
 			}
 		}
-		else if (!Binding.Conversion.DestinationToSourceFunction.GetMemberName().IsNone())
+		else
 		{
-			if (const UFunction* DestToSourceFunction = Binding.Conversion.DestinationToSourceFunction.ResolveMember<UFunction>(WidgetBlueprint->SkeletonGeneratedClass))
+			const FMVVMBlueprintPropertyPath Path = bIsSource ? Binding.SourcePath : Binding.DestinationPath;
+			const FMemberReference FunctionReference = bIsSource ? Binding.Conversion.SourceToDestinationFunction : Binding.Conversion.DestinationToSourceFunction;
+			if (!FunctionReference.GetMemberName().IsNone())
 			{
-				NameBuilder << DestToSourceFunction->GetName();
-				NameBuilder << TEXT("(");
-			
-				AppendWidgetPathString(WidgetBlueprint, Binding.DestinationPath, NameBuilder);
+				if (const UFunction* Function = FunctionReference.ResolveMember<UFunction>(WidgetBlueprint->SkeletonGeneratedClass))
+				{
+					NameBuilder << Function->GetName();
 
-				NameBuilder << TEXT(")");
+					if (bAppendFunctionKeywords)
+					{
+						FString FunctionKeywords = Function->GetMetaData(FBlueprintMetadata::MD_FunctionKeywords);
+						if (!FunctionKeywords.IsEmpty())
+						{
+							FunctionKeywordsBuilder << TEXT(".");
+							FunctionKeywordsBuilder << FunctionKeywords;
+						}
+					}
+
+					NameBuilder << TEXT("(");
+
+					if (Path.IsFromViewModel())
+					{
+						AppendViewModelPathString(WidgetBlueprint, BlueprintView, Path, NameBuilder, FunctionKeywordsBuilder, bUseDisplayName, bAppendFunctionKeywords);
+
+					}
+					else
+					{
+						AppendWidgetPathString(WidgetBlueprint, Path, NameBuilder, FunctionKeywordsBuilder, bUseDisplayName, bAppendFunctionKeywords);
+					}
+
+					NameBuilder << TEXT(")");
+				}
+				else
+				{
+					NameBuilder << FunctionReference.GetMemberName();
+					NameBuilder << TEXT("()");
+				}
 			}
 			else
 			{
-				NameBuilder << Binding.Conversion.SourceToDestinationFunction.GetMemberName();
-				NameBuilder << TEXT("()");
-			}
-		}
-		else if (!Binding.DestinationPath.IsEmpty())
-		{
-			AppendWidgetPathString(WidgetBlueprint, Binding.DestinationPath, NameBuilder);
+				if (!Path.IsEmpty())
+				{
+					if (Path.IsFromViewModel())
+					{
+						AppendViewModelPathString(WidgetBlueprint, BlueprintView, Path, NameBuilder, FunctionKeywordsBuilder, bUseDisplayName, bAppendFunctionKeywords);
+
+					}
+					else
+					{
+						AppendWidgetPathString(WidgetBlueprint, Path, NameBuilder, FunctionKeywordsBuilder, bUseDisplayName, bAppendFunctionKeywords);
+					}
+				}
+			} 
+
 		}
 
 		FString Name = NameBuilder.ToString();
@@ -226,13 +314,14 @@ namespace UE::MVVM::Private
 	}
 }
 
-FString FMVVMBlueprintViewBinding::GetDisplayNameString(const UWidgetBlueprint* WidgetBlueprint) const
+FString FMVVMBlueprintViewBinding::GetDisplayNameString(const UWidgetBlueprint* WidgetBlueprint, bool bUseDisplayName) const
 {
 	check(WidgetBlueprint);
 
 	TStringBuilder<256> NameBuilder;
+	TStringBuilder<2> FunctionKeywordsBuilder; // This is only passed to GetBindingPathName but never used in this function.
 
-	NameBuilder << UE::MVVM::Private::GetBindingWidgetName(WidgetBlueprint, *this);
+	NameBuilder << UE::MVVM::Private::GetBindingPathName(WidgetBlueprint, *this, false, FunctionKeywordsBuilder, bUseDisplayName);
 
 	if (BindingType == EMVVMBindingMode::TwoWay)
 	{
@@ -251,7 +340,40 @@ FString FMVVMBlueprintViewBinding::GetDisplayNameString(const UWidgetBlueprint*
 		NameBuilder << TEXT(" ??? "); // shouldn't happen
 	}
 
-	NameBuilder << UE::MVVM::Private::GetBindingViewModelName(WidgetBlueprint, *this);
+	NameBuilder << UE::MVVM::Private::GetBindingPathName(WidgetBlueprint, *this, true, FunctionKeywordsBuilder, bUseDisplayName);
 
 	return NameBuilder.ToString();
 }
+
+FString FMVVMBlueprintViewBinding::GetSearchableString(const UWidgetBlueprint* WidgetBlueprint) const
+{
+	check(WidgetBlueprint);
+
+	FString SearchString;
+
+	// Get the binding string with variable names.
+	SearchString.Append(GetDisplayNameString(WidgetBlueprint, false));
+
+	// Remove the extra formatting that we don't need for search.
+	// We will include the formatted string as well in the second call to GetBindingPathName.
+	SearchString.ReplaceInline(TEXT(" "), TEXT(""));
+	SearchString.ReplaceInline(TEXT(")"), TEXT(""));
+	SearchString.ReplaceInline(TEXT("("), TEXT("."));
+	SearchString.ReplaceInline(TEXT(","), TEXT("."));
+
+	SearchString.Append(TEXT("."));
+
+	// Get the binding string with display names.
+	SearchString.Append(GetDisplayNameString(WidgetBlueprint, true));
+
+	// Create the function keywords string.
+	TStringBuilder<128> FunctionKeywordsBuilder;
+	UE::MVVM::Private::GetBindingPathName(WidgetBlueprint, *this, false, FunctionKeywordsBuilder, false, true);
+	UE::MVVM::Private::GetBindingPathName(WidgetBlueprint, *this, true, FunctionKeywordsBuilder, false, true);
+
+	SearchString.Append(FunctionKeywordsBuilder);
+
+	return SearchString;
+}
+
+#undef LOCTEXT_NAMESPACE
\ No newline at end of file
diff --git a/Engine/Plugins/Runtime/ModelViewViewModel/Source/ModelViewViewModelBlueprint/Public/MVVMBlueprintView.h b/Engine/Plugins/Runtime/ModelViewViewModel/Source/ModelViewViewModelBlueprint/Public/MVVMBlueprintView.h
--- a/Engine/Plugins/Runtime/ModelViewViewModel/Source/ModelViewViewModelBlueprint/Public/MVVMBlueprintView.h
+++ b/Engine/Plugins/Runtime/ModelViewViewModel/Source/ModelViewViewModelBlueprint/Public/MVVMBlueprintView.h
@@ -101,6 +101,9 @@ public:
 	DECLARE_EVENT(UMVVMBlueprintView, FOnBindingsUpdated);
 	FOnBindingsUpdated OnBindingsUpdated;
 
+	DECLARE_EVENT(UMVVMBlueprintView, FOnBindingsAdded);
+	FOnBindingsAdded OnBindingsAdded;
+
 	DECLARE_EVENT(UMVVMBlueprintView, FOnViewModelsUpdated);
 	FOnViewModelsUpdated OnViewModelsUpdated;
 
diff --git a/Engine/Plugins/Runtime/ModelViewViewModel/Source/ModelViewViewModelBlueprint/Public/MVVMBlueprintViewBinding.h b/Engine/Plugins/Runtime/ModelViewViewModel/Source/ModelViewViewModelBlueprint/Public/MVVMBlueprintViewBinding.h
--- a/Engine/Plugins/Runtime/ModelViewViewModel/Source/ModelViewViewModelBlueprint/Public/MVVMBlueprintViewBinding.h
+++ b/Engine/Plugins/Runtime/ModelViewViewModel/Source/ModelViewViewModelBlueprint/Public/MVVMBlueprintViewBinding.h
@@ -83,5 +83,12 @@ struct MODELVIEWVIEWMODELBLUEPRINT_API FMVVMBlueprintViewBinding
 	 * Get a string that identifies this binding. 
 	 * This is of the form: Widget.Property <- ViewModel.Property
 	 */
-	FString GetDisplayNameString(const UWidgetBlueprint* WidgetBlueprint) const;
+	FString GetDisplayNameString(const UWidgetBlueprint* WidgetBlueprint, bool bUseDisplayName = false) const;
+
+	/**
+	 * Get a string that identifies this binding and is specifically formatted for search. 
+	 * This includes the display name and variable name of all fields and widgets, as well as all function keywords.
+	 * For use in the UI, use GetDisplayNameString()
+	 */
+	FString GetSearchableString(const UWidgetBlueprint* WidgetBlueprint) const;
 };
diff --git a/Engine/Plugins/Runtime/ModelViewViewModel/Source/ModelViewViewModelEditor/Private/Widgets/SMVVMViewBindingListView.cpp b/Engine/Plugins/Runtime/ModelViewViewModel/Source/ModelViewViewModelEditor/Private/Widgets/SMVVMViewBindingListView.cpp
--- a/Engine/Plugins/Runtime/ModelViewViewModel/Source/ModelViewViewModelEditor/Private/Widgets/SMVVMViewBindingListView.cpp
+++ b/Engine/Plugins/Runtime/ModelViewViewModel/Source/ModelViewViewModelEditor/Private/Widgets/SMVVMViewBindingListView.cpp
@@ -107,6 +107,11 @@ struct FBindingEntry
 		Children.Add(Child);
 	}
 
+	void RemoveChildAt(int32 IndexToRemove)
+	{
+		Children.RemoveAt(IndexToRemove);
+	}
+
 	void ResetChildren()
 	{
 		Children.Reset();
@@ -119,6 +124,31 @@ struct FBindingEntry
 			BindingId == Other.BindingId;
 	}
 
+	FString GetSearchNameString(UMVVMBlueprintView* View, UWidgetBlueprint* WidgetBP)
+	{
+		FString RowToString;
+		FString FunctionKeywords;
+		
+		FMVVMBlueprintViewBinding* BindingInRow; // Initialized and used only when RowType is Binding.
+
+		switch (RowType)
+		{
+		case UE::MVVM::FBindingEntry::ERowType::Group:
+		case UE::MVVM::FBindingEntry::ERowType::Parameter:
+			RowToString = Name.ToString();
+			break;
+		case UE::MVVM::FBindingEntry::ERowType::Binding:
+			BindingInRow = GetBinding(View);
+			RowToString.Append(BindingInRow->GetSearchableString(WidgetBP));
+			break;
+		default:
+			break;
+		}
+
+		RowToString.ReplaceInline(TEXT(" "), TEXT(""));
+
+		return RowToString;
+	}
 private:
 	ERowType RowType = ERowType::None;
 	FName Name;
@@ -177,6 +207,72 @@ namespace Private
 		}
 		return TSharedPtr<FBindingEntry>();
 	}
+
+	void FilterBindingsList(FString FilterString, TArray<TSharedPtr<FBindingEntry>>& RootGroups, UMVVMBlueprintView* BlueprintView, UMVVMWidgetBlueprintExtension_View* MVVMExtensionPtr)
+	{
+		if (!FilterString.TrimStartAndEnd().IsEmpty())
+		{
+			TArray<FString> SearchKeywords;
+			FilterString.ParseIntoArray(SearchKeywords, TEXT(" "));
+			FString EntryString;
+
+			auto IsAllKeywordsInString = [](FString EntryString, TArray<FString>& SearchKeywords) -> bool {
+				for (const FString& Keyword : SearchKeywords)
+				{
+					if (!EntryString.Contains(Keyword))
+					{
+						return false;
+					}
+				}
+				return true;
+			};
+
+			for (int32 GroupEntryIndex = RootGroups.Num() - 1; GroupEntryIndex >= 0; GroupEntryIndex--)
+			{
+				TSharedPtr<FBindingEntry> GroupEntry = RootGroups[GroupEntryIndex];
+				EntryString = GroupEntry->GetSearchNameString(BlueprintView, MVVMExtensionPtr->GetWidgetBlueprint());
+
+				// If the filter text is found in the group name, we keep the entire group.
+				if (IsAllKeywordsInString(EntryString, SearchKeywords))
+				{
+					continue;
+				}
+				for (int32 BindingEntryIndex = GroupEntry->GetChildren().Num() - 1; BindingEntryIndex >= 0; BindingEntryIndex--)
+				{
+					TSharedPtr<FBindingEntry> BindingEntry = GroupEntry->GetChildren()[BindingEntryIndex];
+					EntryString = BindingEntry->GetSearchNameString(BlueprintView, MVVMExtensionPtr->GetWidgetBlueprint());
+					
+					// If the filter text is found in the binding string, we keep the entire binding.
+					if (IsAllKeywordsInString(EntryString, SearchKeywords))
+					{
+						continue;
+					}
+
+					bool IsParameterSearched = false;
+					for (TSharedPtr<FBindingEntry> ParameterEntry : BindingEntry->GetChildren())
+					{
+						EntryString = ParameterEntry->GetSearchNameString(BlueprintView, MVVMExtensionPtr->GetWidgetBlueprint());
+
+						// If the filter text is found in any of the parameter names, it is sufficient to keep the owner binding.
+						if (IsAllKeywordsInString(EntryString, SearchKeywords))
+						{
+							IsParameterSearched = true;
+							break;
+						}
+
+					}
+					if (!IsParameterSearched)
+					{
+						GroupEntry->RemoveChildAt(BindingEntryIndex);
+					}
+				}
+				if (GroupEntry->GetChildren().Num() == 0)
+				{
+					RootGroups.RemoveAt(GroupEntryIndex);
+				}
+			}
+		}
+	}
 }
 
 class SWidgetRow : public STableRow<TSharedPtr<FBindingEntry>>
@@ -1206,6 +1302,7 @@ void SBindingsList::Construct(const FArguments& InArgs, TSharedPtr<SBindingsPane
 
 	MVVMExtension->OnBlueprintViewChangedDelegate().AddSP(this, &SBindingsList::Refresh);
 	MVVMExtension->GetBlueprintView()->OnBindingsUpdated.AddSP(this, &SBindingsList::Refresh);
+	MVVMExtension->GetBlueprintView()->OnBindingsAdded.AddSP(this, &SBindingsList::ClearFilterText);
 	MVVMExtension->GetBlueprintView()->OnViewModelsUpdated.AddSP(this, &SBindingsList::Refresh);
 
 	ChildSlot
@@ -1376,6 +1473,7 @@ void SBindingsList::Refresh()
 				}
 			}
 		}
+		Private::FilterBindingsList(FilterText.ToString(), RootGroups, BlueprintView, MVVMExtensionPtr);
 	}
 
 	if (TreeView.IsValid())
@@ -1420,6 +1518,17 @@ TSharedRef<ITableRow> SBindingsList::GenerateEntryRow(TSharedPtr<FBindingEntry>
 	return SNew(STableRow<TSharedPtr<FBindingEntry>>, OwnerTable);
 }
 
+void SBindingsList::OnFilterTextChanged(const FText& InFilterText)
+{
+	FilterText = InFilterText;
+	Refresh();
+}
+
+void SBindingsList::ClearFilterText()
+{
+	FilterText = FText::GetEmpty();
+}
+
 namespace Private
 {
 	void GatherChildBindings(UMVVMBlueprintView* BlueprintView, const TConstArrayView<TSharedPtr<FBindingEntry>>& Entries, TArray<const FMVVMBlueprintViewBinding*>& OutBindings)
diff --git a/Engine/Plugins/Runtime/ModelViewViewModel/Source/ModelViewViewModelEditor/Private/Widgets/SMVVMViewBindingListView.h b/Engine/Plugins/Runtime/ModelViewViewModel/Source/ModelViewViewModelEditor/Private/Widgets/SMVVMViewBindingListView.h
--- a/Engine/Plugins/Runtime/ModelViewViewModel/Source/ModelViewViewModelEditor/Private/Widgets/SMVVMViewBindingListView.h
+++ b/Engine/Plugins/Runtime/ModelViewViewModel/Source/ModelViewViewModelEditor/Private/Widgets/SMVVMViewBindingListView.h
@@ -29,6 +29,9 @@ public:
 
 	void Refresh();
 
+	void OnFilterTextChanged(const FText& InFilterText);
+	void ClearFilterText();
+
 	/** Constructs context menu used for right click and dropdown button */
 	TSharedPtr<SWidget> OnSourceConstructContextMenu();
 
@@ -48,6 +51,7 @@ private:
 	TArray<TSharedPtr<FBindingEntry>> RootGroups;
 	TWeakObjectPtr<UMVVMWidgetBlueprintExtension_View> MVVMExtension;
 	mutable bool bSelectionChangedGuard = false;
+	FText FilterText;
 };
 
 } // namespace UE::MVVM
diff --git a/Engine/Plugins/Runtime/ModelViewViewModel/Source/ModelViewViewModelEditor/Private/Widgets/SMVVMViewBindingPanel.cpp b/Engine/Plugins/Runtime/ModelViewViewModel/Source/ModelViewViewModelEditor/Private/Widgets/SMVVMViewBindingPanel.cpp
--- a/Engine/Plugins/Runtime/ModelViewViewModel/Source/ModelViewViewModelEditor/Private/Widgets/SMVVMViewBindingPanel.cpp
+++ b/Engine/Plugins/Runtime/ModelViewViewModel/Source/ModelViewViewModelEditor/Private/Widgets/SMVVMViewBindingPanel.cpp
@@ -25,6 +25,8 @@
 #include "Toolkits/IToolkitHost.h"
 #include "Widgets/Docking/SDockTab.h"
 #include "Widgets/Input/SButton.h"
+#include "Widgets/Input/SCheckBox.h"
+#include "Widgets/Input/SSearchBox.h"
 #include "Widgets/Images/SImage.h"
 #include "Widgets/Layout/SSpacer.h"
 #include "Widgets/Layout/SSplitter.h"
@@ -446,7 +448,17 @@ TSharedRef<SWidget> SBindingsPanel::GenerateEditViewWidget()
 {
 	FSlateIcon EmptyIcon(FAppStyle::GetAppStyleSetName(), "Icon.Empty");
 
+	BindingsList = nullptr;
+	if (MVVMExtension.Get())
+	{
+		BindingsList = SNew(SBindingsList, StaticCastSharedRef<SBindingsPanel>(AsShared()), MVVMExtension.Get());
+	}
+
+	TSharedPtr<SHorizontalBox> BindingPanelToolBar = SNew(SHorizontalBox);
+
 	FSlimHorizontalToolBarBuilder ToolbarBuilderGlobal(TSharedPtr<const FUICommandList>(), FMultiBoxCustomization::None);
+
+	// Insert widgets in the toolbar to the left of the search bar
 	ToolbarBuilderGlobal.BeginSection("Picking");
 	{
 		ToolbarBuilderGlobal.AddToolBarButton(
@@ -479,10 +491,39 @@ TSharedRef<SWidget> SBindingsPanel::GenerateEditViewWidget()
 	}
 	ToolbarBuilderGlobal.EndSection();
 
+	// Pre-search box slot
+	BindingPanelToolBar->AddSlot()
+		.AutoWidth()
+		.HAlign(HAlign_Left)
+		.VAlign(VAlign_Center)
+		.Padding(0.0f, 0.0f, 4.0f, 0.0f)
+		[
+			ToolbarBuilderGlobal.MakeWidget()
+		];
+
+	TSharedRef<SWidget> SearchTextWidget = SNullWidget::NullWidget;
+	if (BindingsList)
 	{
-		ToolbarBuilderGlobal.AddWidget(SNew(SSpacer), NAME_None, true, HAlign_Right);
+		SearchTextWidget =
+			SNew(SSearchBox)
+			.HintText(LOCTEXT("SearchHint", "Search"))
+			.SelectAllTextWhenFocused(false)
+			.OnTextChanged(BindingsList.ToSharedRef(), &SBindingsList::OnFilterTextChanged);
+		
 	}
 
+	// Search box slot
+	BindingPanelToolBar->AddSlot()
+		.FillWidth(1.f)
+		.VAlign(VAlign_Center)
+		[
+			SearchTextWidget
+		];
+
+	// Reset the toolbar builder and insert widgets to the right of the search bar
+	ToolbarBuilderGlobal = FSlimHorizontalToolBarBuilder(TSharedPtr<const FUICommandList>(), FMultiBoxCustomization::None);
+	ToolbarBuilderGlobal.AddWidget(SNew(SSpacer), NAME_None, true, HAlign_Right);
+
 	{
 		ToolbarBuilderGlobal.AddWidget(CreateDrawerDockButton());
 
@@ -507,11 +548,15 @@ TSharedRef<SWidget> SBindingsPanel::GenerateEditViewWidget()
 		ToolbarBuilderGlobal.EndSection();
 	}
 
-	BindingsList = nullptr;
-	if (MVVMExtension.Get())
-	{
-		BindingsList = SNew(SBindingsList, StaticCastSharedRef<SBindingsPanel>(AsShared()), MVVMExtension.Get());
-	}
+	// Post-search box slot
+	BindingPanelToolBar->AddSlot()
+		.AutoWidth()
+		.HAlign(HAlign_Right)
+		.VAlign(VAlign_Center)
+		.Padding(0.0f, 0.0f, 4.0f, 0.0f)
+		[
+			ToolbarBuilderGlobal.MakeWidget()
+		];
 
 	TSharedRef<SWidget> BindingWidget = SNew(SOverlay)
 		+ SOverlay::Slot()
@@ -581,7 +626,7 @@ TSharedRef<SWidget> SBindingsPanel::GenerateEditViewWidget()
 		.AutoHeight()
 		.Padding(8.0f, 2.0f, 0.0f, 2.0f)
 		[
-			ToolbarBuilderGlobal.MakeWidget()
+			BindingPanelToolBar.ToSharedRef()
 		]
 
 		+SVerticalBox::Slot()
