MVVM: Add stats to MVVM View runtime.

diff --git a/Engine/Plugins/Runtime/ModelViewViewModel/Source/ModelViewViewModel/Private/Bindings/MVVMCompiledBindingLibrary.cpp b/Engine/Plugins/Runtime/ModelViewViewModel/Source/ModelViewViewModel/Private/Bindings/MVVMCompiledBindingLibrary.cpp
--- a/Engine/Plugins/Runtime/ModelViewViewModel/Source/ModelViewViewModel/Private/Bindings/MVVMCompiledBindingLibrary.cpp
+++ b/Engine/Plugins/Runtime/ModelViewViewModel/Source/ModelViewViewModel/Private/Bindings/MVVMCompiledBindingLibrary.cpp
@@ -5,6 +5,7 @@
 #include "Bindings/MVVMBindingHelper.h"
 #include "IFieldNotificationClassDescriptor.h"
 #include "INotifyFieldValueChanged.h"
+#include "ModelViewViewModelModule.h"
 #include "Templates/ValueOrError.h"
 #include "Types/MVVMFieldContext.h"
 #include "Types/MVVMFunctionContext.h"
@@ -13,6 +14,7 @@
 
 #include UE_INLINE_GENERATED_CPP_BY_NAME(MVVMCompiledBindingLibrary)
 
+DECLARE_CYCLE_STAT(TEXT("Load Library"), STAT_UMG_Viewmodel_LoadLibrary, STATGROUP_UMG_Viewmodel);
 
 #define LOCTEXT_NAMESPACE "MVVMCompiledBindingLibrary"
 
@@ -71,6 +73,8 @@ FText FMVVMCompiledBindingLibrary::LexToText(EExecutionFailingReason Reason)
 
 void FMVVMCompiledBindingLibrary::Load()
 {
+	SCOPE_CYCLE_COUNTER(STAT_UMG_Viewmodel_LoadLibrary);
+
 	ensureAlwaysMsgf(LoadedProperties.Num() == 0, TEXT("The binding library was loaded more than once."));
 	ensureAlwaysMsgf(LoadedFunctions.Num() == 0, TEXT("The binding library was loaded more than once."));
 	ensureAlwaysMsgf(LoadedFieldIds.Num() == 0, TEXT("The binding library was loaded more than once."));
diff --git a/Engine/Plugins/Runtime/ModelViewViewModel/Source/ModelViewViewModel/Private/ModelViewViewModelModule.h b/Engine/Plugins/Runtime/ModelViewViewModel/Source/ModelViewViewModel/Private/ModelViewViewModelModule.h
--- a/Engine/Plugins/Runtime/ModelViewViewModel/Source/ModelViewViewModel/Private/ModelViewViewModelModule.h
+++ b/Engine/Plugins/Runtime/ModelViewViewModel/Source/ModelViewViewModel/Private/ModelViewViewModelModule.h
@@ -3,9 +3,13 @@
 #pragma once
 
 #include "Modules/ModuleInterface.h"
+#include "Stats/Stats2.h"
 
 class IConsoleVariable;
 
+/** */
+DECLARE_STATS_GROUP(TEXT("UMG Viewmodel"), STATGROUP_UMG_Viewmodel, STATCAT_Advanced);
+
 /**
  *
  */
diff --git a/Engine/Plugins/Runtime/ModelViewViewModel/Source/ModelViewViewModel/Private/View/MVVMView.cpp b/Engine/Plugins/Runtime/ModelViewViewModel/Source/ModelViewViewModel/Private/View/MVVMView.cpp
--- a/Engine/Plugins/Runtime/ModelViewViewModel/Source/ModelViewViewModel/Private/View/MVVMView.cpp
+++ b/Engine/Plugins/Runtime/ModelViewViewModel/Source/ModelViewViewModel/Private/View/MVVMView.cpp
@@ -10,11 +10,18 @@
 #include "Debugging/MVVMDebugging.h"
 #include "Engine/Engine.h"
 #include "MVVMMessageLog.h"
+#include "ModelViewViewModelModule.h"
 #include "Templates/ValueOrError.h"
 #include "Types/MVVMFieldContext.h"
 
 #include UE_INLINE_GENERATED_CPP_BY_NAME(MVVMView)
 
+DECLARE_CYCLE_STAT(TEXT("InitializeBindings"), STAT_UMG_Viewmodel_InitializeBindings, STATGROUP_UMG_Viewmodel);
+DECLARE_CYCLE_STAT(TEXT("DeinitializeBindings"), STAT_UMG_Viewmodel_DeinitializeBindings, STATGROUP_UMG_Viewmodel);
+DECLARE_CYCLE_STAT(TEXT("SetSource"), STAT_UMG_Viewmodel_SetSource, STATGROUP_UMG_Viewmodel);
+DECLARE_CYCLE_STAT(TEXT("ExecuteBinding ValueChanged"), STAT_UMG_Viewmodel_ExecuteBinding_ValueChanged, STATGROUP_UMG_Viewmodel);
+DECLARE_CYCLE_STAT(TEXT("ExecuteBinding Delayed"), STAT_UMG_Viewmodel_ExecuteBinding_Delayed, STATGROUP_UMG_Viewmodel);
+DECLARE_CYCLE_STAT(TEXT("ExecuteBinding Tick"), STAT_UMG_Viewmodel_ExecuteBinding_Tick, STATGROUP_UMG_Viewmodel);
 
 #define LOCTEXT_NAMESPACE "MVVMView"
 
@@ -53,6 +60,8 @@ void UMVVMView::Construct()
 
 void UMVVMView::InitializeBindings()
 {
+	SCOPE_CYCLE_COUNTER(STAT_UMG_Viewmodel_InitializeBindings);
+
 	check(bInitialized == false);
 
 	// Init ViewModel instances
@@ -132,6 +141,8 @@ void UMVVMView::Destruct()
 
 void UMVVMView::DeinitializeBindings()
 {
+	SCOPE_CYCLE_COUNTER(STAT_UMG_Viewmodel_DeinitializeBindings);
+
 	check(bInitialized == true);
 	bInitialized = false;
 
@@ -257,6 +268,8 @@ bool UMVVMView::EvaluateSourceCreator(int32 SourceCreatorIndex)
 
 bool UMVVMView::SetSourceInternal(FName ViewModelName, TScriptInterface<INotifyFieldValueChanged> NewValue, bool bForDynamicSource)
 {
+	SCOPE_CYCLE_COUNTER(STAT_UMG_Viewmodel_SetSource);
+
 	if (ViewModelName.IsNone())
 	{
 		UE::MVVM::FMessageLog Log(GetUserWidget());
@@ -409,6 +422,8 @@ namespace UE::MVVM::Private
 
 void UMVVMView::HandledLibraryBindingValueChanged(UObject* InViewModelOrWidget, UE::FieldNotification::FFieldId InFieldId, int32 InCompiledBindingIndex) const
 {
+	SCOPE_CYCLE_COUNTER(STAT_UMG_Viewmodel_ExecuteBinding_ValueChanged);
+
 	check(InViewModelOrWidget);
 	check(InFieldId.IsValid());
 
@@ -456,6 +471,8 @@ void UMVVMView::HandledLibraryBindingValueChanged(UObject* InViewModelOrWidget,
 
 void UMVVMView::ExecuteDelayedBinding(const FMVVMViewDelayedBinding& DelayedBinding) const
 {
+	SCOPE_CYCLE_COUNTER(STAT_UMG_Viewmodel_ExecuteBinding_Delayed);
+
 	if (ensure(ClassExtension) && bInitialized)
 	{
 		if (ensure(ClassExtension->GetCompiledBindings().IsValidIndex(DelayedBinding.GetCompiledBindingIndex())))
@@ -494,6 +511,8 @@ void UMVVMView::ExecuteDelayedBinding(const FMVVMViewDelayedBinding& DelayedBind
 
 void UMVVMView::ExecuteEveryTickBindings() const
 {
+	SCOPE_CYCLE_COUNTER(STAT_UMG_Viewmodel_ExecuteBinding_Tick);
+
 	ensure(bHasEveryTickBinding);
 	ensure(bInitialized);
 
