MVVM: Change error message when the field does not exist to include a posible solution to the problem. Test if the field exist on the downcasted struct instead of the previous property owner struct.

diff --git a/Engine/Plugins/Runtime/ModelViewViewModel/Source/ModelViewViewModel/Private/Bindings/MVVMFieldPathHelper.cpp b/Engine/Plugins/Runtime/ModelViewViewModel/Source/ModelViewViewModel/Private/Bindings/MVVMFieldPathHelper.cpp
--- a/Engine/Plugins/Runtime/ModelViewViewModel/Source/ModelViewViewModel/Private/Bindings/MVVMFieldPathHelper.cpp
+++ b/Engine/Plugins/Runtime/ModelViewViewModel/Source/ModelViewViewModel/Private/Bindings/MVVMFieldPathHelper.cpp
@@ -254,7 +254,7 @@ TValueOrError<TArray<FMVVMConstFieldVariant>, FText> GenerateFieldPathList(TArra
 		
 		if (Field.IsProperty())
 		{
-			TValueOrError<FMVVMConstFieldVariant, FText> TransformedField = Private::TransformWithAccessor(CurrentContainer, Field, (bForSourceBinding || !bLastField));
+			TValueOrError<FMVVMConstFieldVariant, FText> TransformedField = Private::TransformWithAccessor(Field.GetOwner(), Field, (bForSourceBinding || !bLastField));
 			if (TransformedField.HasError())
 			{
 				return MakeError(TransformedField.StealError());
diff --git a/Engine/Plugins/Runtime/ModelViewViewModel/Source/ModelViewViewModelBlueprint/Private/Bindings/MVVMCompiledBindingLibraryCompiler.cpp b/Engine/Plugins/Runtime/ModelViewViewModel/Source/ModelViewViewModelBlueprint/Private/Bindings/MVVMCompiledBindingLibraryCompiler.cpp
--- a/Engine/Plugins/Runtime/ModelViewViewModel/Source/ModelViewViewModelBlueprint/Private/Bindings/MVVMCompiledBindingLibraryCompiler.cpp
+++ b/Engine/Plugins/Runtime/ModelViewViewModel/Source/ModelViewViewModelBlueprint/Private/Bindings/MVVMCompiledBindingLibraryCompiler.cpp
@@ -210,7 +210,7 @@ TValueOrError<FCompiledBindingLibraryCompiler::FFieldIdHandle, FText> FCompiledB
 
 	if (FieldId.IsNone())
 	{
-		return MakeError(LOCTEXT("FieldNotDefined", "The Field is not defined. Is it missing the FieldNotify flag?."));
+		return MakeError(LOCTEXT("FieldNotDefined", "The Field does not have the specifier FieldNotify and cannot be used as a binding source. You may want to use the 'One Time' binding mode."));
 	}
 
 	if (!SourceClass->ImplementsInterface(UNotifyFieldValueChanged::StaticClass()))
diff --git a/Engine/Plugins/Runtime/ModelViewViewModel/Source/ModelViewViewModelEditor/Private/Widgets/SMVVMCachedViewBindingPropertyPath.cpp b/Engine/Plugins/Runtime/ModelViewViewModel/Source/ModelViewViewModelEditor/Private/Widgets/SMVVMCachedViewBindingPropertyPath.cpp
--- a/Engine/Plugins/Runtime/ModelViewViewModel/Source/ModelViewViewModelEditor/Private/Widgets/SMVVMCachedViewBindingPropertyPath.cpp
+++ b/Engine/Plugins/Runtime/ModelViewViewModel/Source/ModelViewViewModelEditor/Private/Widgets/SMVVMCachedViewBindingPropertyPath.cpp
@@ -23,6 +23,8 @@ void SCachedViewBindingPropertyPath::Construct(const FArguments& InArgs, const U
 
 	const FMVVMBlueprintPropertyPath PropertyPath = OnGetPropertyPath.Execute();
 	CachedPropertyPath = PropertyPath.GetFields(InWidgetBlueprint->SkeletonGeneratedClass);
+	CachedWidgetName = PropertyPath.GetWidgetName();
+	CachedContextId = PropertyPath.GetViewModelId();
 
 	ChildSlot
 	[
@@ -38,12 +40,14 @@ void SCachedViewBindingPropertyPath::Tick(const FGeometry& AllottedGeometry, con
 {
 	Super::Tick(AllottedGeometry, InCurrentTime, InDeltaTime);
 
-	const FMVVMBlueprintPropertyPath PropertyPath = OnGetPropertyPath.Execute();
-	TArray<FMVVMConstFieldVariant> NewCachedPropertyPath = PropertyPath.GetFields(WidgetBlueprint.Get()->SkeletonGeneratedClass);
-	if (NewCachedPropertyPath != CachedPropertyPath)
+	const FMVVMBlueprintPropertyPath NewPropertyPath = OnGetPropertyPath.Execute();
+	TArray<FMVVMConstFieldVariant> NewCachedPropertyPath = NewPropertyPath.GetFields(WidgetBlueprint.Get()->SkeletonGeneratedClass);
+	if (NewCachedPropertyPath != CachedPropertyPath || NewPropertyPath.GetWidgetName() != CachedWidgetName || NewPropertyPath.GetViewModelId() != CachedContextId)
 	{
 		CachedPropertyPath = MoveTemp(NewCachedPropertyPath);
-		PropertyPathWidget->SetPropertyPath(PropertyPath);
+		CachedWidgetName = NewPropertyPath.GetWidgetName();
+		CachedContextId = NewPropertyPath.GetViewModelId();
+		PropertyPathWidget->SetPropertyPath(NewPropertyPath);
 	}
 }
 
diff --git a/Engine/Plugins/Runtime/ModelViewViewModel/Source/ModelViewViewModelEditor/Private/Widgets/SMVVMCachedViewBindingPropertyPath.h b/Engine/Plugins/Runtime/ModelViewViewModel/Source/ModelViewViewModelEditor/Private/Widgets/SMVVMCachedViewBindingPropertyPath.h
--- a/Engine/Plugins/Runtime/ModelViewViewModel/Source/ModelViewViewModelEditor/Private/Widgets/SMVVMCachedViewBindingPropertyPath.h
+++ b/Engine/Plugins/Runtime/ModelViewViewModel/Source/ModelViewViewModelEditor/Private/Widgets/SMVVMCachedViewBindingPropertyPath.h
@@ -40,6 +40,8 @@ private:
 	TSharedPtr<SPropertyPath> PropertyPathWidget;
 	FOnGetPropertyPath OnGetPropertyPath;
 	TArray<FMVVMConstFieldVariant> CachedPropertyPath;
+	FName CachedWidgetName;
+	FGuid CachedContextId;
 };
 
 } // namespace UE::MVVM
