MVVM: broadcast BindingUpdated when creating a binding via the Bind drop down in details panel.

diff --git a/Engine/Plugins/Runtime/ModelViewViewModel/Source/ModelViewViewModelEditor/Private/Customizations/MVVMPropertyBindingExtension.cpp b/Engine/Plugins/Runtime/ModelViewViewModel/Source/ModelViewViewModelEditor/Private/Customizations/MVVMPropertyBindingExtension.cpp
--- a/Engine/Plugins/Runtime/ModelViewViewModel/Source/ModelViewViewModelEditor/Private/Customizations/MVVMPropertyBindingExtension.cpp
+++ b/Engine/Plugins/Runtime/ModelViewViewModel/Source/ModelViewViewModelEditor/Private/Customizations/MVVMPropertyBindingExtension.cpp
@@ -106,6 +106,8 @@ static void ExtendBindingsMenu(FMenuBuilder& MenuBuilder, const UWidgetBlueprint
 		NewBinding.DestinationPath = DestinationPropertyPath;
 
 		NewBinding.BindingType = EMVVMBindingMode::OneWayToDestination;
+
+		MVVMBlueprintView->OnBindingsUpdated.Broadcast();
 	};
 
 	for (const FMVVMBlueprintViewModelContext& ViewModel : MVVMBlueprintView->GetViewModels())
