MVVM - Enable creation of viewmodel collections in modules other than MVVM

diff --git a/Engine/Plugins/Runtime/ModelViewViewModel/Source/ModelViewViewModel/Public/Types/MVVMViewModelCollection.h b/Engine/Plugins/Runtime/ModelViewViewModel/Source/ModelViewViewModel/Public/Types/MVVMViewModelCollection.h
--- a/Engine/Plugins/Runtime/ModelViewViewModel/Source/ModelViewViewModel/Public/Types/MVVMViewModelCollection.h
+++ b/Engine/Plugins/Runtime/ModelViewViewModel/Source/ModelViewViewModel/Public/Types/MVVMViewModelCollection.h
@@ -39,7 +39,7 @@ private:
 
 /** */
 UCLASS(meta = (DisplayName = "MVVM View Model Collection Object"))
-class UMVVMViewModelCollectionObject : public UObject
+class MODELVIEWVIEWMODEL_API UMVVMViewModelCollectionObject : public UObject
 {
 	GENERATED_BODY()
 
