MVVM: Add option to generate conversion graph when compiling and not saved in the packages.
Useful when graph cannot be saved.

diff --git a/Engine/Plugins/Runtime/ModelViewViewModel/Source/ModelViewViewModelBlueprint/Private/Bindings/MVVMCompiledBindingLibraryCompiler.cpp b/Engine/Plugins/Runtime/ModelViewViewModel/Source/ModelViewViewModelBlueprint/Private/Bindings/MVVMCompiledBindingLibraryCompiler.cpp
--- a/Engine/Plugins/Runtime/ModelViewViewModel/Source/ModelViewViewModelBlueprint/Private/Bindings/MVVMCompiledBindingLibraryCompiler.cpp
+++ b/Engine/Plugins/Runtime/ModelViewViewModel/Source/ModelViewViewModelBlueprint/Private/Bindings/MVVMCompiledBindingLibraryCompiler.cpp
@@ -442,6 +442,14 @@ TValueOrError<FCompiledBindingLibraryCompiler::FFieldPathHandle, FText> FCompile
 	const UClass* SourceClass = FBlueprintEditorUtils::GetMostUpToDateClass(InSourceClass);
 	const UFunction* Function = FBlueprintEditorUtils::GetMostUpToDateFunction(InFunction);
 
+	// Transient Conversion function are only added to generated class and not to the skeletal class.
+	bool bTransientConversionFunction = false;
+	if (Function == nullptr && InFunction && InFunction->GetTypedOuter<UClass>() == InSourceClass)
+	{
+		Function = InFunction;
+		bTransientConversionFunction = true;
+	}
+
 	if (SourceClass == nullptr)
 	{
 		return MakeError(LOCTEXT("SourceClassInvalid", "The source class is invalid."));
@@ -460,7 +468,7 @@ TValueOrError<FCompiledBindingLibraryCompiler::FFieldPathHandle, FText> FCompile
 
 	if (!Function->HasAllFunctionFlags(FUNC_Static))
 	{
-		if (!SourceClass->IsChildOf(Function->GetOuterUClass()))
+		if (!SourceClass->IsChildOf(Function->GetOuterUClass()) && !bTransientConversionFunction)
 		{
 			return MakeError(FText::Format(LOCTEXT("FunctionHasInvalidSelf", "Function {0} is going to be executed with an invalid self."), Function->GetDisplayNameText()));
 		}
diff --git a/Engine/Plugins/Runtime/ModelViewViewModel/Source/ModelViewViewModelBlueprint/Private/Bindings/MVVMConversionFunctionHelper.cpp b/Engine/Plugins/Runtime/ModelViewViewModel/Source/ModelViewViewModelBlueprint/Private/Bindings/MVVMConversionFunctionHelper.cpp
--- a/Engine/Plugins/Runtime/ModelViewViewModel/Source/ModelViewViewModelBlueprint/Private/Bindings/MVVMConversionFunctionHelper.cpp
+++ b/Engine/Plugins/Runtime/ModelViewViewModel/Source/ModelViewViewModelBlueprint/Private/Bindings/MVVMConversionFunctionHelper.cpp
@@ -2,26 +2,125 @@
 
 #include "Bindings/MVVMConversionFunctionHelper.h"
 
+#include "Bindings/MVVMBindingHelper.h"
 #include "Components/Widget.h"
 #include "Containers/Deque.h"
-#include "K2Node_CallFunction.h"
-#include "K2Node_VariableGet.h"
 #include "MVVMBlueprintView.h"
+#include "MVVMBlueprintViewConversionFunction.h"
 #include "MVVMWidgetBlueprintExtension_View.h"
 #include "UObject/MetaData.h"
 
+#include "K2Node_BreakStruct.h"
+#include "K2Node_CallFunction.h"
+#include "K2Node_FunctionEntry.h"
+#include "K2Node_FunctionResult.h"
+#include "K2Node_Self.h"
+#include "K2Node_VariableGet.h"
+#include "Kismet2/BlueprintEditorUtils.h"
+
 namespace UE::MVVM::ConversionFunctionHelper
 {
-static const FName ConversionFunctionMetadataKey = "ConversionFunction";
 
 namespace Private
 {
-	UEdGraph* FindExistingConversionFunctionWrapper(const UWidgetBlueprint* WidgetBlueprint, FName WrapperName)
+	static const FLazyName ConversionFunctionMetadataKey = "ConversionFunction";
+	static const FLazyName ConversionFunctionCategory = "Conversion Functions";
+
+	UMVVMBlueprintView* GetView(const UBlueprint* Blueprint)
 	{
-		const TObjectPtr<UEdGraph>* Result = WidgetBlueprint->FunctionGraphs.FindByPredicate([WrapperName](const UEdGraph* GraphPtr) { return GraphPtr->GetFName() == WrapperName; });
-		return Result ? Result->Get() : nullptr;
+		const TObjectPtr<UBlueprintExtension>* ExtensionViewPtr = Blueprint->GetExtensions().FindByPredicate([](const UBlueprintExtension* Other) { return Other && Other->GetClass() == UMVVMWidgetBlueprintExtension_View::StaticClass(); });
+		if (ExtensionViewPtr)
+		{
+			return CastChecked<UMVVMWidgetBlueprintExtension_View>(*ExtensionViewPtr)->GetBlueprintView();
+		}
+		return nullptr;
+	}
+
+	void MarkAsConversionFunction(const UK2Node* FunctionNode, const UEdGraph* Graph)
+	{
+		check(FunctionNode != nullptr);
+		FunctionNode->GetPackage()->GetMetaData()->SetValue(FunctionNode, ConversionFunctionMetadataKey.Resolve(), TEXT(""));
 	}
-} // namespace Private
+
+	UK2Node_FunctionEntry* FindFunctionEntry(UEdGraph* Graph)
+	{
+		for (UEdGraphNode* Node : Graph->Nodes)
+		{
+			if (UK2Node_FunctionEntry* FunctionEntry = Cast<UK2Node_FunctionEntry>(Node))
+			{
+				return FunctionEntry;
+			}
+		}
+		return nullptr;
+	}
+	
+	UK2Node_FunctionResult* FindFunctionResult(UEdGraph* Graph)
+	{
+		for (UEdGraphNode* Node : Graph->Nodes)
+		{
+			if (UK2Node_FunctionResult* FunctionResult = Cast<UK2Node_FunctionResult>(Node))
+			{
+				return FunctionResult;
+			}
+		}
+		return nullptr;
+	}
+
+	TTuple<UEdGraph*, UK2Node_FunctionEntry*, UK2Node_FunctionResult*> CreateGraph(UBlueprint* Blueprint, FName GraphName, bool bIsEditable, bool bAddToBlueprint)
+	{
+		FName UniqueFunctionName = FBlueprintEditorUtils::FindUniqueKismetName(Blueprint, GraphName.ToString());
+		UEdGraph* FunctionGraph = FBlueprintEditorUtils::CreateNewGraph(Blueprint, UniqueFunctionName, UEdGraph::StaticClass(), UEdGraphSchema_K2::StaticClass());
+		FunctionGraph->bEditable = bIsEditable;
+		if (bAddToBlueprint)
+		{
+			Blueprint->FunctionGraphs.Add(FunctionGraph);
+		}
+		else
+		{
+			FunctionGraph->SetFlags(RF_Transient);
+		}
+
+		const UEdGraphSchema_K2* Schema = CastChecked<UEdGraphSchema_K2>(FunctionGraph->GetSchema());
+		Schema->MarkFunctionEntryAsEditable(FunctionGraph, bIsEditable);
+		Schema->CreateDefaultNodesForGraph(*FunctionGraph);
+
+		// function entry node
+		FGraphNodeCreator<UK2Node_FunctionEntry> FunctionEntryCreator(*FunctionGraph);
+		UK2Node_FunctionEntry* FunctionEntry = FunctionEntryCreator.CreateNode();
+		FunctionEntry->FunctionReference.SetSelfMember(FunctionGraph->GetFName());
+		FunctionEntry->AddExtraFlags(FUNC_BlueprintCallable | FUNC_BlueprintPure | FUNC_Const | FUNC_Protected);
+		FunctionEntry->bIsEditable = bIsEditable;
+		FunctionEntry->MetaData.Category = FText::FromName(ConversionFunctionCategory.Resolve());
+		FunctionEntry->NodePosX = -500;
+		FunctionEntry->NodePosY = 0;
+		FunctionEntryCreator.Finalize();
+
+		FGraphNodeCreator<UK2Node_FunctionResult> FunctionResultCreator(*FunctionGraph);
+		UK2Node_FunctionResult* FunctionResult = FunctionResultCreator.CreateNode();
+		FunctionResult->FunctionReference.SetSelfMember(FunctionGraph->GetFName());
+		FunctionResult->bIsEditable = bIsEditable;
+		FunctionResult->NodePosX = 500;
+		FunctionResult->NodePosY = 0;
+		FunctionResultCreator.Finalize();
+
+		return { FunctionGraph, FunctionEntry, FunctionResult };
+	}
+}
+
+bool RequiresWrapper(const UFunction* ConversionFunction)
+{
+	if (ConversionFunction == nullptr)
+	{
+		return false;
+	}
+
+	TValueOrError<TArray<const FProperty*>, FText> ArgumentsResult = UE::MVVM::BindingHelper::TryGetArgumentsForConversionFunction(ConversionFunction);
+	if (ArgumentsResult.HasValue())
+	{
+		return (ArgumentsResult.GetValue().Num() > 1);
+	}
+	return false;
+}
 
 FName CreateWrapperName(const FMVVMBlueprintViewBinding& Binding, bool bSourceToDestination)
 {
@@ -33,59 +132,125 @@ FName CreateWrapperName(const FMVVMBlueprintViewBinding& Binding, bool bSourceTo
 	return FName(StringBuilder.ToString());
 }
 
-
-UEdGraph* GetGraph(const UWidgetBlueprint* WidgetBlueprint, const FMVVMBlueprintViewBinding& Binding, bool bSourceToDestination)
+TPair<UEdGraph*, UK2Node*> CreateGraph(UBlueprint* Blueprint, FName GraphName, const UFunction* ConversionFunction, bool bTransient)
 {
-	check(WidgetBlueprint);
-	const FName WrapperName = bSourceToDestination ? Binding.Conversion.SourceToDestinationWrapper : Binding.Conversion.DestinationToSourceWrapper;
-	return Private::FindExistingConversionFunctionWrapper(WidgetBlueprint, WrapperName);
-}
+	bool bIsEditable = false;
+	bool bAddToBlueprint = !bTransient;
 
-UK2Node_CallFunction* GetFunctionNode(UEdGraph* Graph)
-{
-	if (Graph == nullptr)
+	UEdGraph* FunctionGraph = nullptr;
+	UK2Node_FunctionEntry* FunctionEntry = nullptr;
+	UK2Node_FunctionResult* FunctionResult = nullptr;
 	{
-		return nullptr;
+		TTuple<UEdGraph*, UK2Node_FunctionEntry*, UK2Node_FunctionResult*> NewGraph = Private::CreateGraph(Blueprint, GraphName, bIsEditable, bAddToBlueprint);
+		FunctionGraph = NewGraph.Get<0>();
+		FunctionEntry = NewGraph.Get<1>();
+		FunctionResult = NewGraph.Get<2>();
 	}
 
-	TArray<UK2Node_CallFunction*> FunctionNodes;
-	Graph->GetNodesOfClass<UK2Node_CallFunction>(FunctionNodes);
+	const FProperty* ReturnProperty = UE::MVVM::BindingHelper::GetReturnProperty(ConversionFunction);
+	check(ReturnProperty);
 
-	if (FunctionNodes.Num() == 1)
+	// create return value pin
 	{
-		return FunctionNodes[0];
+		TSharedPtr<FUserPinInfo> PinInfo = MakeShared<FUserPinInfo>();
+		GetDefault<UEdGraphSchema_K2>()->ConvertPropertyToPinType(ReturnProperty, PinInfo->PinType);
+		PinInfo->PinName = ReturnProperty->GetFName();
+		PinInfo->DesiredPinDirection = EGPD_Input;
+		FunctionResult->UserDefinedPins.Add(PinInfo);
+		FunctionResult->ReconstructNode();
+	}
+	
+	UK2Node_CallFunction* CallFunctionNode = nullptr;
+	{
+		FGraphNodeCreator<UK2Node_CallFunction> CallFunctionCreator(*FunctionGraph);
+		CallFunctionNode = CallFunctionCreator.CreateNode();
+		CallFunctionNode->SetFromFunction(ConversionFunction);
+		CallFunctionNode->NodePosX = 0;
+		CallFunctionCreator.Finalize();
+		Private::MarkAsConversionFunction(CallFunctionNode, FunctionGraph);
 	}
 
-	for (UK2Node_CallFunction* FunctionNode : FunctionNodes)
+	// Make link Entry -> CallFunction || Entry -> Return
 	{
-		// check if we've set any metadata on the nodes to figure out which one it is
-		if (FunctionNode->GetPackage()->GetMetaData()->HasValue(FunctionNode, ConversionFunctionMetadataKey))
+		UEdGraphPin* FunctionEntryThenPin = FunctionEntry->FindPin(UEdGraphSchema_K2::PN_Then, EGPD_Output);
+		UEdGraphPin* FunctionResultExecPin = FunctionResult->GetExecPin();
+		
+		if (!CallFunctionNode->IsNodePure())
 		{
-			return FunctionNode;
+			UEdGraphPin* CallFunctionExecPin = CallFunctionNode->GetExecPin();
+			UEdGraphPin* CallFunctionThenPin = CallFunctionNode->FindPin(UEdGraphSchema_K2::PN_Then, EGPD_Output);
+
+			FunctionEntryThenPin->MakeLinkTo(CallFunctionExecPin);
+			CallFunctionThenPin->MakeLinkTo(FunctionResultExecPin);
+
+			CallFunctionNode->NodePosY = 0;
+		}
+		else
+		{
+			FunctionEntryThenPin->MakeLinkTo(FunctionResultExecPin);
+			CallFunctionNode->NodePosY = 100;
 		}
 	}
 
-	return nullptr;
+	{
+		UEdGraphPin* FunctionReturnPin = CallFunctionNode->FindPin(ReturnProperty->GetName(), EGPD_Output);
+		UEdGraphPin* FunctionResultPin = FunctionResult->FindPin(ReturnProperty->GetFName(), EGPD_Input);
+		check(FunctionResultPin && FunctionReturnPin);
+		FunctionReturnPin->MakeLinkTo(FunctionResultPin);
+	}
+
+	return { FunctionGraph , CallFunctionNode };
 }
 
-void MarkAsConversionFunction(const UK2Node_CallFunction* FunctionNode, const FMVVMBlueprintViewBinding& Binding)
+UK2Node* GetWrapperNode(UEdGraph* Graph)
 {
-	check(FunctionNode != nullptr);
+	if (Graph == nullptr)
+	{
+		return nullptr;
+	}
+
+	FName ConversionFunctionMetadataKey = Private::ConversionFunctionMetadataKey.Resolve();
+	for (UEdGraphNode* Node : Graph->Nodes)
+	{
+		// check if we've set any metadata on the nodes to figure out which one it is
+		if (Cast<UK2Node>(Node) && Node->GetPackage()->GetMetaData()->HasValue(Node, ConversionFunctionMetadataKey))
+		{
+			return CastChecked<UK2Node>(Node);
+		}
+	}
+
+	if (UK2Node_FunctionResult* FunctionResult = Private::FindFunctionResult(Graph))
+	{
+		for (UEdGraphPin* GraphPin : FunctionResult->Pins)
+		{
+			if (GraphPin->GetFName() != UEdGraphSchema_K2::PN_Execute && GraphPin->LinkedTo.Num() == 1)
+			{
+				if (UK2Node* Node = Cast<UK2Node>(GraphPin->LinkedTo[0]->GetOwningNode()))
+				{
+					Private::MarkAsConversionFunction(Node, Graph);
+					return Node;
+				}
+			}
+		}
+	}
 
-	FunctionNode->GetPackage()->GetMetaData()->SetValue(FunctionNode, ConversionFunctionMetadataKey, *Binding.BindingId.ToString());
+	return nullptr;
 }
 
 namespace Private
 {
-	FMVVMBlueprintPropertyPath GetPropertyPathForPin(const UWidgetBlueprint* WidgetBlueprint, const UEdGraphPin* StartPin, bool bSkipResolve)
+	FMVVMBlueprintPropertyPath GetPropertyPathForPin(const UBlueprint* Blueprint, const UEdGraphPin* StartPin, bool bSkipResolve)
 	{
 		if (StartPin->Direction != EGPD_Input || StartPin->PinName == UEdGraphSchema_K2::PN_Self)
 		{
 			return FMVVMBlueprintPropertyPath();
 		}
 		
-		UMVVMWidgetBlueprintExtension_View* ExtensionView = UMVVMWidgetBlueprintExtension_View::GetExtension<UMVVMWidgetBlueprintExtension_View>(WidgetBlueprint);
-		UMVVMBlueprintView* BlueprintView = ExtensionView->GetBlueprintView();
+		UMVVMBlueprintView* BlueprintView = GetView(Blueprint);
+		if (BlueprintView == nullptr)
+		{
+			return FMVVMBlueprintPropertyPath();
+		}
 
 		TArray<FMemberReference> PathParts;
 		TDeque<UEdGraphNode*> NodesToSearch;
@@ -145,9 +310,9 @@ namespace Private
 		}
 		else
 		{
-			if (const FObjectProperty* Property = CastField<FObjectProperty>(RootRef.ResolveMember<FProperty>(WidgetBlueprint->SkeletonGeneratedClass)))
+			if (const FObjectProperty* Property = CastField<FObjectProperty>(RootRef.ResolveMember<FProperty>(Blueprint->SkeletonGeneratedClass)))
 			{
-				if (Property->PropertyClass->IsChildOf<UWidget>() || Property->PropertyClass->IsChildOf<UWidgetBlueprint>())
+				if (Property->PropertyClass->IsChildOf<UWidget>() || Property->PropertyClass->IsChildOf<UBlueprint>())
 				{
 					ResultPath.SetWidgetName(Property->GetFName());
 				}
@@ -168,13 +333,13 @@ namespace Private
 
 		for (const FMemberReference& MemberReference : PathParts)
 		{
-			if (UFunction* Function = MemberReference.ResolveMember<UFunction>(WidgetBlueprint->SkeletonGeneratedClass)) 
+			if (UFunction* Function = MemberReference.ResolveMember<UFunction>(Blueprint->SkeletonGeneratedClass)) 
 			{
-				ResultPath.AppendPropertyPath(WidgetBlueprint, UE::MVVM::FMVVMConstFieldVariant(Function));
+				ResultPath.AppendPropertyPath(Blueprint, UE::MVVM::FMVVMConstFieldVariant(Function));
 			}
-			else if (const FProperty* Property = MemberReference.ResolveMember<FProperty>(WidgetBlueprint->SkeletonGeneratedClass))
+			else if (const FProperty* Property = MemberReference.ResolveMember<FProperty>(Blueprint->SkeletonGeneratedClass))
 			{
-				ResultPath.AppendPropertyPath(WidgetBlueprint, UE::MVVM::FMVVMConstFieldVariant(Property));
+				ResultPath.AppendPropertyPath(Blueprint, UE::MVVM::FMVVMConstFieldVariant(Property));
 			}
 		}
 
@@ -182,25 +347,281 @@ namespace Private
 	}
 }
 
-FMVVMBlueprintPropertyPath GetPropertyPathForArgument(const UWidgetBlueprint* WidgetBlueprint, const UK2Node_CallFunction* FunctionNode, FName ArgumentName, bool bSkipResolve)
+FMVVMBlueprintPropertyPath GetPropertyPathForPin(const UBlueprint* Blueprint, const UEdGraphPin* Pin, bool bSkipResolve)
 {
-	const UEdGraphPin* ArgumentPin = FunctionNode->FindPin(ArgumentName, EGPD_Input);
-	if (ArgumentPin == nullptr || ArgumentPin->LinkedTo.Num() == 0)
+	if (Pin == nullptr || Pin->LinkedTo.Num() == 0)
 	{
 		return FMVVMBlueprintPropertyPath();
 	}
 
-	return Private::GetPropertyPathForPin(WidgetBlueprint, ArgumentPin, false);
+	return Private::GetPropertyPathForPin(Blueprint, Pin, bSkipResolve);
 }
 
-TMap<FName, FMVVMBlueprintPropertyPath> GetAllArgumentPropertyPaths(const UWidgetBlueprint* WidgetBlueprint, const UK2Node_CallFunction* FunctionNode, bool bSkipResolve)
+void SetPropertyPathForPin(const UBlueprint* Blueprint, const FMVVMBlueprintPropertyPath& Path, UEdGraphPin* PathPin)
+{
+	if (PathPin == nullptr)
+	{
+		return;
+	}
+
+	UEdGraphNode* ConversionNode = PathPin->GetOwningNode();
+	UEdGraph* FunctionGraph = ConversionNode ? ConversionNode->GetGraph() : nullptr;
+	const UEdGraphSchema* Schema = FunctionGraph ? FunctionGraph->GetSchema() : nullptr;
+
+	// Remove previous nodes
+	{
+		TArray<UEdGraphPin*> PinsToClear;
+		TArray<UEdGraphNode*> NodesToRemove;
+		PinsToClear.Add(PathPin);
+		while (PinsToClear.Num() > 0)
+		{
+			UEdGraphPin* Pin = PinsToClear[0];
+			PinsToClear.RemoveAt(0);
+
+			PinsToClear.Append(Pin->LinkedTo);
+			Pin->BreakAllPinLinks();
+
+			// rename the node out of the way
+			UEdGraphNode* Node = Pin->GetOwningNode();
+			if (Node != ConversionNode)
+			{
+				NodesToRemove.AddUnique(Node);
+			}
+		}
+
+		if (FunctionGraph)
+		{
+			for (UEdGraphNode* Node : NodesToRemove)
+			{
+				// Break the links, and relink to the other nodes.
+				if (UK2Node_CallFunction* CallFunctionNode = Cast<UK2Node_CallFunction>(Node))
+				{
+					UEdGraphPin* ThenPin = CallFunctionNode->GetThenPin();
+					UEdGraphPin* ExecPin = CallFunctionNode->GetExecPin();
+					if (ThenPin && ThenPin->LinkedTo.Num() == 1)
+					{
+						Schema->MovePinLinks(*ExecPin, *ThenPin->LinkedTo[0], false);
+					}
+				}
+				FunctionGraph->RemoveNode(Node, true);
+			}
+		}
+	}
+
+	// Add new nodes
+	if (FunctionGraph && ConversionNode && !Path.IsEmpty())
+	{
+		const int32 ArgumentIndex = ConversionNode->Pins.IndexOfByPredicate([PathPin](const UEdGraphPin* Other){ return Other == PathPin; });
+		if (!ensure(ConversionNode->Pins.IsValidIndex(ArgumentIndex)))
+		{
+			return;
+		}
+
+		TArray<UE::MVVM::FMVVMConstFieldVariant> Fields = Path.GetFields(Blueprint->SkeletonGeneratedClass);
+		float PosX = ConversionNode->NodePosX - 300 * (Fields.Num() + 1);
+		const float PosY = ConversionNode->NodePosY + ArgumentIndex * 100;
+
+		UEdGraphPin* PreviousExecPin = nullptr;
+		{
+			UEdGraphNode* FunctionEntry = Private::FindFunctionEntry(FunctionGraph);
+			if (!ensure(FunctionEntry))
+			{
+				return;
+			}
+			PreviousExecPin = FunctionEntry->FindPinChecked(UEdGraphSchema_K2::PN_Then);
+		}
+		UEdGraphPin* PreviousDataPin = nullptr;
+		UClass* PreviousClass = nullptr;
+
+		// create the root property getter node, ie. the Widget/ViewModel
+		{
+			const FProperty* RootProperty = nullptr;
+			bool bCreateSelfNodeForRootProperty = false;
+			{
+				if (Path.IsFromWidget())
+				{
+					RootProperty = Blueprint->SkeletonGeneratedClass->FindPropertyByName(Path.GetWidgetName());
+					bCreateSelfNodeForRootProperty = Blueprint->GetFName() == Path.GetWidgetName();
+				}
+				else if (Path.IsFromViewModel())
+				{
+					UMVVMBlueprintView* View = Private::GetView(Blueprint);
+					const FMVVMBlueprintViewModelContext* Context = View->FindViewModel(Path.GetViewModelId());
+					RootProperty = Blueprint->SkeletonGeneratedClass->FindPropertyByName(Context->GetViewModelName());
+				}
+
+				if (RootProperty == nullptr)
+				{
+					ensureMsgf(bCreateSelfNodeForRootProperty, TEXT("Could not resolve root property!"));
+					return;
+				}
+			}
+
+			if (bCreateSelfNodeForRootProperty)
+			{
+				FGraphNodeCreator<UK2Node_Self> RootGetterCreator(*FunctionGraph);
+				UK2Node_Self* RootSelfNode = RootGetterCreator.CreateNode();
+				RootSelfNode->NodePosX = PosX;
+				RootSelfNode->NodePosY = PosY;
+				RootGetterCreator.Finalize();
+
+				PreviousDataPin = RootSelfNode->FindPinChecked(UEdGraphSchema_K2::PSC_Self);
+				PreviousClass = Blueprint->SkeletonGeneratedClass;
+			}
+			else
+			{
+				FGraphNodeCreator<UK2Node_VariableGet> RootGetterCreator(*FunctionGraph);
+				UK2Node_VariableGet* RootGetterNode = RootGetterCreator.CreateNode();
+				RootGetterNode->NodePosX = PosX;
+				RootGetterNode->NodePosY = PosY;
+				RootGetterNode->VariableReference.SetFromField<FProperty>(RootProperty, true, Blueprint->SkeletonGeneratedClass);
+				RootGetterCreator.Finalize();
+
+				PreviousDataPin = RootGetterNode->Pins[0];
+				PreviousClass = CastField<FObjectProperty>(RootProperty)->PropertyClass;
+			}
+
+			PosX += 300;
+		}
+
+		// create all the subsequent nodes in the path
+		for (int32 Index = 0; Index < Fields.Num(); ++Index)
+		{
+			if (!ensureMsgf(PreviousClass != nullptr, TEXT("Previous class not set!")))
+			{
+				return;
+			}
+
+			UEdGraphNode* NewNode = nullptr;
+			const UE::MVVM::FMVVMConstFieldVariant& Field = Fields[Index];
+			const bool bLastField = (Index == Fields.Num() - 1);
+			if (Field.IsProperty())
+			{
+				const FProperty* Property = Field.GetProperty();
+
+				// for structs in the middle of a path, we need to use a break node
+				const FStructProperty* StructProperty = CastField<FStructProperty>(Property);
+				if (!bLastField && StructProperty != nullptr)
+				{
+					FGraphNodeCreator<UK2Node_BreakStruct> BreakCreator(*FunctionGraph);
+					UK2Node_BreakStruct* BreakNode = BreakCreator.CreateNode();
+					BreakNode->AllocateDefaultPins();
+					BreakCreator.Finalize();
+
+					NewNode = BreakNode;
+
+					PreviousClass = nullptr;
+				}
+				else if (PreviousClass != nullptr)
+				{
+					FGraphNodeCreator<UK2Node_VariableGet> GetterCreator(*FunctionGraph);
+					UK2Node_VariableGet* GetterNode = GetterCreator.CreateNode();
+					GetterNode->SetFromProperty(Property, false, PreviousClass);
+					GetterNode->AllocateDefaultPins();
+					GetterCreator.Finalize();
+
+					NewNode = GetterNode;
+					if (const FObjectProperty* ObjectProperty = CastField<FObjectProperty>(Property))
+					{
+						PreviousClass = ObjectProperty->PropertyClass;
+					}
+					else
+					{
+						PreviousClass = nullptr;
+					}
+				}
+			}
+			else if (Field.IsFunction())
+			{
+				const UFunction* Function = Field.GetFunction();
+
+				FGraphNodeCreator<UK2Node_CallFunction> CallFunctionCreator(*FunctionGraph);
+				UK2Node_CallFunction* FunctionNode = CallFunctionCreator.CreateNode();
+				FunctionNode->SetFromFunction(Function);
+				FunctionNode->AllocateDefaultPins();
+				CallFunctionCreator.Finalize();
+
+				NewNode = FunctionNode;
+
+				const FProperty* ReturnProperty = UE::MVVM::BindingHelper::GetReturnProperty(Function);
+				if (const FObjectProperty* ObjectProperty = CastField<FObjectProperty>(ReturnProperty))
+				{
+					PreviousClass = ObjectProperty->PropertyClass;
+				}
+				else
+				{
+					PreviousClass = nullptr;
+				}
+			}
+			else
+			{
+				ensureMsgf(false, TEXT("Invalid path, empty field in path."));
+				return;
+			}
+
+			// create new connections
+			for (UEdGraphPin* Pin : NewNode->Pins)
+			{
+				if (Pin->Direction == EGPD_Input)
+				{
+					if (Pin->PinName == UEdGraphSchema_K2::PN_Execute)
+					{
+						Schema->MovePinLinks(*PreviousExecPin, *Pin, false);
+					}
+					else if (Schema->ArePinsCompatible(PreviousDataPin, Pin, PreviousClass))
+					{
+						Pin->MakeLinkTo(PreviousDataPin);
+					}
+				}
+			}
+
+			// then update our previous pin pointers
+			for (UEdGraphPin* Pin : NewNode->Pins)
+			{
+				if (Pin->Direction == EGPD_Output)
+				{
+					if (Pin->PinName == UEdGraphSchema_K2::PN_Then)
+					{
+						PreviousExecPin = Pin;
+					}
+					else
+					{
+						PreviousDataPin = Pin;
+					}
+				}
+			}
+
+			NewNode->NodePosX = PosX;
+			NewNode->NodePosY = PosY;
+
+			PosX += 300;
+		}
+
+		// finish by linking to our ultimate destinations
+		UEdGraphPin* LastExecPin = ConversionNode->FindPin(UEdGraphSchema_K2::PN_Execute);
+		if (LastExecPin == nullptr)
+		{
+			LastExecPin = Private::FindFunctionResult(FunctionGraph)->FindPin(UEdGraphSchema_K2::PN_Execute);
+		}
+		PreviousExecPin->MakeLinkTo(LastExecPin);
+		PreviousDataPin->MakeLinkTo(PathPin);
+	}
+}
+
+FMVVMBlueprintPropertyPath GetPropertyPathForArgument(const UBlueprint* WidgetBlueprint, const UK2Node_CallFunction* FunctionNode, FName ArgumentName, bool bSkipResolve)
+{
+	const UEdGraphPin* ArgumentPin = FunctionNode->FindPin(ArgumentName, EGPD_Input);
+	return GetPropertyPathForPin(WidgetBlueprint, ArgumentPin, bSkipResolve);
+}
+
+TMap<FName, FMVVMBlueprintPropertyPath> GetAllArgumentPropertyPaths(const UBlueprint* Blueprint, const UK2Node_CallFunction* FunctionNode, bool bSkipResolve)
 {
 	check(FunctionNode);
 
 	TMap<FName, FMVVMBlueprintPropertyPath> Paths;
 	for (const UEdGraphPin* Pin : FunctionNode->GetAllPins())
 	{
-		FMVVMBlueprintPropertyPath Path = Private::GetPropertyPathForPin(WidgetBlueprint, Pin, bSkipResolve);
+		FMVVMBlueprintPropertyPath Path = Private::GetPropertyPathForPin(Blueprint, Pin, bSkipResolve);
 		if (!Path.IsEmpty())
 		{
 			Paths.Add(Pin->PinName, Path);
@@ -210,13 +631,16 @@ TMap<FName, FMVVMBlueprintPropertyPath> GetAllArgumentPropertyPaths(const UWidge
 	return Paths;
 }
 
-TMap<FName, FMVVMBlueprintPropertyPath> GetAllArgumentPropertyPaths(const UWidgetBlueprint* WidgetBlueprint, const FMVVMBlueprintViewBinding& Binding, bool bSourceToDestination, bool bSkipResolve)
+TMap<FName, FMVVMBlueprintPropertyPath> GetAllArgumentPropertyPaths(const UBlueprint* Blueprint, const FMVVMBlueprintViewBinding& Binding, bool bSourceToDestination, bool bSkipResolve)
 {
-	if (UEdGraph* ConversionFunctionGraph = ConversionFunctionHelper::GetGraph(WidgetBlueprint, Binding, bSourceToDestination))
+	if (UMVVMBlueprintViewConversionFunction* ConversionFunction = Binding.Conversion.GetConversionFunction(bSourceToDestination))
 	{
-		if (UK2Node_CallFunction* ConversionNode = ConversionFunctionHelper::GetFunctionNode(ConversionFunctionGraph))
+		if (UEdGraph* ConversionFunctionGraph = ConversionFunction->GetWrapperGraph())
 		{
-			return GetAllArgumentPropertyPaths(WidgetBlueprint, ConversionNode, bSkipResolve);
+			if (UK2Node_CallFunction* ConversionNode = Cast<UK2Node_CallFunction>(ConversionFunctionHelper::GetWrapperNode(ConversionFunctionGraph)))
+			{
+				return GetAllArgumentPropertyPaths(Blueprint, ConversionNode, bSkipResolve);
+			}
 		}
 	}
 
diff --git a/Engine/Plugins/Runtime/ModelViewViewModel/Source/ModelViewViewModelBlueprint/Private/MVVMBlueprintPin.cpp b/Engine/Plugins/Runtime/ModelViewViewModel/Source/ModelViewViewModelBlueprint/Private/MVVMBlueprintPin.cpp
new file mode 100644
--- /dev/null
+++ b/Engine/Plugins/Runtime/ModelViewViewModel/Source/ModelViewViewModelBlueprint/Private/MVVMBlueprintPin.cpp
@@ -0,0 +1,82 @@
+// Copyright Epic Games, Inc. All Rights Reserved.
+
+#include "MVVMBlueprintPin.h"
+
+#include "Bindings/MVVMConversionFunctionHelper.h"
+#include "EdGraph/EdGraphPin.h"
+
+#include UE_INLINE_GENERATED_CPP_BY_NAME(MVVMBlueprintPin)
+
+FMVVMBlueprintPin FMVVMBlueprintPin::CreateFromPin(const UBlueprint* Blueprint, const UEdGraphPin* Pin)
+{
+	FMVVMBlueprintPin Result;
+	Result.PinName = Pin->PinName;
+	Result.Path = UE::MVVM::ConversionFunctionHelper::GetPropertyPathForPin(Blueprint, Pin, true);
+	Result.DefaultObject = Pin->DefaultObject;
+	Result.DefaultString = Pin->DefaultValue;
+	Result.DefaultText = Pin->DefaultTextValue;
+	return Result;
+}
+
+void FMVVMBlueprintPin::SetDefaultValue(UObject* Value)
+{
+	Reset();
+	DefaultObject = Value;
+}
+
+void FMVVMBlueprintPin::SetDefaultValue(const FText& Value)
+{
+	Reset();
+	DefaultText = Value;
+}
+
+void FMVVMBlueprintPin::SetDefaultValue(const FString& Value)
+{
+	Reset();
+	DefaultString = Value;
+}
+
+void FMVVMBlueprintPin::SetPath(const FMVVMBlueprintPropertyPath& Value)
+{
+	Reset();
+	Path = Value;
+}
+
+FString FMVVMBlueprintPin::GetValueAsString(const UClass* SelfContext) const
+{
+	if (!Path.IsEmpty())
+	{
+		return Path.GetPropertyPath(SelfContext);
+	}
+	else if (DefaultObject)
+	{
+		return DefaultObject.GetPathName();
+	}
+	else if (!DefaultText.IsEmpty())
+	{
+		FString TextAsString;
+		FTextStringHelper::WriteToBuffer(TextAsString, DefaultText);
+		return TextAsString;
+	}
+	else
+	{
+		return DefaultString;
+	}
+}
+
+void FMVVMBlueprintPin::CopyTo(const UBlueprint* Blueprint, UEdGraphPin* Pin) const
+{
+	Pin->DefaultObject = DefaultObject;
+	Pin->DefaultValue = DefaultString;
+	Pin->DefaultTextValue = DefaultText;
+
+	UE::MVVM::ConversionFunctionHelper::SetPropertyPathForPin(Blueprint, Path, Pin);
+}
+
+void FMVVMBlueprintPin::Reset()
+{
+	Path = FMVVMBlueprintPropertyPath();
+	DefaultString.Empty();
+	DefaultText = FText();
+	DefaultObject = nullptr;
+}
diff --git a/Engine/Plugins/Runtime/ModelViewViewModel/Source/ModelViewViewModelBlueprint/Private/MVVMBlueprintView.cpp b/Engine/Plugins/Runtime/ModelViewViewModel/Source/ModelViewViewModelBlueprint/Private/MVVMBlueprintView.cpp
--- a/Engine/Plugins/Runtime/ModelViewViewModel/Source/ModelViewViewModelBlueprint/Private/MVVMBlueprintView.cpp
+++ b/Engine/Plugins/Runtime/ModelViewViewModel/Source/ModelViewViewModelBlueprint/Private/MVVMBlueprintView.cpp
@@ -4,6 +4,7 @@
 
 #include "Components/Widget.h"
 #include "Kismet2/BlueprintEditorUtils.h"
+#include "MVVMBlueprintViewConversionFunction.h"
 #include "MVVMWidgetBlueprintExtension_View.h"
 #include "WidgetBlueprint.h"
 #include "UObject/FortniteMainBranchObjectVersion.h"
@@ -112,6 +113,15 @@ void UMVVMBlueprintView::RemoveBindingAt(int32 Index)
 {
 	if (Bindings.IsValidIndex(Index))
 	{
+		if (Bindings[Index].Conversion.SourceToDestinationConversion)
+		{
+			Bindings[Index].Conversion.SourceToDestinationConversion->RemoveWrapperGraph(GetOuterUMVVMWidgetBlueprintExtension_View()->GetWidgetBlueprint());
+		}
+		if (Bindings[Index].Conversion.DestinationToSourceConversion)
+		{
+			Bindings[Index].Conversion.DestinationToSourceConversion->RemoveWrapperGraph(GetOuterUMVVMWidgetBlueprintExtension_View()->GetWidgetBlueprint());
+		}
+
 		Bindings.RemoveAt(Index);
 		OnBindingsUpdated.Broadcast();
 	}
@@ -232,6 +242,7 @@ void UMVVMBlueprintView::Serialize(FArchive& Ar)
 	Ar.UsingCustomVersion(FFortniteMainBranchObjectVersion::GUID);	
 }
 
+#if WITH_EDITOR
 void UMVVMBlueprintView::PostLoad()
 {
 	Super::PostLoad();
@@ -244,14 +255,13 @@ void UMVVMBlueprintView::PostLoad()
 		}
 	}
 
-#if WITH_EDITOR
+	GetOuterUMVVMWidgetBlueprintExtension_View()->ConditionalPostLoad();
+	GetOuterUMVVMWidgetBlueprintExtension_View()->GetWidgetBlueprint()->ConditionalPostLoad();
+
 	// Make sure all bindings uses the skeletal class
 	if (GetLinkerCustomVersion(FFortniteMainBranchObjectVersion::GUID) < FFortniteMainBranchObjectVersion::MVVMConvertPropertyPathToSkeletalClass)
 	{
-		GetOuterUMVVMWidgetBlueprintExtension_View()->ConditionalPostLoad();
-		GetOuterUMVVMWidgetBlueprintExtension_View()->GetWidgetBlueprint()->ConditionalPostLoad();
 		const UWidgetBlueprint* ThisWidgetBlueprint = GetOuterUMVVMWidgetBlueprintExtension_View()->GetWidgetBlueprint();
-
 		for (FMVVMBlueprintViewBinding& Binding : Bindings)
 		{
 			for (const FMVVMBlueprintFieldPath& FieldPath : Binding.SourcePath.GetFieldPaths())
@@ -264,17 +274,23 @@ void UMVVMBlueprintView::PostLoad()
 			}
 		}
 	}
-#endif
-}
 
-#if WITH_EDITOR
+	for (FMVVMBlueprintViewBinding& Binding : Bindings)
+	{
+		Binding.Conversion.DeprecateViewConversionFunction(GetOuterUMVVMWidgetBlueprintExtension_View()->GetWidgetBlueprint());
+	}
+}
 
-/*
-void UMVVMBlueprintView::PostEditUndo() 
+void UMVVMBlueprintView::PreSave(FObjectPreSaveContext Context)
 {
-	OnBindingsUpdated.Broadcast();
-	OnViewModelsUpdated.Broadcast();
-}*/
+	UWidgetBlueprint* WidgetBlueprint = GetOuterUMVVMWidgetBlueprintExtension_View()->GetWidgetBlueprint();
+	for (FMVVMBlueprintViewBinding& Binding : Bindings)
+	{
+		Binding.Conversion.SavePinValues(WidgetBlueprint);
+	}
+
+	Super::PreSave(Context);
+}
 
 void UMVVMBlueprintView::PostEditChangeProperty(FPropertyChangedEvent& PropertyChangedEvent)
 {
diff --git a/Engine/Plugins/Runtime/ModelViewViewModel/Source/ModelViewViewModelBlueprint/Private/MVVMBlueprintViewBinding.cpp b/Engine/Plugins/Runtime/ModelViewViewModel/Source/ModelViewViewModelBlueprint/Private/MVVMBlueprintViewBinding.cpp
--- a/Engine/Plugins/Runtime/ModelViewViewModel/Source/ModelViewViewModelBlueprint/Private/MVVMBlueprintViewBinding.cpp
+++ b/Engine/Plugins/Runtime/ModelViewViewModel/Source/ModelViewViewModelBlueprint/Private/MVVMBlueprintViewBinding.cpp
@@ -6,12 +6,67 @@
 #include "Bindings/MVVMConversionFunctionHelper.h"
 #include "K2Node_CallFunction.h"
 #include "MVVMBlueprintView.h"
+#include "MVVMBlueprintViewConversionFunction.h"
 #include "MVVMWidgetBlueprintExtension_View.h"
 
+#include "MVVMFunctionGraphHelper.h"
+
 #include UE_INLINE_GENERATED_CPP_BY_NAME(MVVMBlueprintViewBinding)
 
 #define LOCTEXT_NAMESPACE "MVVMBlueprintViewBinding"
 
+void FMVVMBlueprintViewConversionPath::GenerateWrapper(UBlueprint* Blueprint)
+{
+	if (DestinationToSourceConversion)
+	{
+		DestinationToSourceConversion->GetOrCreateWrapperGraph(Blueprint);
+	}
+	if (SourceToDestinationConversion)
+	{
+		SourceToDestinationConversion->GetOrCreateWrapperGraph(Blueprint);
+	}
+}
+
+void FMVVMBlueprintViewConversionPath::SavePinValues(UBlueprint* Blueprint)
+{
+	if (DestinationToSourceConversion)
+	{
+		DestinationToSourceConversion->SavePinValues(Blueprint);
+	}
+	if (SourceToDestinationConversion)
+	{
+		SourceToDestinationConversion->SavePinValues(Blueprint);
+	}
+}
+
+void FMVVMBlueprintViewConversionPath::DeprecateViewConversionFunction(UBlueprint* Blueprint)
+{
+	auto Deprecate = [Blueprint](FName& Wrapper, FMemberReference& Reference, TObjectPtr<UMVVMBlueprintViewConversionFunction>& ConversionFunction)
+	{
+		if (!Wrapper.IsNone())
+		{
+			TObjectPtr<UEdGraph>* GraphPtr = Blueprint->FunctionGraphs.FindByPredicate([Wrapper](const UEdGraph* Other) { return Other->GetFName() == Wrapper; });
+			if (GraphPtr)
+			{
+				ConversionFunction = NewObject<UMVVMBlueprintViewConversionFunction>(Blueprint);
+				ConversionFunction->InitializeFromWrapperGraph(Blueprint, *GraphPtr);
+			}
+		}
+		else if (!Reference.GetMemberName().IsNone())
+		{
+			ConversionFunction = NewObject<UMVVMBlueprintViewConversionFunction>(Blueprint);
+			ConversionFunction->InitializeFromMemberReference(Blueprint, Reference);
+		}
+		Wrapper = FName();
+		Reference = FMemberReference();
+	};
+
+	PRAGMA_DISABLE_DEPRECATION_WARNINGS
+	Deprecate(DestinationToSourceWrapper_DEPRECATED, DestinationToSourceFunction_DEPRECATED, DestinationToSourceConversion);
+	Deprecate(SourceToDestinationWrapper_DEPRECATED, SourceToDestinationFunction_DEPRECATED, SourceToDestinationConversion);
+	PRAGMA_ENABLE_DEPRECATION_WARNINGS
+}
+
 FName FMVVMBlueprintViewBinding::GetFName() const
 {
 	return *BindingId.ToString(EGuidFormats::DigitsWithHyphensLower);
@@ -169,31 +224,37 @@ namespace UE::MVVM::Private
 
 		TStringBuilder<256> NameBuilder;
 
-		if (UEdGraph* Graph = ConversionFunctionHelper::GetGraph(WidgetBlueprint, Binding, bIsSource))
+		auto AddPath = [&](const FMVVMBlueprintPropertyPath& ArgumentPath)
 		{
-			if (const UK2Node_CallFunction* CallFunctionNode = ConversionFunctionHelper::GetFunctionNode(Graph))
+			if (ArgumentPath.IsFromViewModel())
 			{
-				if (bUseDisplayName)
+				AppendViewModelPathString(WidgetBlueprint, BlueprintView, ArgumentPath, NameBuilder, FunctionKeywordsBuilder, bUseDisplayName, bAppendFunctionKeywords);
+			}
+			else
+			{
+				AppendWidgetPathString(WidgetBlueprint, ArgumentPath, NameBuilder, FunctionKeywordsBuilder, bUseDisplayName, bAppendFunctionKeywords);
+			}
+		};
+
+		if (UMVVMBlueprintViewConversionFunction* ViewConversionFunction = Binding.Conversion.GetConversionFunction(bIsSource))
+		{
+			TVariant<const UFunction*, TSubclassOf<UK2Node>> ConversionFunction = ViewConversionFunction->GetConversionFunction(WidgetBlueprint->SkeletonGeneratedClass);
+			if (ConversionFunction.IsType<const UFunction*>())
+			{
+				if (const UFunction* ConversionFunctionPtr = ConversionFunction.Get<const UFunction*>())
 				{
-					if (UFunction* TargetFunction = CallFunctionNode->GetTargetFunction())
+					if (bUseDisplayName)
 					{
-						NameBuilder << TargetFunction->GetDisplayNameText().ToString();
+						NameBuilder << ConversionFunctionPtr->GetDisplayNameText().ToString();
 					}
-					else
+					else 
 					{
-						NameBuilder << CallFunctionNode->GetFunctionName();
+						NameBuilder << ConversionFunctionPtr->GetFName();
 					}
-				}
-				else
-				{
-					NameBuilder << CallFunctionNode->GetFunctionName();
-				}
 
-				if (bAppendFunctionKeywords)
-				{
-					if (UFunction* TargetFunction = CallFunctionNode->GetTargetFunction())
+					if (bAppendFunctionKeywords)
 					{
-						FString FunctionKeywords = TargetFunction->GetMetaData(FBlueprintMetadata::MD_FunctionKeywords);
+						FString FunctionKeywords = ConversionFunctionPtr->GetMetaData(FBlueprintMetadata::MD_FunctionKeywords);
 						if (!FunctionKeywords.IsEmpty())
 						{
 							FunctionKeywordsBuilder << TEXT(".");
@@ -201,107 +262,72 @@ namespace UE::MVVM::Private
 						}
 					}
 				}
-
-				NameBuilder << TEXT("(");
-
-				bool bFirst = true;
-
-				for (const UEdGraphPin* Pin : CallFunctionNode->Pins)
+				else
 				{
-					if (Pin->PinName == UEdGraphSchema_K2::PN_Self || 
-						Pin->Direction != EGPD_Input)
-					{
-						continue;
-					}
-
-					if (!bFirst)
-					{
-						NameBuilder << TEXT(", ");
-					}
-
-					FMVVMBlueprintPropertyPath ArgumentPath = ConversionFunctionHelper::GetPropertyPathForArgument(WidgetBlueprint, CallFunctionNode, Pin->GetFName(), true);
-					if (!ArgumentPath.IsEmpty())
-					{
-						if (ArgumentPath.IsFromViewModel())
-						{
-							AppendViewModelPathString(WidgetBlueprint, BlueprintView, ArgumentPath, NameBuilder, FunctionKeywordsBuilder, bUseDisplayName, bAppendFunctionKeywords);
-						}
-						else
-						{
-							AppendWidgetPathString(WidgetBlueprint, ArgumentPath, NameBuilder, FunctionKeywordsBuilder, bUseDisplayName, bAppendFunctionKeywords);
-						}
-					}
-					else
-					{
-						NameBuilder << Pin->GetDefaultAsString();
-					}
-
-					bFirst = false;
+					NameBuilder << TEXT("<ERROR>");
 				}
-
-				NameBuilder << TEXT(")");
 			}
 			else
 			{
-				NameBuilder << TEXT("<error>");
-			}
-		}
-		else
-		{
-			const FMVVMBlueprintPropertyPath Path = bIsSource ? Binding.SourcePath : Binding.DestinationPath;
-			const FMemberReference FunctionReference = bIsSource ? Binding.Conversion.SourceToDestinationFunction : Binding.Conversion.DestinationToSourceFunction;
-			if (!FunctionReference.GetMemberName().IsNone())
-			{
-				if (const UFunction* Function = FunctionReference.ResolveMember<UFunction>(WidgetBlueprint->SkeletonGeneratedClass))
+				TSubclassOf<UK2Node> ConversionFunctionNode = ConversionFunction.Get<TSubclassOf<UK2Node>>();
+				if (ConversionFunctionNode.Get())
 				{
-					NameBuilder << Function->GetName();
-
-					if (bAppendFunctionKeywords)
+					if (bUseDisplayName)
 					{
-						FString FunctionKeywords = Function->GetMetaData(FBlueprintMetadata::MD_FunctionKeywords);
-						if (!FunctionKeywords.IsEmpty())
-						{
-							FunctionKeywordsBuilder << TEXT(".");
-							FunctionKeywordsBuilder << FunctionKeywords;
-						}
-					}
-
-					NameBuilder << TEXT("(");
-
-					if (Path.IsFromViewModel())
-					{
-						AppendViewModelPathString(WidgetBlueprint, BlueprintView, Path, NameBuilder, FunctionKeywordsBuilder, bUseDisplayName, bAppendFunctionKeywords);
-
+						NameBuilder << ConversionFunctionNode->GetDisplayNameText().ToString();
 					}
 					else
 					{
-						AppendWidgetPathString(WidgetBlueprint, Path, NameBuilder, FunctionKeywordsBuilder, bUseDisplayName, bAppendFunctionKeywords);
+						NameBuilder << ConversionFunctionNode->GetFName();
 					}
-
-					NameBuilder << TEXT(")");
 				}
 				else
 				{
-					NameBuilder << FunctionReference.GetMemberName();
-					NameBuilder << TEXT("()");
+					NameBuilder << TEXT("<ERROR>");
 				}
 			}
-			else
+
+			// AddPins
 			{
-				if (!Path.IsEmpty())
+				NameBuilder << TEXT("(");
+
+				if (ViewConversionFunction->NeedsWrapperGraph())
 				{
-					if (Path.IsFromViewModel())
+					bool bFirst = true;
+					for (const FMVVMBlueprintPin& Pin : ViewConversionFunction->GetPins())
 					{
-						AppendViewModelPathString(WidgetBlueprint, BlueprintView, Path, NameBuilder, FunctionKeywordsBuilder, bUseDisplayName, bAppendFunctionKeywords);
+						if (!bFirst)
+						{
+							NameBuilder << TEXT(", ");
+						}
 
-					}
-					else
-					{
-						AppendWidgetPathString(WidgetBlueprint, Path, NameBuilder, FunctionKeywordsBuilder, bUseDisplayName, bAppendFunctionKeywords);
+
+
+						if (Pin.UsedPathAsValue())
+						{
+							AddPath(Pin.GetPath());
+						}
+						else
+						{
+							NameBuilder << Pin.GetValueAsString(WidgetBlueprint->SkeletonGeneratedClass);
+						}
+
+						bFirst = false;
 					}
 				}
-			} 
+				else
+				{
+					const FMVVMBlueprintPropertyPath& Path = bIsSource ? Binding.SourcePath : Binding.DestinationPath;
+					AddPath(Path);
+				}
 
+				NameBuilder << TEXT(")");
+			}
+		}
+		else
+		{
+			const FMVVMBlueprintPropertyPath& PropertyPath = bIsSource ? Binding.DestinationPath : Binding.SourcePath;
+			AddPath(PropertyPath);
 		}
 
 		FString Name = NameBuilder.ToString();
diff --git a/Engine/Plugins/Runtime/ModelViewViewModel/Source/ModelViewViewModelBlueprint/Private/MVVMBlueprintViewConversionFunction.cpp b/Engine/Plugins/Runtime/ModelViewViewModel/Source/ModelViewViewModelBlueprint/Private/MVVMBlueprintViewConversionFunction.cpp
new file mode 100644
--- /dev/null
+++ b/Engine/Plugins/Runtime/ModelViewViewModel/Source/ModelViewViewModelBlueprint/Private/MVVMBlueprintViewConversionFunction.cpp
@@ -0,0 +1,313 @@
+// Copyright Epic Games, Inc. All Rights Reserved.
+
+#include "MVVMBlueprintViewConversionFunction.h"
+#include "MVVMDeveloperProjectSettings.h"
+
+#include "Bindings/MVVMConversionFunctionHelper.h"
+#include "Bindings/MVVMBindingHelper.h"
+#include "Engine/Blueprint.h"
+#include "Kismet2/BlueprintEditorUtils.h"
+
+#include "EdGraphSchema_K2.h"
+#include "EdGraph/EdGraphPin.h"
+#include "K2Node_CallFunction.h"
+#include "K2Node_FunctionEntry.h"
+#include "K2Node_FunctionResult.h"
+#include "KismetCompiler.h"
+
+#include UE_INLINE_GENERATED_CPP_BY_NAME(MVVMBlueprintViewConversionFunction)
+
+namespace UE::MVVM::Private
+{
+	static FLazyName DefaultConversionFunctionName = TEXT("__ConversionFunction");
+}
+
+const UFunction* UMVVMBlueprintViewConversionFunction::GetCompiledFunction(const UClass* SelfContext) const
+{
+	if (NeedsWrapperGraph())
+	{
+		FMemberReference CompiledFunction;
+		CompiledFunction.SetSelfMember(GraphName);
+		return CompiledFunction.ResolveMember<UFunction>(const_cast<UClass*>(SelfContext));
+	}
+	return FunctionReference.ResolveMember<UFunction>(const_cast<UClass*>(SelfContext));
+}
+
+FName UMVVMBlueprintViewConversionFunction::GetCompiledFunctionName() const
+{
+	if (NeedsWrapperGraph())
+	{
+		return GraphName;
+	}
+	return FunctionReference.GetMemberName();
+}
+
+TVariant<const UFunction*, TSubclassOf<UK2Node>> UMVVMBlueprintViewConversionFunction::GetConversionFunction(const UClass* SelfContext) const
+{
+	if (FunctionNode.Get())
+	{
+		return TVariant<const UFunction*, TSubclassOf<UK2Node>>( TInPlaceType<TSubclassOf<UK2Node>>(), FunctionNode);
+	}
+	else
+	{
+		const UFunction* ConversionFunction = FunctionReference.ResolveMember<UFunction>(const_cast<UClass*>(SelfContext));
+		return TVariant<const UFunction*, TSubclassOf<UK2Node>>( TInPlaceType<const UFunction*>(), ConversionFunction);
+	}
+	//return FunctionReference.ResolveMember<UFunction>(WidgetBlueprint->SkeletonGeneratedClass);
+//UK2Node_CallFunction* CallFunctionNode = GetConversionFunctionNode(WidgetBlueprint, Binding, bSourceToDestination);
+//if (CallFunctionNode)
+//{
+//	return CallFunctionNode->GetTargetFunction();
+//}
+}
+
+void UMVVMBlueprintViewConversionFunction::Reset()
+{
+	FunctionReference = FMemberReference();
+	FunctionNode = TSubclassOf<UK2Node>();
+	GraphName = FName();
+	bWrapperGraphTransient = false;
+	SavedPins.Reset();
+	CachedWrapperGraph = nullptr;
+	CachedWrapperNode = nullptr;
+}
+
+void UMVVMBlueprintViewConversionFunction::InitFromFunction(UBlueprint* InContext, const UFunction* InFunction)
+{
+	InitFromFunction(InContext, InFunction, UE::MVVM::Private::DefaultConversionFunctionName.Resolve());
+}
+
+void UMVVMBlueprintViewConversionFunction::InitFromFunction(UBlueprint* InContext, const UFunction* InFunction, FName InGraphName)
+{
+	Reset();
+
+	if (InFunction)
+	{
+		UClass* OwnerClass = InFunction->GetOuterUClass();
+		FGuid MemberGuid;
+		UBlueprint::GetGuidFromClassByFieldName<UFunction>(OwnerClass, InFunction->GetFName(), MemberGuid);
+
+		bool bIsSelf = (InContext->GeneratedClass && InContext->GeneratedClass->IsChildOf(OwnerClass))
+			|| (InContext->SkeletonGeneratedClass && InContext->SkeletonGeneratedClass->IsChildOf(OwnerClass));
+		if (bIsSelf)
+		{
+			FunctionReference.SetSelfMember(InFunction->GetFName(), MemberGuid);
+		}
+		else
+		{
+			if (UBlueprint* VariableOwnerBP = Cast<UBlueprint>(OwnerClass->ClassGeneratedBy))
+			{
+				OwnerClass = VariableOwnerBP->SkeletonGeneratedClass;
+			}
+
+			FunctionReference.SetExternalMember(InFunction->GetFName(), OwnerClass, MemberGuid);
+		}
+		// used to be SetFromField<UFunction>(Function, SelfContext->SkeletonGeneratedClass);
+
+		if (UE::MVVM::ConversionFunctionHelper::RequiresWrapper(InFunction))
+		{
+			GraphName = InGraphName;
+			bWrapperGraphTransient = !GetDefault<UMVVMDeveloperProjectSettings>()->bAllowConversionFunctionGeneratedGraphInEditor;
+			GetOrCreateWrapperGraphInternal(InContext, InFunction);
+			SavePinValues(InContext);
+		}
+	}
+}
+
+void UMVVMBlueprintViewConversionFunction::InitializeFromWrapperGraph(UBlueprint* SelfContext, UEdGraph* Graph)
+{
+	Reset();
+
+	if (UK2Node* WrapperNode = UE::MVVM::ConversionFunctionHelper::GetWrapperNode(Graph))
+	{
+		if (UK2Node_CallFunction* CallFunction = Cast<UK2Node_CallFunction>(WrapperNode))
+		{
+			FunctionReference = CallFunction->FunctionReference;
+		}
+		else
+		{
+			FunctionNode = WrapperNode->GetClass();
+		}
+		CachedWrapperGraph = WrapperNode->GetGraph();
+		check(CachedWrapperGraph);
+
+		GraphName = CachedWrapperGraph->GetFName();
+		bWrapperGraphTransient = !GetDefault<UMVVMDeveloperProjectSettings>()->bAllowConversionFunctionGeneratedGraphInEditor;
+
+		CachedWrapperNode = WrapperNode;
+		SavePinValues(SelfContext);
+
+		if (bWrapperGraphTransient && bWrapperGraphTransient)
+		{
+			SelfContext->FunctionGraphs.RemoveSingle(CachedWrapperGraph);
+		}
+	}
+}
+
+void UMVVMBlueprintViewConversionFunction::InitializeFromMemberReference(UBlueprint* SelfContext, FMemberReference MemberReference)
+{
+	Reset();
+
+	FunctionReference = MemberReference;
+}
+
+UEdGraph* UMVVMBlueprintViewConversionFunction::GetOrCreateIntermediateWrapperGraph(FKismetCompilerContext& Context) const
+{
+	if (!NeedsWrapperGraph())
+	{
+		return nullptr;
+	}
+
+	if (CachedWrapperGraph)
+	{
+		return CachedWrapperGraph;
+	}
+
+	if (IsWrapperGraphTransient())
+	{
+		if (FunctionNode.Get())
+		{
+			ensure(false); // not supported for now
+			return nullptr;
+		}
+		else
+		{
+			const UFunction* ConversionFunction = FunctionReference.ResolveMember<UFunction>(Context.NewClass);
+			return GetOrCreateWrapperGraphInternal(Context, ConversionFunction);
+		}
+	}
+	else
+	{
+		TObjectPtr<UEdGraph>* Found = Context.Blueprint->FunctionGraphs.FindByPredicate([GraphName = GetWrapperGraphName()](const UEdGraph* Other) { return Other->GetFName() == GraphName; });
+		if (Found)
+		{
+			CachedWrapperGraph = *Found;
+			CachedWrapperNode = UE::MVVM::ConversionFunctionHelper::GetWrapperNode(CachedWrapperGraph);
+		}
+		return CachedWrapperGraph;
+	}
+}
+
+UEdGraph* UMVVMBlueprintViewConversionFunction::GetOrCreateWrapperGraph(UBlueprint* Blueprint) const
+{
+	if (!NeedsWrapperGraph())
+	{
+		return nullptr;
+	}
+
+	if (CachedWrapperGraph)
+	{
+		return CachedWrapperGraph;
+	}
+
+	if (IsWrapperGraphTransient())
+	{
+		if (FunctionNode.Get())
+		{
+			ensure(false); // not supported for now
+			return nullptr;
+		}
+		else
+		{
+			const UFunction* ConversionFunction = FunctionReference.ResolveMember<UFunction>(Blueprint->SkeletonGeneratedClass);
+			return GetOrCreateWrapperGraphInternal(Blueprint, ConversionFunction);
+		}
+	}
+	else
+	{
+		TObjectPtr<UEdGraph>* Found = Blueprint->FunctionGraphs.FindByPredicate([GraphName = GetWrapperGraphName()](const UEdGraph* Other) { return Other->GetFName() == GraphName; });
+		if (Found)
+		{
+			CachedWrapperGraph = *Found;
+			CachedWrapperNode = UE::MVVM::ConversionFunctionHelper::GetWrapperNode(CachedWrapperGraph);
+		}
+		return CachedWrapperGraph;
+	}
+}
+
+UEdGraphPin* UMVVMBlueprintViewConversionFunction::GetOrCreateGraphPin(UBlueprint* Blueprint, FName PinName) const
+{
+	GetOrCreateWrapperGraph(Blueprint);
+	if (CachedWrapperNode)
+	{
+		return CachedWrapperNode->FindPin(PinName);
+	}
+	return nullptr;
+}
+
+UEdGraph* UMVVMBlueprintViewConversionFunction::GetOrCreateWrapperGraphInternal(FKismetCompilerContext& Context, const UFunction* Function) const
+{
+	return GetOrCreateWrapperGraphInternal(Context.Blueprint, Function);
+}
+
+UEdGraph* UMVVMBlueprintViewConversionFunction::GetOrCreateWrapperGraphInternal(UBlueprint* Blueprint, const UFunction* Function) const
+{
+	TPair<UEdGraph*, UK2Node*> Result = UE::MVVM::ConversionFunctionHelper::CreateGraph(Blueprint, GraphName, Function, bWrapperGraphTransient);
+	CachedWrapperGraph = Result.Get<0>();
+	CachedWrapperNode = Result.Get<1>();
+	LoadPinValuesInternal(Blueprint);
+	return CachedWrapperGraph;
+}
+
+void UMVVMBlueprintViewConversionFunction::RemoveWrapperGraph(UBlueprint* Blueprint)
+{
+	TObjectPtr<UEdGraph>* Result = Blueprint->FunctionGraphs.FindByPredicate([WrapperName = GraphName](const UEdGraph* GraphPtr) { return GraphPtr->GetFName() == WrapperName; });
+	if (Result)
+	{
+		FBlueprintEditorUtils::RemoveGraph(Blueprint, Result->Get());
+	}
+	CachedWrapperGraph = nullptr;
+	CachedWrapperNode = nullptr;
+}
+
+void UMVVMBlueprintViewConversionFunction::SetGraphPin(UBlueprint* Blueprint, FName PinName, const FMVVMBlueprintPropertyPath& Path)
+{
+	if (!NeedsWrapperGraph())
+	{
+		return;
+	}
+
+	UEdGraphPin* GraphPin = GetOrCreateGraphPin(Blueprint, PinName);
+
+	// Set the value and make the blueprint as dirty before creating the pin.
+	//A property may not be created yet and the skeletal needs to be recreated.
+	FMVVMBlueprintPin* Pin = SavedPins.FindByPredicate([PinName](const FMVVMBlueprintPin& Other) { return PinName == Other.GetName(); });
+	if (!Pin)
+	{
+		Pin = &SavedPins.Add_GetRef(FMVVMBlueprintPin::CreateFromPin(Blueprint, GraphPin));
+	}
+	Pin->SetPath(Path);
+
+	FBlueprintEditorUtils::MarkBlueprintAsStructurallyModified(Blueprint);
+
+	UE::MVVM::ConversionFunctionHelper::SetPropertyPathForPin(Blueprint, Path, GraphPin);
+}
+
+void UMVVMBlueprintViewConversionFunction::SavePinValues(UBlueprint* Blueprint)
+{
+	if (CachedWrapperNode)
+	{
+		SavedPins.Reset();
+		for (UEdGraphPin* Pin : CachedWrapperNode->Pins)
+		{
+			if (Pin->PinName != UEdGraphSchema_K2::PN_Self && Pin->PinName != UEdGraphSchema_K2::PN_Execute && Pin->Direction == EGPD_Input)
+			{
+				SavedPins.Add(FMVVMBlueprintPin::CreateFromPin(Blueprint, Pin));
+			}
+		}
+	}
+}
+
+void UMVVMBlueprintViewConversionFunction::LoadPinValuesInternal(UBlueprint* Blueprint) const
+{
+	if (CachedWrapperNode)
+	{
+		for (const FMVVMBlueprintPin& Pin : SavedPins)
+		{
+			if (UEdGraphPin* GraphPin = CachedWrapperNode->FindPin(Pin.GetName()))
+			{
+				Pin.CopyTo(Blueprint, GraphPin);
+			}
+		}
+	}
+}
\ No newline at end of file
diff --git a/Engine/Plugins/Runtime/ModelViewViewModel/Source/ModelViewViewModelBlueprint/Private/MVVMViewBlueprintCompiler.cpp b/Engine/Plugins/Runtime/ModelViewViewModel/Source/ModelViewViewModelBlueprint/Private/MVVMViewBlueprintCompiler.cpp
--- a/Engine/Plugins/Runtime/ModelViewViewModel/Source/ModelViewViewModelBlueprint/Private/MVVMViewBlueprintCompiler.cpp
+++ b/Engine/Plugins/Runtime/ModelViewViewModel/Source/ModelViewViewModelBlueprint/Private/MVVMViewBlueprintCompiler.cpp
@@ -8,6 +8,7 @@
 #include "Components/Widget.h"
 #include "HAL/IConsoleManager.h"
 #include "MVVMBlueprintView.h"
+#include "MVVMBlueprintViewConversionFunction.h"
 #include "MVVMDeveloperProjectSettings.h"
 #include "MVVMMessageLog.h"
 #include "PropertyPermissionList.h"
@@ -205,19 +206,23 @@ void FMVVMViewBlueprintCompiler::CreateFunctions(UMVVMBlueprintView* BlueprintVi
 	// Build the list of BP destinations
 	CreateBindingDestinationContexts(BlueprintView);
 
-	if (!GetDefault<UMVVMDeveloperProjectSettings>()->bAllowGeneratedViewModelSetter)
+	if (!bAreSourcesCreatorValid || !bIsBindingsValid)
 	{
 		return;
 	}
 
-	for (const FCompilerSourceCreatorContext& SourceCreator : CompilerSourceCreatorContexts)
+	// Generate the setter code
+	if (GetDefault<UMVVMDeveloperProjectSettings>()->bAllowGeneratedViewModelSetter)
 	{
-		if (SourceCreator.SetterGraph)
+		for (const FCompilerSourceCreatorContext& SourceCreator : CompilerSourceCreatorContexts)
 		{
-			if (!UE::MVVM::FunctionGraphHelper::GenerateViewModelSetter(WidgetBlueprintCompilerContext, SourceCreator.SetterGraph, SourceCreator.ViewModelContext.GetViewModelName()))
+			if (SourceCreator.SetterGraph)
 			{
-				AddErrorForViewModel(SourceCreator.ViewModelContext, LOCTEXT("SetterFunctionCouldNotBeGenerated", "The setter function could not be generated."));
-				continue;
+				if (!UE::MVVM::FunctionGraphHelper::GenerateViewModelSetter(WidgetBlueprintCompilerContext, SourceCreator.SetterGraph, SourceCreator.ViewModelContext.GetViewModelName()))
+				{
+					AddErrorForViewModel(SourceCreator.ViewModelContext, LOCTEXT("SetterFunctionCouldNotBeGenerated", "The setter function could not be generated."));
+					continue;
+				}
 			}
 		}
 	}
@@ -242,6 +247,10 @@ void FMVVMViewBlueprintCompiler::CreateVariables(const FWidgetBlueprintCompilerC
 		CreateSourceLists(Context, BlueprintView);
 		CreateFunctionsDeclaration(Context, BlueprintView);
 	}
+	else
+	{
+		CreateIntermediateGraphFunctions(Context, BlueprintView);
+	}
 
 	auto CreateVariable = [&Context](const FCompilerUserWidgetPropertyContext& SourceContext) -> FProperty*
 	{
@@ -543,8 +552,8 @@ void FMVVMViewBlueprintCompiler::CreateSourceLists(const FWidgetBlueprintCompile
 
 		if (bIsForwardBinding || bIsBackwardBinding)
 		{
-			TMap<FName, FMVVMBlueprintPropertyPath> ConversionFunctionArguments = ConversionFunctionHelper::GetAllArgumentPropertyPaths(WidgetBlueprintCompilerContext.WidgetBlueprint(), Binding, bIsForwardBinding, true);
-			if (ConversionFunctionArguments.Num() > 0)
+			const UMVVMBlueprintViewConversionFunction* ConversionFunction = Binding.Conversion.GetConversionFunction(bIsForwardBinding);
+			if (ConversionFunction)
 			{
 				if (Binding.BindingType == EMVVMBindingMode::TwoWay)
 				{
@@ -554,9 +563,12 @@ void FMVVMViewBlueprintCompiler::CreateSourceLists(const FWidgetBlueprintCompile
 				}
 
 				// generate sources for conversion function arguments
-				for (const TPair<FName, FMVVMBlueprintPropertyPath>& Arg : ConversionFunctionArguments)
+				for (const FMVVMBlueprintPin& Pin: ConversionFunction->GetPins())
 				{
-					bAreSourceContextsValid &= GenerateCompilerSourceContext(Arg.Value, Arg.Key);
+					if (Pin.UsedPathAsValue())
+					{
+						bAreSourceContextsValid &= GenerateCompilerSourceContext(Pin.GetPath(), Pin.GetName());
+					}
 				}
 
 				// generate destination source
@@ -603,34 +615,81 @@ void FMVVMViewBlueprintCompiler::CreateFunctionsDeclaration(const FWidgetBluepri
 		BlueprintView->TemporaryGraph.Reset();
 	}
 
-	if (!GetDefault<UMVVMDeveloperProjectSettings>()->bAllowGeneratedViewModelSetter)
+	if (GetDefault<UMVVMDeveloperProjectSettings>()->bAllowGeneratedViewModelSetter)
 	{
-		return;
+		for (FCompilerSourceCreatorContext& SourceCreator : CompilerSourceCreatorContexts)
+		{
+			if (!SourceCreator.SetterFunctionName.IsEmpty() && SourceCreator.Type == ECompilerSourceCreatorType::ViewModel)
+			{
+				ensure(SourceCreator.SetterGraph == nullptr);
+
+				SourceCreator.SetterGraph = UE::MVVM::FunctionGraphHelper::CreateIntermediateFunctionGraph(
+					WidgetBlueprintCompilerContext
+					, SourceCreator.SetterFunctionName
+					, (FUNC_BlueprintCallable | FUNC_Public)
+					, TEXT("Viewmodel")
+					, false);
+				BlueprintView->TemporaryGraph.Add(SourceCreator.SetterGraph);
+
+				if (SourceCreator.SetterGraph == nullptr || SourceCreator.SetterGraph->GetFName() != FName(*SourceCreator.SetterFunctionName))
+				{
+					WidgetBlueprintCompilerContext.MessageLog.Error(*FText::Format(LOCTEXT("SetterNameAlreadyExists", "The setter name {0} already exists and could not be autogenerated."),
+						FText::FromString(SourceCreator.SetterFunctionName)
+					).ToString()
+					);
+				}
+
+				UE::MVVM::FunctionGraphHelper::AddFunctionArgument(SourceCreator.SetterGraph, SourceCreator.ViewModelContext.GetViewModelClass(), "Viewmodel");
+			}
+		}
 	}
+}
+
 
-	for (FCompilerSourceCreatorContext& SourceCreator : CompilerSourceCreatorContexts)
+void FMVVMViewBlueprintCompiler::CreateIntermediateGraphFunctions(const FWidgetBlueprintCompilerContext::FCreateVariableContext & Context, UMVVMBlueprintView * BlueprintView)
+{
+	// Add Generated Conversion functions to the Context.List
+	const int32 NumBindings = BlueprintView->GetNumBindings();
+	for (int32 Index = 0; Index < NumBindings; ++Index)
 	{
-		if (!SourceCreator.SetterFunctionName.IsEmpty() && SourceCreator.Type == ECompilerSourceCreatorType::ViewModel)
+		FMVVMBlueprintViewBinding* BindingPtr = BlueprintView->GetBindingAt(Index);
+		if (BindingPtr == nullptr)
 		{
-			ensure(SourceCreator.SetterGraph == nullptr);
+			WidgetBlueprintCompilerContext.MessageLog.Error(*FText::Format(LOCTEXT("InvalidBindingIndex", "Internal error. Invalid binding index given."), Index).ToString());
+			bIsBindingsValid = false;
+			continue;
+		}
 
-			SourceCreator.SetterGraph = UE::MVVM::FunctionGraphHelper::CreateIntermediateFunctionGraph(
-				WidgetBlueprintCompilerContext
-				, SourceCreator.SetterFunctionName
-				, (FUNC_BlueprintCallable | FUNC_Public)
-				, TEXT("Viewmodel")
-				, false);
-			BlueprintView->TemporaryGraph.Add(SourceCreator.SetterGraph);
+		FMVVMBlueprintViewBinding& Binding = *BindingPtr;
+		if (!Binding.bCompile)
+		{
+			continue;
+		}
 
-			if (SourceCreator.SetterGraph == nullptr || SourceCreator.SetterGraph->GetFName() != FName(*SourceCreator.SetterFunctionName))
+		auto AddGeneratedConversionFunction = [&Context, BindingPtr, this](bool bIsForwardBinding)
+		{
+			UMVVMBlueprintViewConversionFunction* ConversionFunction = BindingPtr->Conversion.GetConversionFunction(bIsForwardBinding);
+			if (ConversionFunction && ConversionFunction->IsWrapperGraphTransient())
 			{
-				WidgetBlueprintCompilerContext.MessageLog.Error(*FText::Format(LOCTEXT("SetterNameAlreadyExists", "The setter name {0} already exists and could not be autogenerated."), 
-					FText::FromString(SourceCreator.SetterFunctionName)
-					).ToString()
-				);
+				UEdGraph* WrapperGraph = ConversionFunction->GetOrCreateIntermediateWrapperGraph(WidgetBlueprintCompilerContext);
+				if (WrapperGraph)
+				{
+					bool bAlreadyContained = WidgetBlueprintCompilerContext.Blueprint->FunctionGraphs.Contains(WrapperGraph);
+					if (ensure(!bAlreadyContained))
+					{
+						Context.AddGeneratedFunctionGraph(WrapperGraph);
+					}
+				}
 			}
+		};
 
-			UE::MVVM::FunctionGraphHelper::AddFunctionArgument(SourceCreator.SetterGraph, SourceCreator.ViewModelContext.GetViewModelClass(), "Viewmodel");
+		if (IsForwardBinding(Binding.BindingType))
+		{
+			AddGeneratedConversionFunction(true);
+		}
+		if (IsBackwardBinding(Binding.BindingType))
+		{
+			AddGeneratedConversionFunction(false);
 		}
 	}
 }
@@ -703,7 +762,7 @@ void FMVVMViewBlueprintCompiler::CreateBindingDestinationContexts(UMVVMBlueprint
 
 			NewDestinationContext.bCanBeSetInNative = CanBeSetInNative(SkeletalGeneratedField.GetValue());
 
-			// If the destination can't be done cpp, we need to generate a BP function to set the value.
+			// If the destination can't be set in cpp, we need to generate a BP function to set the value.
 			if (!NewDestinationContext.bCanBeSetInNative)
 			{
 				const FProperty* SetterProperty = nullptr;
@@ -897,25 +956,30 @@ bool FMVVMViewBlueprintCompiler::PreCompileBindingSources(UWidgetBlueprintGenera
 		auto CreateSourcesForConversionFunction = [this, &Binding, BlueprintView, &CreateSourceContextForPropertyPath](bool bForwardBinding)
 		{
 			ECreateSourcesForConversionFunctionResult Result = ECreateSourcesForConversionFunctionResult::Continue;
-			TMap<FName, FMVVMBlueprintPropertyPath> ArgumentPaths = ConversionFunctionHelper::GetAllArgumentPropertyPaths(WidgetBlueprintCompilerContext.WidgetBlueprint(), Binding, bForwardBinding, true);
-			for (const TPair<FName, FMVVMBlueprintPropertyPath>& Pair : ArgumentPaths)
+			if (const UMVVMBlueprintViewConversionFunction* ConversionFunction = Binding.Conversion.GetConversionFunction(bForwardBinding))
 			{
-				if (Pair.Key.IsNone())
+				for (const FMVVMBlueprintPin& Pin : ConversionFunction->GetPins())
 				{
-					AddMessageForBinding(Binding, BlueprintView,
-						FText::Format(LOCTEXT("InvalidArgumentPathName", "The conversion function {0} has an invalid argument"), FText::FromString(Binding.GetDisplayNameString(WidgetBlueprintCompilerContext.WidgetBlueprint()))),
-						EBindingMessageType::Error
-					);
-					Result = ECreateSourcesForConversionFunctionResult::Failed;
-				}
-				else if (CreateSourceContextForPropertyPath(Pair.Value, bForwardBinding, Pair.Key))
-				{
-					Result = ECreateSourcesForConversionFunctionResult::Valid;
-				}
-				else
-				{
-					Result = ECreateSourcesForConversionFunctionResult::Failed;
-					break;
+					if (Pin.UsedPathAsValue())
+					{
+						if (Pin.GetName().IsNone())
+						{
+							AddMessageForBinding(Binding, BlueprintView,
+								FText::Format(LOCTEXT("InvalidArgumentPathName", "The conversion function {0} has an invalid argument"), FText::FromString(Binding.GetDisplayNameString(WidgetBlueprintCompilerContext.WidgetBlueprint()))),
+								EBindingMessageType::Error
+							);
+							Result = ECreateSourcesForConversionFunctionResult::Failed;
+						}
+						else if (CreateSourceContextForPropertyPath(Pin.GetPath(), bForwardBinding, Pin.GetName()))
+						{
+							Result = ECreateSourcesForConversionFunctionResult::Valid;
+						}
+						else
+						{
+							Result = ECreateSourcesForConversionFunctionResult::Failed;
+							break;
+						}
+					}
 				}
 			}
 			return Result;
@@ -1345,22 +1409,22 @@ bool FMVVMViewBlueprintCompiler::PreCompileBindings(UWidgetBlueprintGeneratedCla
 			}
 		}
 
-		FMemberReference ConversionFunctionReference = BindingSourceContext.bIsForwardBinding ? Binding.Conversion.SourceToDestinationFunction : Binding.Conversion.DestinationToSourceFunction;
-		FName ConversionFunctionWrapper = BindingSourceContext.bIsForwardBinding ? Binding.Conversion.SourceToDestinationWrapper : Binding.Conversion.DestinationToSourceWrapper;
-		if (!ConversionFunctionWrapper.IsNone())
-		{
-			ConversionFunctionReference.SetSelfMember(ConversionFunctionWrapper);
-		}
-
-		const UFunction* ConversionFunction = ConversionFunctionReference.ResolveMember<UFunction>(Class);
-		if (!ConversionFunctionWrapper.IsNone() && ConversionFunction == nullptr)
+		const UFunction* ConversionFunction = nullptr;
 		{
-			AddMessageForBinding(Binding, BlueprintView, FText::Format(LOCTEXT("ConversionFunctionNotFound", "The conversion function '{0}' could not be found."), 
-				FText::FromName(ConversionFunctionReference.GetMemberName())),
-				EBindingMessageType::Error
-			);
-			bIsBindingsValid = false;
-			continue;
+			const UMVVMBlueprintViewConversionFunction* ViewConversionFunction = Binding.Conversion.GetConversionFunction(BindingSourceContext.bIsForwardBinding);
+			if (ViewConversionFunction)
+			{
+				ConversionFunction = ViewConversionFunction->GetCompiledFunction(Class);
+				if (ConversionFunction == nullptr)
+				{
+					AddMessageForBinding(Binding, BlueprintView, FText::Format(LOCTEXT("ConversionFunctionNotFound", "The conversion function '{0}' could not be found."), 
+						FText::FromName(ViewConversionFunction->GetCompiledFunctionName())),
+						EBindingMessageType::Error
+					);
+					bIsBindingsValid = false;
+					continue;
+				}
+			}
 		}
 
 		TValueOrError<FCompiledBinding, FText> AddBindingResult = CreateCompiledBinding(Class, BindingSourceContext.PropertyPath, SetterPath, ConversionFunction, BindingSourceContext.bIsComplexBinding);
@@ -1927,6 +1991,7 @@ bool FMVVMViewBlueprintCompiler::CanBeSetInNative(TArrayView<const FMVVMConstFie
 		{
 			return true;
 		}
+
 		// If the BP is defined in BP and has Net flags or FieldNotify flag, then the VaraibleSet K2Node need to be used to generate the proper byte-code.
 		if (Variant.IsProperty())
 		{
diff --git a/Engine/Plugins/Runtime/ModelViewViewModel/Source/ModelViewViewModelBlueprint/Private/MVVMViewBlueprintCompiler.h b/Engine/Plugins/Runtime/ModelViewViewModel/Source/ModelViewViewModelBlueprint/Private/MVVMViewBlueprintCompiler.h
--- a/Engine/Plugins/Runtime/ModelViewViewModel/Source/ModelViewViewModelBlueprint/Private/MVVMViewBlueprintCompiler.h
+++ b/Engine/Plugins/Runtime/ModelViewViewModel/Source/ModelViewViewModelBlueprint/Private/MVVMViewBlueprintCompiler.h
@@ -66,6 +66,7 @@ private:
 	void CreateWidgetMap(const FWidgetBlueprintCompilerContext::FCreateVariableContext& Context, UMVVMBlueprintView* BlueprintView);
 	void CreateSourceLists(const FWidgetBlueprintCompilerContext::FCreateVariableContext& Context, UMVVMBlueprintView* BlueprintView);
 	void CreateFunctionsDeclaration(const FWidgetBlueprintCompilerContext::FCreateVariableContext& Context, UMVVMBlueprintView* BlueprintView);
+	void CreateIntermediateGraphFunctions(const FWidgetBlueprintCompilerContext::FCreateVariableContext& Context, UMVVMBlueprintView* BlueprintView);
 	void CreateBindingDestinationContexts(UMVVMBlueprintView* BlueprintView);
 
 	bool PreCompileBindingSources(UWidgetBlueprintGeneratedClass* Class, UMVVMBlueprintView* BlueprintView);
diff --git a/Engine/Plugins/Runtime/ModelViewViewModel/Source/ModelViewViewModelBlueprint/Private/MVVMWidgetBlueprintExtension_View.cpp b/Engine/Plugins/Runtime/ModelViewViewModel/Source/ModelViewViewModelBlueprint/Private/MVVMWidgetBlueprintExtension_View.cpp
--- a/Engine/Plugins/Runtime/ModelViewViewModel/Source/ModelViewViewModelBlueprint/Private/MVVMWidgetBlueprintExtension_View.cpp
+++ b/Engine/Plugins/Runtime/ModelViewViewModel/Source/ModelViewViewModelBlueprint/Private/MVVMWidgetBlueprintExtension_View.cpp
@@ -2,6 +2,7 @@
 
 #include "MVVMWidgetBlueprintExtension_View.h"
 #include "MVVMBlueprintView.h"
+#include "MVVMBlueprintViewConversionFunction.h"
 #include "MVVMViewBlueprintCompiler.h"
 #include "View/MVVMViewClass.h"
 
@@ -55,6 +56,17 @@ void UMVVMWidgetBlueprintExtension_View::HandlePreloadObjectsForCompilation(UBlu
 				UBlueprint::ForceLoad(AvailableViewModel.GetViewModelClass());
 			}
 		}
+		for (FMVVMBlueprintViewBinding& Binding : BlueprintView->GetBindings())
+		{
+			if (Binding.Conversion.DestinationToSourceConversion)
+			{
+				UBlueprint::ForceLoad(Binding.Conversion.DestinationToSourceConversion);
+			}
+			if (Binding.Conversion.SourceToDestinationConversion)
+			{
+				UBlueprint::ForceLoad(Binding.Conversion.SourceToDestinationConversion);
+			}
+		}
 	}
 }
 
diff --git a/Engine/Plugins/Runtime/ModelViewViewModel/Source/ModelViewViewModelBlueprint/Public/Bindings/MVVMConversionFunctionHelper.h b/Engine/Plugins/Runtime/ModelViewViewModel/Source/ModelViewViewModelBlueprint/Public/Bindings/MVVMConversionFunctionHelper.h
--- a/Engine/Plugins/Runtime/ModelViewViewModel/Source/ModelViewViewModelBlueprint/Public/Bindings/MVVMConversionFunctionHelper.h
+++ b/Engine/Plugins/Runtime/ModelViewViewModel/Source/ModelViewViewModelBlueprint/Public/Bindings/MVVMConversionFunctionHelper.h
@@ -5,36 +5,45 @@
 #include "Containers/Map.h"
 
 struct FMVVMBlueprintPropertyPath;
+struct FMVVMBlueprintViewBinding;
 
+class UBlueprint;
 class UClass;
 class UEdGraph;
+class UEdGraphPin;
+class UK2Node;
 class UK2Node_CallFunction;
+class FKismetCompilerContext;
 class UMVVMBlueprintView;
-class UWidgetBlueprint;
-struct FMVVMBlueprintViewBinding;
 
 namespace UE::MVVM::ConversionFunctionHelper
 {
+	/** Conversion function requires a wrapper. */
+	MODELVIEWVIEWMODELBLUEPRINT_API bool RequiresWrapper(const UFunction* Function);
+
 	/** Find all BlueprintPropertyPath used by the given binding. */
-	MODELVIEWVIEWMODELBLUEPRINT_API TMap<FName, FMVVMBlueprintPropertyPath> GetAllArgumentPropertyPaths(const UWidgetBlueprint* WidgetBlueprint, const FMVVMBlueprintViewBinding& Binding, bool bSourceToDestination, bool bSkipResolve);
+	MODELVIEWVIEWMODELBLUEPRINT_API TMap<FName, FMVVMBlueprintPropertyPath> GetAllArgumentPropertyPaths(const UBlueprint* WidgetBlueprint, const FMVVMBlueprintViewBinding& Binding, bool bSourceToDestination, bool bSkipResolve);
 
 	/** Find all BlueprintPropertyPath used by the given conversion function node. */
-	MODELVIEWVIEWMODELBLUEPRINT_API TMap<FName, FMVVMBlueprintPropertyPath> GetAllArgumentPropertyPaths(const UWidgetBlueprint* WidgetBlueprint, const UK2Node_CallFunction* FunctionNode, bool bSkipResolve);
+	MODELVIEWVIEWMODELBLUEPRINT_API TMap<FName, FMVVMBlueprintPropertyPath> GetAllArgumentPropertyPaths(const UBlueprint* WidgetBlueprint, const UK2Node_CallFunction* FunctionNode, bool bSkipResolve);
 	
 	/** Find the property path of a given argument in the conversion function. */
-	MODELVIEWVIEWMODELBLUEPRINT_API FMVVMBlueprintPropertyPath GetPropertyPathForArgument(const UWidgetBlueprint* WidgetBlueprint, const UK2Node_CallFunction* Function, FName ArgumentName, bool bSkipResolve);
+	MODELVIEWVIEWMODELBLUEPRINT_API FMVVMBlueprintPropertyPath GetPropertyPathForPin(const UBlueprint* WidgetBlueprint, const UEdGraphPin* Pin, bool bSkipResolve);
+	
+	/** Set the property path of a given argument in the conversion function. */
+	MODELVIEWVIEWMODELBLUEPRINT_API void SetPropertyPathForPin(const UBlueprint* Blueprint, const FMVVMBlueprintPropertyPath& Path, UEdGraphPin* Pin);
+	
+	/** Find the property path of a given argument in the conversion function. */
+	MODELVIEWVIEWMODELBLUEPRINT_API FMVVMBlueprintPropertyPath GetPropertyPathForArgument(const UBlueprint* WidgetBlueprint, const UK2Node_CallFunction* Function, FName ArgumentName, bool bSkipResolve);
 
 	/** Create the name of the conversion function wrapper this binding should have. */
 	MODELVIEWVIEWMODELBLUEPRINT_API FName CreateWrapperName(const FMVVMBlueprintViewBinding& Binding, bool bSourceToDestination);
 
-	/** Get an existing conversion graph for the corresponding binding if it exist. */
-	MODELVIEWVIEWMODELBLUEPRINT_API UEdGraph* GetGraph(const UWidgetBlueprint* WidgetBlueprint, const FMVVMBlueprintViewBinding& Binding, bool bSourceToDestination);
+	/** */
+	MODELVIEWVIEWMODELBLUEPRINT_API TPair<UEdGraph*, UK2Node*> CreateGraph(UBlueprint* WidgetBlueprint, FName GraphName, const UFunction* Function, bool bTransient);
 
 	/** Find the conversion function node from the given graph. */
-	MODELVIEWVIEWMODELBLUEPRINT_API UK2Node_CallFunction* GetFunctionNode(UEdGraph* Graph);
-
-	/** Mark the given function node as the conversion function node. */
-	MODELVIEWVIEWMODELBLUEPRINT_API void MarkAsConversionFunction(const UK2Node_CallFunction* FunctionNode, const FMVVMBlueprintViewBinding& Binding);
+	MODELVIEWVIEWMODELBLUEPRINT_API UK2Node* GetWrapperNode(UEdGraph* Graph);
 } //namespace
 
 #if UE_ENABLE_INCLUDE_ORDER_DEPRECATED_IN_5_2
diff --git a/Engine/Plugins/Runtime/ModelViewViewModel/Source/ModelViewViewModelBlueprint/Public/MVVMBlueprintPin.h b/Engine/Plugins/Runtime/ModelViewViewModel/Source/ModelViewViewModelBlueprint/Public/MVVMBlueprintPin.h
new file mode 100644
--- /dev/null
+++ b/Engine/Plugins/Runtime/ModelViewViewModel/Source/ModelViewViewModelBlueprint/Public/MVVMBlueprintPin.h
@@ -0,0 +1,70 @@
+// Copyright Epic Games, Inc. All Rights Reserved.
+
+#pragma once
+
+#include "MVVMPropertyPath.h"
+
+#include "MVVMBlueprintPin.generated.h"
+
+class UWidgetBlueprint;
+class UEdGraphPin;
+
+/**
+*
+*/
+USTRUCT()
+struct MODELVIEWVIEWMODELBLUEPRINT_API FMVVMBlueprintPin
+{
+	GENERATED_BODY()
+
+private:
+	UPROPERTY(VisibleAnywhere, Category = "Viewmodel")
+	FName PinName;
+
+	UPROPERTY(VisibleAnywhere, Category = "Viewmodel")
+	FMVVMBlueprintPropertyPath Path;
+
+	/** Default value for this pin (used if the pin has no connections), stored as a string */
+	UPROPERTY(VisibleAnywhere, Category = "Viewmodel")
+	FString DefaultString;
+
+	/** If the default value for this pin should be an FText, it is stored here. */
+	UPROPERTY(VisibleAnywhere, Category = "Viewmodel")
+	FText DefaultText;
+
+	/** If the default value for this pin should be an object, we store a pointer to it */
+	UPROPERTY(VisibleAnywhere, Category = "Viewmodel")
+	TObjectPtr<class UObject> DefaultObject;
+
+public:
+	FMVVMBlueprintPin() = default;
+
+	FName GetName() const
+	{
+		return PinName;
+	}
+
+	/** Are we using the path. */
+	bool UsedPathAsValue() const
+	{
+		return !Path.IsEmpty();
+	}
+
+	const FMVVMBlueprintPropertyPath& GetPath() const
+	{
+		return Path;
+	}
+
+	FString GetValueAsString(const UClass* SelfContext) const;
+
+	void SetDefaultValue(UObject* Value);
+	void SetDefaultValue(const FText& Value);
+	void SetDefaultValue(const FString& Value);
+	void SetPath(const FMVVMBlueprintPropertyPath& Value);
+
+	static FMVVMBlueprintPin CreateFromPin(const UBlueprint* WidgetBlueprint, const UEdGraphPin* Pin);
+	void CopyTo(const UBlueprint* WidgetBlueprint, UEdGraphPin* Pin) const;
+
+private:
+	void Reset();
+};
diff --git a/Engine/Plugins/Runtime/ModelViewViewModel/Source/ModelViewViewModelBlueprint/Public/MVVMBlueprintView.h b/Engine/Plugins/Runtime/ModelViewViewModel/Source/ModelViewViewModelBlueprint/Public/MVVMBlueprintView.h
--- a/Engine/Plugins/Runtime/ModelViewViewModel/Source/ModelViewViewModelBlueprint/Public/MVVMBlueprintView.h
+++ b/Engine/Plugins/Runtime/ModelViewViewModel/Source/ModelViewViewModelBlueprint/Public/MVVMBlueprintView.h
@@ -86,8 +86,8 @@ public:
 	void ResetBindingMessages();
 
 #if WITH_EDITOR
-//	virtual void PostEditUndo() override;
-
+	virtual void PostLoad() override;
+	virtual void PreSave(FObjectPreSaveContext Context) override;
 	virtual void PostEditChangeProperty(FPropertyChangedEvent& PropertyChangedEvent) override;
 	virtual void PostEditChangeChainProperty(FPropertyChangedChainEvent& PropertyChainEvent) override;
 
@@ -95,7 +95,6 @@ public:
 	void WidgetRenamed(FName OldObjectName, FName NewObjectName);
 #endif
 
-	virtual void PostLoad() override;
 	virtual void Serialize(FArchive& Ar) override;
 
 	DECLARE_EVENT(UMVVMBlueprintView, FOnBindingsUpdated);
diff --git a/Engine/Plugins/Runtime/ModelViewViewModel/Source/ModelViewViewModelBlueprint/Public/MVVMBlueprintViewBinding.h b/Engine/Plugins/Runtime/ModelViewViewModel/Source/ModelViewViewModelBlueprint/Public/MVVMBlueprintViewBinding.h
--- a/Engine/Plugins/Runtime/ModelViewViewModel/Source/ModelViewViewModelBlueprint/Public/MVVMBlueprintViewBinding.h
+++ b/Engine/Plugins/Runtime/ModelViewViewModel/Source/ModelViewViewModelBlueprint/Public/MVVMBlueprintViewBinding.h
@@ -10,6 +10,7 @@
 
 class UWidgetBlueprint;
 class UMVVMBlueprintView;
+class UMVVMBlueprintViewConversionFunction;
 
 /**
 *
@@ -19,19 +20,68 @@ struct MODELVIEWVIEWMODELBLUEPRINT_API FMVVMBlueprintViewConversionPath
 {
 	GENERATED_BODY()
 
-	/** The Conversion function when converting the value from the destination to the source. */
-	UPROPERTY(VisibleAnywhere, Category = "Viewmodel", AdvancedDisplay)
-	FMemberReference DestinationToSourceFunction;
+	UE_DEPRECATED(5.3, "DestinationToSourceFunction was moved to MVVMBlueprintViewConversionFunction.")
+	/**
+	 * The Conversion function when converting the value from the destination to the source.
+	 * @note if a wrapper is needed, this is the function that is wrapped.
+	 */
+	UPROPERTY()
+	FMemberReference DestinationToSourceFunction_DEPRECATED;
 	
+	UE_DEPRECATED(5.3, "DestinationToSourceWrapper was moved to MVVMBlueprintViewConversionFunction.")
+	/**
+	 * The name of the graph that contains the conversion function when converting the value from the destination to the source.
+	 * Valid when the function needs a wrapper and when the graph is generated on the fly.
+	 */
 	UPROPERTY()
-	FName DestinationToSourceWrapper;
+	FName DestinationToSourceWrapper_DEPRECATED;
 
-	/** The Conversion function when converting the value from the source to the destination. */
-	UPROPERTY(VisibleAnywhere, Category = "Viewmodel", AdvancedDisplay)
-	FMemberReference SourceToDestinationFunction;
+	UE_DEPRECATED(5.3, "SourceToDestinationFunction was moved to MVVMBlueprintViewConversionFunction.")
+	/**
+	 * The Conversion function when converting the value from the source to the destination.
+	 * @note if a wrapper is needed, this is the function that is wrapped.
+	 */
+	UPROPERTY()
+	FMemberReference SourceToDestinationFunction_DEPRECATED;
 
+	UE_DEPRECATED(5.3, "SourceToDestinationWrapper was moved to MVVMBlueprintViewConversionFunction.")
+	/**
+	 * The name of the graph that contains the conversion function when converting the value from the source to the destination.
+	 * When the function needs a wrapper. Valid also when the graph is generated on the fly.
+	 */
 	UPROPERTY()
-	FName SourceToDestinationWrapper;
+	FName SourceToDestinationWrapper_DEPRECATED;
+
+	/**
+	 * The graph that contains the conversion function when converting the value from the destination to the source.
+	 * When the function doesn't need a wrapper.
+	 */
+	UPROPERTY(Instanced)
+	TObjectPtr<UMVVMBlueprintViewConversionFunction> DestinationToSourceConversion;
+
+	/**
+	 * The graph that contains the conversion function when converting the value from the source to the destination.
+	 * When the function doesn't need a wrapper.
+	 */
+	UPROPERTY(Instanced)
+	TObjectPtr<UMVVMBlueprintViewConversionFunction> SourceToDestinationConversion;
+
+public:
+	UMVVMBlueprintViewConversionFunction* GetConversionFunction(bool bSourceToDestination) const
+	{
+		return bSourceToDestination ? SourceToDestinationConversion : DestinationToSourceConversion;
+	}
+
+	void GenerateWrapper(UBlueprint* Blueprint);
+	void SavePinValues(UBlueprint* Blueprint);
+	void DeprecateViewConversionFunction(UBlueprint* Blueprint);
+
+public:
+PRAGMA_DISABLE_DEPRECATION_WARNINGS
+	FMVVMBlueprintViewConversionPath() = default;
+	FMVVMBlueprintViewConversionPath(const FMVVMBlueprintViewConversionPath&) = default;
+	FMVVMBlueprintViewConversionPath& operator=(const FMVVMBlueprintViewConversionPath&) = default;
+PRAGMA_ENABLE_DEPRECATION_WARNINGS
 };
 
 /**
diff --git a/Engine/Plugins/Runtime/ModelViewViewModel/Source/ModelViewViewModelBlueprint/Public/MVVMBlueprintViewConversionFunction.h b/Engine/Plugins/Runtime/ModelViewViewModel/Source/ModelViewViewModelBlueprint/Public/MVVMBlueprintViewConversionFunction.h
new file mode 100644
--- /dev/null
+++ b/Engine/Plugins/Runtime/ModelViewViewModel/Source/ModelViewViewModelBlueprint/Public/MVVMBlueprintViewConversionFunction.h
@@ -0,0 +1,146 @@
+// Copyright Epic Games, Inc. All Rights Reserved.
+
+#pragma once
+
+#include "MVVMBlueprintPin.h"
+#include "MVVMBlueprintView.h"
+#include "Engine/MemberReference.h"
+
+#include "MVVMBlueprintViewConversionFunction.generated.h"
+
+class UK2Node;
+class UEdGraphPin;
+class UEdGraph;
+
+/**
+ *
+ */
+UCLASS()
+class MODELVIEWVIEWMODELBLUEPRINT_API UMVVMBlueprintViewConversionFunction : public UObject
+{
+	GENERATED_BODY()
+
+public:
+	/** @return the conversion function uses at runtime. The wrapper function if complex or GetFunction is simple. */
+	const UFunction* GetCompiledFunction(const UClass* SelfContext) const;
+
+	/** @return the conversion function uses at runtime. The wrapper function if complex or GetFunction is simple. */
+	FName GetCompiledFunctionName() const;
+
+	/** @return the conversion function. */
+	TVariant<const UFunction*, TSubclassOf<UK2Node>> GetConversionFunction(const UClass* SelfContext) const;
+
+	/** Set the function. Generate a Graph is the conversion function is complex. */
+	void InitFromFunction(UBlueprint* SelfContext, const UFunction* Function);
+
+	/** Set the function. Generate a Graph is the conversion function is complex. */
+	void InitFromFunction(UBlueprint* SelfContext, const UFunction* Function, FName GraphName);
+
+	// For deprecation
+	void InitializeFromWrapperGraph(UBlueprint* SelfContext, UEdGraph* Graph);
+
+	// For deprecation
+	void InitializeFromMemberReference(UBlueprint* SelfContext, FMemberReference MemberReference);
+
+	/** The function has more than one argument and requires a wrapper or it uses a FunctionNode. */
+	bool NeedsWrapperGraph() const
+	{
+		return !GraphName.IsNone();
+	}
+
+	/** The wrapper Graph is generated on domains and is not saved. */
+	bool IsWrapperGraphTransient() const
+	{
+		return NeedsWrapperGraph()  && bWrapperGraphTransient;
+	}
+
+	/** Return the wrapper graph, if it exists. */
+	UEdGraph* GetWrapperGraph() const
+	{
+		return CachedWrapperGraph;
+	}
+
+	FName GetWrapperGraphName() const
+	{
+		return GraphName;
+	}
+
+	UK2Node* GetWrapperNode() const
+	{
+		return CachedWrapperNode;
+	}
+
+	/**
+	 * If needed, create the graph and all the nodes for that graph when compiling.
+	 * Returns the existing one, if one was created from GetOrCreateWrapperGraph.
+	 */
+	UEdGraph* GetOrCreateIntermediateWrapperGraph(FKismetCompilerContext& Context) const;
+
+	/** If needed, create the graph and all the nodes for that graph. */
+	UEdGraph* GetOrCreateWrapperGraph(UBlueprint* Blueprint) const;
+
+	/**
+	 * The conversion function is going to be removed from the Blueprint.
+	 * Do any cleanup that is needed.
+	 */
+	void RemoveWrapperGraph(UBlueprint* Blueprint);
+
+	/**
+	 * Returns the pin from the graph.
+	 * Create the graph and all the nodes for that graph if the graph doesn't exist and it's needed.
+	 */
+	UEdGraphPin* GetOrCreateGraphPin(UBlueprint* Blueprint, FName PinName) const;
+
+	TArrayView<const FMVVMBlueprintPin> GetPins() const
+	{
+		return SavedPins;
+	}
+
+	/** */
+	void SetGraphPin(UBlueprint* Blueprint, FName PinName, const FMVVMBlueprintPropertyPath& Path);
+
+	/** Generates SavedPins from the wrapper graph, if it exists. */
+	void SavePinValues(UBlueprint* Blueprint);
+
+private:
+	UEdGraph* GetOrCreateWrapperGraphInternal(FKismetCompilerContext& Context, const UFunction* Function) const;
+	UEdGraph* GetOrCreateWrapperGraphInternal(UBlueprint* Blueprint, const UFunction* Function) const;
+	void LoadPinValuesInternal(UBlueprint* Blueprint) const;
+	void Reset();
+
+private:
+	/**
+	 * The conversion UFunction when simple or when it's complex.
+	 * @note Only one of FunctionReference or the GraphNode can be valid.
+	 */
+	UPROPERTY(VisibleAnywhere, Category = "Viewmodel")
+	FMemberReference FunctionReference;
+	
+	/**
+	 * The conversion K2Node the graph is generated for.
+	 * @note Only one of FunctionReference or the GraphNode can be valid.
+	 * @note The conversion is complex.
+	 */
+	UPROPERTY(VisibleAnywhere, Category = "Viewmodel")
+	TSubclassOf<UK2Node> FunctionNode;
+
+	/** Name of the generated graph if a wrapper is needed. */
+	UPROPERTY(VisibleAnywhere, Category = "Viewmodel")
+	FName GraphName;
+
+	/**
+	 * The pin that are modified and we saved data.
+	 * The data may not be modified. We used the default value of the K2Node in that case.
+	 */
+	UPROPERTY(VisibleAnywhere, Category = "Viewmodel")
+	TArray<FMVVMBlueprintPin> SavedPins;
+
+	UPROPERTY(VisibleAnywhere, Category = "Viewmodel")
+	bool bWrapperGraphTransient = false;
+
+	UPROPERTY(Transient, DuplicateTransient)
+	mutable TObjectPtr<UEdGraph> CachedWrapperGraph;
+	
+	UPROPERTY(Transient, DuplicateTransient)
+	mutable TObjectPtr<UK2Node> CachedWrapperNode;
+};
diff --git a/Engine/Plugins/Runtime/ModelViewViewModel/Source/ModelViewViewModelBlueprint/Public/MVVMDeveloperProjectSettings.h b/Engine/Plugins/Runtime/ModelViewViewModel/Source/ModelViewViewModelBlueprint/Public/MVVMDeveloperProjectSettings.h
--- a/Engine/Plugins/Runtime/ModelViewViewModel/Source/ModelViewViewModelBlueprint/Public/MVVMDeveloperProjectSettings.h
+++ b/Engine/Plugins/Runtime/ModelViewViewModel/Source/ModelViewViewModelBlueprint/Public/MVVMDeveloperProjectSettings.h
@@ -120,6 +120,10 @@ public:
 	UPROPERTY(EditAnywhere, config, Category = "Viewmodel")
 	bool bShowDetailViewOptionInBindingPanel = true;
 
+	/** When a conversion function requires a wrapper graph, add and save the generated graph to the blueprint. */
+	UPROPERTY(EditAnywhere, config, Category = "Viewmodel")
+	bool bAllowConversionFunctionGeneratedGraphInEditor = true;
+
 	/** Permission list for filtering which execution mode is allowed. */
 	UPROPERTY(EditAnywhere, config, Category = "Viewmodel")
 	EMVVMDeveloperConversionFunctionFilterType ConversionFunctionFilter = EMVVMDeveloperConversionFunctionFilterType::BlueprintActionRegistry;
diff --git a/Engine/Plugins/Runtime/ModelViewViewModel/Source/ModelViewViewModelEditor/Private/MVVMEditorSubsystem.cpp b/Engine/Plugins/Runtime/ModelViewViewModel/Source/ModelViewViewModelEditor/Private/MVVMEditorSubsystem.cpp
--- a/Engine/Plugins/Runtime/ModelViewViewModel/Source/ModelViewViewModelEditor/Private/MVVMEditorSubsystem.cpp
+++ b/Engine/Plugins/Runtime/ModelViewViewModel/Source/ModelViewViewModelEditor/Private/MVVMEditorSubsystem.cpp
@@ -17,6 +17,7 @@
 #include "K2Node_FunctionResult.h"
 #include "K2Node_VariableGet.h"
 #include "MVVMBlueprintView.h"
+#include "MVVMBlueprintViewConversionFunction.h"
 #include "MVVMDeveloperProjectSettings.h"
 #include "MVVMSubsystem.h"
 #include "MVVMWidgetBlueprintExtension_View.h"
@@ -58,53 +59,6 @@ namespace UE::MVVM::Private
 		BlueprintView->PostEditChangeChainProperty(ChainEvent);
 	}
 
-	UEdGraph* CreateFunctionGraph(UBlueprint* Blueprint, FName InFunctionName, EFunctionFlags ExtraFunctionFlag, const FStringView Category)
-	{
-		bool bIsEditable = false;
-
-		FName UniqueFunctionName = FBlueprintEditorUtils::FindUniqueKismetName(Blueprint, InFunctionName.ToString());
-		UEdGraph* FunctionGraph = FBlueprintEditorUtils::CreateNewGraph(Blueprint, UniqueFunctionName, UEdGraph::StaticClass(), UEdGraphSchema_K2::StaticClass());
-		Blueprint->FunctionGraphs.Add(FunctionGraph);
-		FunctionGraph->bEditable = bIsEditable;
-
-		const UEdGraphSchema_K2* Schema = CastChecked<UEdGraphSchema_K2>(FunctionGraph->GetSchema());
-		Schema->MarkFunctionEntryAsEditable(FunctionGraph, bIsEditable);
-		Schema->CreateDefaultNodesForGraph(*FunctionGraph);
-
-		// function entry node
-		FGraphNodeCreator<UK2Node_FunctionEntry> FunctionEntryCreator(*FunctionGraph);
-		UK2Node_FunctionEntry* FunctionEntry = FunctionEntryCreator.CreateNode();
-		FunctionEntry->FunctionReference.SetSelfMember(FunctionGraph->GetFName());
-		FunctionEntry->AddExtraFlags(ExtraFunctionFlag);
-		FunctionEntry->bIsEditable = bIsEditable;
-		FunctionEntry->MetaData.Category = FText::FromStringView(Category);
-		FunctionEntry->NodePosX = -500;
-		FunctionEntry->NodePosY = 0;
-		FunctionEntryCreator.Finalize();
-
-		FGraphNodeCreator<UK2Node_FunctionResult> FunctionResultCreator(*FunctionGraph);
-		UK2Node_FunctionResult* FunctionResult = FunctionResultCreator.CreateNode();
-		FunctionResult->FunctionReference.SetSelfMember(FunctionGraph->GetFName());
-		FunctionResult->bIsEditable = bIsEditable;
-		FunctionResult->NodePosX = 500;
-		FunctionResult->NodePosY = 0;
-		FunctionResultCreator.Finalize();
-
-		return FunctionGraph;
-	}
-
-	UK2Node_FunctionEntry* FindFunctionEntry(UEdGraph* Graph)
-	{
-		for (UEdGraphNode* Node : Graph->Nodes)
-		{
-			if (UK2Node_FunctionEntry* FunctionEntry = Cast<UK2Node_FunctionEntry>(Node))
-			{
-				return FunctionEntry;
-			}
-		}
-		return nullptr;
-	}
-
 	UK2Node_FunctionResult* FindFunctionResult(UEdGraph* Graph)
 	{
 		for (UEdGraphNode* Node : Graph->Nodes)
@@ -146,20 +100,6 @@ namespace UE::MVVM::Private
 		return ConversionNode;
 	}
 
-	UEdGraph* FindExistingConversionFunctionWrapper(UWidgetBlueprint* WidgetBlueprint, FName WrapperName)
-	{
-		TObjectPtr<UEdGraph>* Result = WidgetBlueprint->FunctionGraphs.FindByPredicate([WrapperName](const UEdGraph* GraphPtr) { return GraphPtr->GetFName() == WrapperName; });
-		return Result ? Result->Get() : nullptr;
-	}
-
-	void RemoveConversionFunctionWrapper(UWidgetBlueprint* WidgetBlueprint, FName WrapperName)
-	{
-		if (UEdGraph* Graph = FindExistingConversionFunctionWrapper(WidgetBlueprint, WrapperName))
-		{
-			FBlueprintEditorUtils::RemoveGraph(WidgetBlueprint, Graph);
-		}
-	}
-
 } //namespace Private
 
 
@@ -284,116 +224,24 @@ void UMVVMEditorSubsystem::RemoveBinding(UWidgetBlueprint* WidgetBlueprint, cons
 	}
 }
 
-UEdGraph* UMVVMEditorSubsystem::GetConversionFunctionGraph(const UWidgetBlueprint* WidgetBlueprint, const FMVVMBlueprintViewBinding& Binding, bool bSourceToDestination) const
-{
-	return UE::MVVM::ConversionFunctionHelper::GetGraph(WidgetBlueprint, Binding, bSourceToDestination);
-}
-
-UK2Node_CallFunction* UMVVMEditorSubsystem::GetConversionFunctionNode(const UWidgetBlueprint* WidgetBlueprint, const FMVVMBlueprintViewBinding& Binding, bool bSourceToDestination) const
-{
-	if (UEdGraph* Graph = UE::MVVM::ConversionFunctionHelper::GetGraph(WidgetBlueprint, Binding, bSourceToDestination))
-	{
-		return UE::MVVM::ConversionFunctionHelper::GetFunctionNode(Graph);
-	}
-	return nullptr;
-}
-
 UFunction* UMVVMEditorSubsystem::GetConversionFunction(const UWidgetBlueprint* WidgetBlueprint, const FMVVMBlueprintViewBinding& Binding, bool bSourceToDestination) const
 {
-	UK2Node_CallFunction* CallFunctionNode = GetConversionFunctionNode(WidgetBlueprint, Binding, bSourceToDestination);
-	if (CallFunctionNode)
+	if (UMVVMBlueprintViewConversionFunction* ConversionFunction = Binding.Conversion.GetConversionFunction(bSourceToDestination))
 	{
-		return CallFunctionNode->GetTargetFunction();
+		TVariant<const UFunction*, TSubclassOf<UK2Node>> Result = ConversionFunction->GetConversionFunction(WidgetBlueprint->SkeletonGeneratedClass);
+		if (Result.IsType<const UFunction*>())
+		{
+			return const_cast<UFunction*>(Result.Get<const UFunction*>());
+		}
 	}
-	
-	// simple conversion functions don't have a wrapper graph
-	FMemberReference ConversionFunctionReference = bSourceToDestination ? Binding.Conversion.SourceToDestinationFunction : Binding.Conversion.DestinationToSourceFunction;
-	return ConversionFunctionReference.ResolveMember<UFunction>(WidgetBlueprint->SkeletonGeneratedClass);
+	return nullptr;
 }
 
 UEdGraphPin* UMVVMEditorSubsystem::GetConversionFunctionArgumentPin(const UWidgetBlueprint* WidgetBlueprint, const FMVVMBlueprintViewBinding& Binding, FName ParameterName, bool bSourceToDestination)
 {
-	UK2Node_CallFunction* CallFunctionNode = GetConversionFunctionNode(WidgetBlueprint, Binding, bSourceToDestination);
-	if (CallFunctionNode == nullptr)
+	if (UMVVMBlueprintViewConversionFunction* ConversionFunction = Binding.Conversion.GetConversionFunction(bSourceToDestination))
 	{
-		return nullptr;
-	}
-
-	UEdGraphPin* Pin = CallFunctionNode->FindPin(ParameterName, EGPD_Input);
-	return Pin;
-}
-
-
-UEdGraph* UMVVMEditorSubsystem::CreateConversionFunctionWrapperGraph(UWidgetBlueprint* WidgetBlueprint, const FMVVMBlueprintViewBinding& Binding, const UFunction* ConversionFunction, bool bSourceToDestination)
-{
-	if (UMVVMBlueprintView* View = GetView(WidgetBlueprint))
-	{
-		TValueOrError<TArray<const FProperty*>, FText> ArgumentsResult = UE::MVVM::BindingHelper::TryGetArgumentsForConversionFunction(ConversionFunction);
-		if (ArgumentsResult.HasValue())
-		{
-			TArray<const FProperty*> Arguments = ArgumentsResult.StealValue();
-			if (Arguments.Num() > 1)
-			{
-				const FName WrapperName = UE::MVVM::ConversionFunctionHelper::CreateWrapperName(Binding, bSourceToDestination);
-
-				UEdGraph* NewGraph = UE::MVVM::Private::CreateFunctionGraph(
-					WidgetBlueprint,
-					WrapperName,
-					FUNC_BlueprintCallable | FUNC_BlueprintPure | FUNC_Const,
-					TEXT("Conversion Functions")
-				);
-
-				UK2Node_FunctionEntry* FunctionEntry = UE::MVVM::Private::FindFunctionEntry(NewGraph);
-				UK2Node_FunctionResult* FunctionResult = UE::MVVM::Private::FindFunctionResult(NewGraph);
-
-				FGraphNodeCreator<UK2Node_CallFunction> CallFunctionCreator(*NewGraph);
-				UK2Node_CallFunction* CallFunctionNode = CallFunctionCreator.CreateNode();
-				CallFunctionNode->SetFromFunction(ConversionFunction);
-				CallFunctionNode->NodePosX = 0;
-				CallFunctionCreator.Finalize();
-
-				UE::MVVM::ConversionFunctionHelper::MarkAsConversionFunction(CallFunctionNode, Binding);
-
-				UEdGraphPin* FunctionEntryThenPin = FunctionEntry->FindPin(UEdGraphSchema_K2::PN_Then, EGPD_Output);
-				UEdGraphPin* FunctionResultExecPin = FunctionResult->GetExecPin();
-			
-				// hook up exec pins 
-				if (!CallFunctionNode->IsNodePure())
-				{
-					UEdGraphPin* CallFunctionExecPin = CallFunctionNode->GetExecPin();
-					UEdGraphPin* CallFunctionThenPin = CallFunctionNode->FindPin(UEdGraphSchema_K2::PN_Then, EGPD_Output);
-
-					FunctionEntryThenPin->MakeLinkTo(CallFunctionExecPin);
-					CallFunctionThenPin->MakeLinkTo(FunctionResultExecPin);
-
-					CallFunctionNode->NodePosY = 0;
-				}
-				else
-				{
-					FunctionEntryThenPin->MakeLinkTo(FunctionResultExecPin);
-
-					CallFunctionNode->NodePosY = 100;
-				}
-
-				const UEdGraphSchema_K2* K2Schema = GetDefault<UEdGraphSchema_K2>();
-
-				// create return value pin
-				const FProperty* ReturnProperty = UE::MVVM::BindingHelper::GetReturnProperty(ConversionFunction);
-				check(ReturnProperty);
-				TSharedPtr<FUserPinInfo> PinInfo = MakeShared<FUserPinInfo>();
-				K2Schema->ConvertPropertyToPinType(ReturnProperty, PinInfo->PinType);
-				PinInfo->PinName = ReturnProperty->GetFName();
-				PinInfo->DesiredPinDirection = EGPD_Input;
-				FunctionResult->UserDefinedPins.Add(PinInfo);
-				FunctionResult->ReconstructNode();
-
-				UEdGraphPin* FunctionReturnPin = CallFunctionNode->FindPin(ReturnProperty->GetName(), EGPD_Output);
-				UEdGraphPin* OutputPin = FunctionResult->FindPin(ReturnProperty->GetName());
-				FunctionReturnPin->MakeLinkTo(OutputPin);
-
-				return NewGraph;
-			}
-		}
+		return ConversionFunction->GetOrCreateGraphPin(const_cast<UWidgetBlueprint*>(WidgetBlueprint), ParameterName);
 	}
 	return nullptr;
 }
@@ -410,21 +258,17 @@ void UMVVMEditorSubsystem::SetSourceToDestinationConversionFunction(UWidgetBluep
 
 			UE::MVVM::Private::OnBindingPreEditChange(View, GET_MEMBER_NAME_CHECKED(FMVVMBlueprintViewBinding, Conversion));
 
-			if (Binding.Conversion.SourceToDestinationWrapper != FName())
+			if (Binding.Conversion.SourceToDestinationConversion)
 			{
-				UE::MVVM::Private::RemoveConversionFunctionWrapper(WidgetBlueprint, Binding.Conversion.SourceToDestinationWrapper);
-				Binding.Conversion.SourceToDestinationWrapper = FName();
+				Binding.Conversion.SourceToDestinationConversion->RemoveWrapperGraph(WidgetBlueprint);
+				Binding.Conversion.SourceToDestinationConversion = nullptr;
 			}
 
 			if (ConversionFunction != nullptr)
 			{
-				UEdGraph* WrapperGraph = CreateConversionFunctionWrapperGraph(WidgetBlueprint, Binding, ConversionFunction, true);
-				Binding.Conversion.SourceToDestinationFunction.SetFromField<UFunction>(ConversionFunction, WidgetBlueprint->SkeletonGeneratedClass);
-				Binding.Conversion.SourceToDestinationWrapper = WrapperGraph != nullptr ? WrapperGraph->GetFName() : FName();
-			}
-			else
-			{
-				Binding.Conversion.SourceToDestinationFunction = FMemberReference();
+				Binding.Conversion.SourceToDestinationConversion = NewObject<UMVVMBlueprintViewConversionFunction>(WidgetBlueprint);
+				FName GraphName = UE::MVVM::ConversionFunctionHelper::CreateWrapperName(Binding, true);
+				Binding.Conversion.SourceToDestinationConversion->InitFromFunction(WidgetBlueprint, ConversionFunction, GraphName);
 			}
 
 			UE::MVVM::Private::OnBindingPostEditChange(View, GET_MEMBER_NAME_CHECKED(FMVVMBlueprintViewBinding, Conversion));
@@ -444,21 +288,17 @@ void UMVVMEditorSubsystem::SetDestinationToSourceConversionFunction(UWidgetBluep
 
 			UE::MVVM::Private::OnBindingPreEditChange(View, GET_MEMBER_NAME_CHECKED(FMVVMBlueprintViewBinding, Conversion));
 
-			if (Binding.Conversion.DestinationToSourceWrapper != FName())
+			if (Binding.Conversion.DestinationToSourceConversion)
 			{
-				UE::MVVM::Private::RemoveConversionFunctionWrapper(WidgetBlueprint, Binding.Conversion.DestinationToSourceWrapper);
-				Binding.Conversion.DestinationToSourceWrapper = FName();
+				Binding.Conversion.DestinationToSourceConversion->RemoveWrapperGraph(WidgetBlueprint);
+				Binding.Conversion.DestinationToSourceConversion = nullptr;
 			}
 
 			if (ConversionFunction != nullptr)
 			{
-				UEdGraph* WrapperGraph = CreateConversionFunctionWrapperGraph(WidgetBlueprint, Binding, ConversionFunction, true);
-				Binding.Conversion.DestinationToSourceFunction.SetFromField<UFunction>(ConversionFunction, WidgetBlueprint->SkeletonGeneratedClass);
-				Binding.Conversion.DestinationToSourceWrapper = WrapperGraph != nullptr ? WrapperGraph->GetFName() : FName();
-			}
-			else
-			{
-				Binding.Conversion.DestinationToSourceFunction = FMemberReference();
+				Binding.Conversion.DestinationToSourceConversion = NewObject<UMVVMBlueprintViewConversionFunction>(WidgetBlueprint);
+				FName GraphName = UE::MVVM::ConversionFunctionHelper::CreateWrapperName(Binding, false);
+				Binding.Conversion.DestinationToSourceConversion->InitFromFunction(WidgetBlueprint, ConversionFunction, GraphName);
 			}
 
 			UE::MVVM::Private::OnBindingPostEditChange(View, GET_MEMBER_NAME_CHECKED(FMVVMBlueprintViewBinding, Conversion));
@@ -673,6 +513,24 @@ bool UMVVMEditorSubsystem::IsSimpleConversionFunction(const UFunction* Function)
 	return ArgumentsResult.GetValue().Num() == 1;
 }
 
+UEdGraph* UMVVMEditorSubsystem::GetConversionFunctionGraph(const UWidgetBlueprint* WidgetBlueprint, const FMVVMBlueprintViewBinding& Binding, bool bSourceToDestination) const
+{
+	if (const UMVVMBlueprintViewConversionFunction* Found = Binding.Conversion.GetConversionFunction(bSourceToDestination))
+	{
+		return Found->GetWrapperGraph();
+	}
+	return nullptr;
+}
+
+UK2Node_CallFunction* UMVVMEditorSubsystem::GetConversionFunctionNode(const UWidgetBlueprint* WidgetBlueprint, const FMVVMBlueprintViewBinding& Binding, bool bSourceToDestination) const
+{
+	if (const UMVVMBlueprintViewConversionFunction* Found = Binding.Conversion.GetConversionFunction(bSourceToDestination))
+	{
+		return Cast<UK2Node_CallFunction>(Found->GetWrapperNode());
+	}
+	return nullptr;
+}
+
 TArray<UFunction*> UMVVMEditorSubsystem::GetAvailableConversionFunctions(const UWidgetBlueprint* WidgetBlueprint, const FMVVMBlueprintPropertyPath& Source, const FMVVMBlueprintPropertyPath& Destination) const
 {
 	TArray<UFunction*> ConversionFunctions;
@@ -763,9 +621,10 @@ TArray<UFunction*> UMVVMEditorSubsystem::GetAvailableConversionFunctions(const U
 
 FMVVMBlueprintPropertyPath UMVVMEditorSubsystem::GetPathForConversionFunctionArgument(const UWidgetBlueprint* WidgetBlueprint, const FMVVMBlueprintViewBinding& Binding, FName ArgumentName, bool bSourceToDestination) const
 {
-	UEdGraph* WrapperGraph = GetConversionFunctionGraph(WidgetBlueprint, Binding, bSourceToDestination);
-	if (WrapperGraph == nullptr)
+	UMVVMBlueprintViewConversionFunction* ConversionFunction = Binding.Conversion.GetConversionFunction(bSourceToDestination);
+	if (ConversionFunction == nullptr || !ConversionFunction->NeedsWrapperGraph())
 	{
+		// simple conversion function
 		if (bSourceToDestination)
 		{
 			return Binding.SourcePath;
@@ -776,19 +635,19 @@ FMVVMBlueprintPropertyPath UMVVMEditorSubsystem::GetPathForConversionFunctionArg
 		}
 	}
 
-	const UK2Node_CallFunction* FunctionNode = GetConversionFunctionNode(WidgetBlueprint, Binding, bSourceToDestination);
-	if (FunctionNode == nullptr)
+	UEdGraphPin* GraphPin = ConversionFunction->GetOrCreateGraphPin(const_cast<UWidgetBlueprint*>(WidgetBlueprint), ArgumentName);
+	if (GraphPin == nullptr)
 	{
 		return FMVVMBlueprintPropertyPath();
 	}
 
-	return UE::MVVM::ConversionFunctionHelper::GetPropertyPathForArgument(WidgetBlueprint, FunctionNode, ArgumentName, false);
+	return UE::MVVM::ConversionFunctionHelper::GetPropertyPathForPin(WidgetBlueprint, GraphPin, false);
 }
 
 void UMVVMEditorSubsystem::SetPathForConversionFunctionArgument(UWidgetBlueprint* WidgetBlueprint, FMVVMBlueprintViewBinding& Binding, FName ArgumentName, const FMVVMBlueprintPropertyPath& Path, bool bSourceToDestination) const
 {
-	UEdGraph* ConversionFunctionGraph = GetConversionFunctionGraph(WidgetBlueprint, Binding, bSourceToDestination);
-	if (ConversionFunctionGraph == nullptr)
+	UMVVMBlueprintViewConversionFunction* ConversionFunction = Binding.Conversion.GetConversionFunction(bSourceToDestination);
+	if (ConversionFunction == nullptr || !ConversionFunction->NeedsWrapperGraph())
 	{
 		// simple conversion function
 		if (bSourceToDestination)
@@ -799,241 +658,12 @@ void UMVVMEditorSubsystem::SetPathForConversionFunctionArgument(UWidgetBlueprint
 		{
 			Binding.DestinationPath = Path;
 		}
-
-		return;
-	}
-
-	UEdGraphNode* ConversionNode = UE::MVVM::Private::FindConversionNode(ConversionFunctionGraph);
-	if (!ensure(ConversionNode != nullptr))
-	{
-		return;
-	}
-
-	const UEdGraphSchema_K2* Schema = CastChecked<UEdGraphSchema_K2>(ConversionFunctionGraph->GetSchema());
-
-	// find the argument pin with the name ArgumentName
-	UEdGraphPin* ArgumentPin = nullptr;
-	int32 ArgumentIndex = 0;
-	for (UEdGraphPin* Pin : ConversionNode->Pins)
-	{
-		if (Pin->Direction == EGPD_Input && Pin->PinName != UEdGraphSchema_K2::PN_Execute)
-		{
-			if (Pin->PinName == ArgumentName)
-			{
-				ArgumentPin = Pin;
-				break;
-			}
-
-			++ArgumentIndex;
-		}
-	}
-
-	if (!ensure(ArgumentPin != nullptr))
-	{
-		return;
-	}
-
-	// delete all linked nodes
-	TArray<UEdGraphPin*> PinsToClear;
-	PinsToClear.Add(ArgumentPin);
-
-	TArray<UEdGraphNode*> NodesToRemove;
-
-	while (PinsToClear.Num() > 0)
-	{
-		UEdGraphPin* Pin = PinsToClear[0];
-		PinsToClear.RemoveAt(0);
-
-		PinsToClear.Append(Pin->LinkedTo);
-		Pin->BreakAllPinLinks();
-
-		// rename the node out of the way
-		UEdGraphNode* Node = Pin->GetOwningNode();
-		if (Node != ConversionNode)
-		{
-			NodesToRemove.AddUnique(Node);
-		}
-	}
-
-	for (UEdGraphNode* Node : NodesToRemove)
-	{
-		ConversionFunctionGraph->RemoveNode(Node, true);
-	}
-
-	if (Path.IsEmpty())
-	{
-		// we were given an empty path and we already cleaned up the path 
-		return;
-	}
-
-	TArray<UE::MVVM::FMVVMConstFieldVariant> Fields = Path.GetFields(WidgetBlueprint->SkeletonGeneratedClass);
-
-	float PosX = ConversionNode->NodePosX - 300 * (Fields.Num() + 1);
-	const float PosY = ConversionNode->NodePosY + ArgumentIndex * 100;
-
-	UEdGraphPin* PreviousExecPin = UE::MVVM::Private::FindFunctionEntry(ConversionFunctionGraph)->FindPin(UEdGraphSchema_K2::PN_Then);
-	UEdGraphPin* PreviousDataPin = nullptr;
-	UClass* PreviousClass = nullptr;
-
-	// create the root property getter node, ie. the Widget/ViewModel
-	const FProperty* RootProperty = nullptr;
-	if (Path.IsFromWidget())
-	{
-		RootProperty = WidgetBlueprint->SkeletonGeneratedClass->FindPropertyByName(Path.GetWidgetName());
-	}
-	else if (Path.IsFromViewModel())
-	{
-		UMVVMBlueprintView* View = GetView(WidgetBlueprint);
-		const FMVVMBlueprintViewModelContext* Context = View->FindViewModel(Path.GetViewModelId());
-
-		RootProperty = WidgetBlueprint->SkeletonGeneratedClass->FindPropertyByName(Context->GetViewModelName());
-	}
-
-	if (RootProperty != nullptr)
-	{
-		FGraphNodeCreator<UK2Node_VariableGet> RootGetterCreator(*ConversionFunctionGraph);
-		UK2Node_VariableGet* RootGetterNode = RootGetterCreator.CreateNode();
-		RootGetterNode->NodePosX = PosX;
-		RootGetterNode->NodePosY = PosY;
-		RootGetterNode->VariableReference.SetFromField<FProperty>(RootProperty, true, WidgetBlueprint->SkeletonGeneratedClass);
-		RootGetterCreator.Finalize();
-
-		PreviousDataPin = RootGetterNode->Pins[0];
-
-		PreviousClass = CastField<FObjectProperty>(RootProperty)->PropertyClass;
-
-		PosX += 300;
 	}
 	else
 	{
-		ensureMsgf(false, TEXT("Could not resolve root property!"));
-		return;
-	}
-
-	// create all the subsequent nodes in the path
-	for (int32 Index = 0; Index < Fields.Num(); ++Index)
-	{
-		if (!ensureMsgf(PreviousClass != nullptr, TEXT("Previous class not set!")))
-		{
-			return;
-		}
-
-		UEdGraphNode* NewNode = nullptr;
-		const UE::MVVM::FMVVMConstFieldVariant& Field = Fields[Index];
-		const bool bLastField = (Index == Fields.Num() - 1);
-
-		if (Field.IsProperty())
-		{
-			const FProperty* Property = Field.GetProperty();
-
-			// for structs in the middle of a path, we need to use a break node
-			const FStructProperty* StructProperty = CastField<FStructProperty>(Property);
-			if (!bLastField && StructProperty != nullptr)
-			{
-				FGraphNodeCreator<UK2Node_BreakStruct> BreakCreator(*ConversionFunctionGraph);
-				UK2Node_BreakStruct* BreakNode = BreakCreator.CreateNode();
-				BreakNode->AllocateDefaultPins();
-				BreakCreator.Finalize();
-
-				NewNode = BreakNode;
-
-				PreviousClass = nullptr;
-			}
-			else if (PreviousClass != nullptr)
-			{
-				FGraphNodeCreator<UK2Node_VariableGet> GetterCreator(*ConversionFunctionGraph);
-				UK2Node_VariableGet* GetterNode = GetterCreator.CreateNode();
-				GetterNode->SetFromProperty(Property, false, PreviousClass);
-				GetterNode->AllocateDefaultPins();
-				GetterCreator.Finalize();
-
-				NewNode = GetterNode;
-
-				if (const FObjectProperty* ObjectProperty = CastField<FObjectProperty>(Property))
-				{
-					PreviousClass = ObjectProperty->PropertyClass;
-				}
-				else
-				{
-					PreviousClass = nullptr;
-				}
-			}
-		}
-		else if (Field.IsFunction())
-		{
-			const UFunction* Function = Field.GetFunction();
-
-			FGraphNodeCreator<UK2Node_CallFunction> CallFunctionCreator(*ConversionFunctionGraph);
-			UK2Node_CallFunction* FunctionNode = CallFunctionCreator.CreateNode();
-			FunctionNode->SetFromFunction(Function);
-			FunctionNode->AllocateDefaultPins();
-			CallFunctionCreator.Finalize();
-
-			NewNode = FunctionNode;
-
-			const FProperty* ReturnProperty = UE::MVVM::BindingHelper::GetReturnProperty(Function);
-			if (const FObjectProperty* ObjectProperty = CastField<FObjectProperty>(ReturnProperty))
-			{
-				PreviousClass = ObjectProperty->PropertyClass;
-			}
-			else
-			{
-				PreviousClass = nullptr;
-			}
-		}
-		else
-		{
-			ensureMsgf(false, TEXT("Invalid path, empty field in path."));
-			return;
-		}
-		
-		// first create new connections
-		for (UEdGraphPin* Pin : NewNode->Pins)
-		{
-			if (Pin->Direction == EGPD_Input)
-			{
-				if (Pin->PinName == UEdGraphSchema_K2::PN_Execute)
-				{
-					PreviousExecPin->BreakAllPinLinks();
-					Pin->MakeLinkTo(PreviousExecPin);
-				}
-				else if (Schema->ArePinsCompatible(PreviousDataPin, Pin, PreviousClass))
-				{
-					Pin->MakeLinkTo(PreviousDataPin);
-				}
-			}
-		}
-
-		// then update our previous pin pointers
-		for (UEdGraphPin* Pin : NewNode->Pins)
-		{
-			if (Pin->Direction == EGPD_Output)
-			{
-				if (Pin->PinName == UEdGraphSchema_K2::PN_Then)
-				{
-					PreviousExecPin = Pin;
-				}
-				else
-				{
-					PreviousDataPin = Pin;
-				}
-			}
-		}
-
-		NewNode->NodePosX = PosX;
-		NewNode->NodePosY = PosY;
-
-		PosX += 300; 
-	}
-
-	// finish by linking to our ultimate destinations
-	UEdGraphPin* LastExecPin = ConversionNode->FindPin(UEdGraphSchema_K2::PN_Execute);
-	if (LastExecPin == nullptr)
-	{
-		LastExecPin = UE::MVVM::Private::FindFunctionResult(ConversionFunctionGraph)->FindPin(UEdGraphSchema_K2::PN_Execute);
+		check(ConversionFunction);
+		ConversionFunction->SetGraphPin(WidgetBlueprint, ArgumentName, Path);
 	}
-	PreviousExecPin->MakeLinkTo(LastExecPin);
-	PreviousDataPin->MakeLinkTo(ArgumentPin);
 }
 
 TArray<UE::MVVM::FBindingSource> UMVVMEditorSubsystem::GetBindableWidgets(const UWidgetBlueprint* WidgetBlueprint) const
diff --git a/Engine/Plugins/Runtime/ModelViewViewModel/Source/ModelViewViewModelEditor/Private/Widgets/SMVVMConversionPath.cpp b/Engine/Plugins/Runtime/ModelViewViewModel/Source/ModelViewViewModelEditor/Private/Widgets/SMVVMConversionPath.cpp
deleted file mode 100644
diff --git a/Engine/Plugins/Runtime/ModelViewViewModel/Source/ModelViewViewModelEditor/Private/Widgets/SMVVMConversionPath.h b/Engine/Plugins/Runtime/ModelViewViewModel/Source/ModelViewViewModelEditor/Private/Widgets/SMVVMConversionPath.h
deleted file mode 100644
diff --git a/Engine/Plugins/Runtime/ModelViewViewModel/Source/ModelViewViewModelEditor/Private/Widgets/SMVVMViewBindingListView.cpp b/Engine/Plugins/Runtime/ModelViewViewModel/Source/ModelViewViewModelEditor/Private/Widgets/SMVVMViewBindingListView.cpp
--- a/Engine/Plugins/Runtime/ModelViewViewModel/Source/ModelViewViewModelEditor/Private/Widgets/SMVVMViewBindingListView.cpp
+++ b/Engine/Plugins/Runtime/ModelViewViewModel/Source/ModelViewViewModelEditor/Private/Widgets/SMVVMViewBindingListView.cpp
@@ -12,6 +12,7 @@
 #include "Framework/Views/TableViewMetadata.h"
 #include "Misc/MessageDialog.h"
 #include "MVVMBlueprintView.h"
+#include "MVVMBlueprintViewConversionFunction.h"
 #include "MVVMEditorSubsystem.h"
 #include "MVVMDeveloperProjectSettings.h"
 #include "MVVMWidgetBlueprintExtension_View.h"
@@ -1338,9 +1339,6 @@ void SBindingsList::GetChildrenOfEntry(TSharedPtr<FBindingEntry> Entry, TArray<T
 
 void SBindingsList::Refresh()
 {
-	UMVVMWidgetBlueprintExtension_View* MVVMExtensionPtr = MVVMExtension.Get();
-	UMVVMBlueprintView* BlueprintView = MVVMExtensionPtr ? MVVMExtensionPtr->GetBlueprintView() : nullptr;
-
 	struct FPreviousGroup
 	{
 		TSharedPtr<FBindingEntry> Group;
@@ -1370,6 +1368,10 @@ void SBindingsList::Refresh()
 
 	TArray<TSharedPtr<FBindingEntry>> NewEntries;
 
+	UMVVMWidgetBlueprintExtension_View* MVVMExtensionPtr = MVVMExtension.Get();
+	UMVVMBlueprintView* BlueprintView = MVVMExtensionPtr ? MVVMExtensionPtr->GetBlueprintView() : nullptr;
+	UWidgetBlueprint* WidgetBlueprint = MVVMExtensionPtr ? MVVMExtensionPtr->GetWidgetBlueprint() : nullptr;
+
 	// generate our entries
 	// for each widget with bindings, create an entry at the root level
 	// then add all bindings that reference that widget as its children
@@ -1379,6 +1381,16 @@ void SBindingsList::Refresh()
 		for (int32 BindingIndex = 0; BindingIndex < Bindings.Num(); ++BindingIndex)
 		{
 			const FMVVMBlueprintViewBinding& Binding = Bindings[BindingIndex];
+
+			// Make sure the graph for the bindings is generated
+			if (Binding.Conversion.SourceToDestinationConversion)
+			{
+				Binding.Conversion.SourceToDestinationConversion->GetOrCreateWrapperGraph(WidgetBlueprint);
+			}
+			if (Binding.Conversion.DestinationToSourceConversion)
+			{
+				Binding.Conversion.DestinationToSourceConversion->GetOrCreateWrapperGraph(WidgetBlueprint);
+			}
 			
 			FName GroupName;
 			if (Binding.DestinationPath.IsFromWidget())
diff --git a/Engine/Plugins/Runtime/ModelViewViewModel/Source/ModelViewViewModelEditor/Public/MVVMEditorSubsystem.h b/Engine/Plugins/Runtime/ModelViewViewModel/Source/ModelViewViewModelEditor/Public/MVVMEditorSubsystem.h
--- a/Engine/Plugins/Runtime/ModelViewViewModel/Source/ModelViewViewModelEditor/Public/MVVMEditorSubsystem.h
+++ b/Engine/Plugins/Runtime/ModelViewViewModel/Source/ModelViewViewModelEditor/Public/MVVMEditorSubsystem.h
@@ -78,6 +78,7 @@ public:
 	UFUNCTION(BlueprintCallable, Category = "Viewmodel")
 	UFunction* GetConversionFunction(const UWidgetBlueprint* WidgetBlueprint, const FMVVMBlueprintViewBinding& Binding, bool bSourceToDestination) const;
 
+	UE_DEPRECATED(5.3, "GetConversionFunctionNode was moved to MVVMBlueprintViewConversionFunction.")
 	UFUNCTION(BlueprintCallable, Category = "Viewmodel")
 	UK2Node_CallFunction* GetConversionFunctionNode(const UWidgetBlueprint* WidgetBlueprint, const FMVVMBlueprintViewBinding& Binding, bool bSourceToDestination) const;
 
@@ -89,13 +90,9 @@ public:
 	UEdGraphPin* GetConversionFunctionArgumentPin(const UWidgetBlueprint* WidgetBlueprint, const FMVVMBlueprintViewBinding& Binding, FName ParameterName, bool bSourceToDestination);
 
 	TArray<UE::MVVM::FBindingSource> GetBindableWidgets(const UWidgetBlueprint* WidgetBlueprint) const;
-	
 	TArray<UE::MVVM::FBindingSource> GetAllViewModels(const UWidgetBlueprint* WidgetBlueprint) const;
 
 	FGuid GetFirstBindingThatUsesViewModel(const UWidgetBlueprint* WidgetBlueprint, FGuid ViewModelId) const;
-
-private:
-	UEdGraph* CreateConversionFunctionWrapperGraph(UWidgetBlueprint* WidgetBlueprint, const FMVVMBlueprintViewBinding& Binding, const UFunction* ConversionFunction, bool bSourceToDestination);
 };
 
 #if UE_ENABLE_INCLUDE_ORDER_DEPRECATED_IN_5_2
diff --git a/Engine/Source/Editor/UMGEditor/Private/WidgetBlueprintCompiler.cpp b/Engine/Source/Editor/UMGEditor/Private/WidgetBlueprintCompiler.cpp
--- a/Engine/Source/Editor/UMGEditor/Private/WidgetBlueprintCompiler.cpp
+++ b/Engine/Source/Editor/UMGEditor/Private/WidgetBlueprintCompiler.cpp
@@ -44,6 +44,11 @@ FProperty* FWidgetBlueprintCompilerContext::FCreateVariableContext::CreateVariab
 	return Context.CreateVariable(Name, Type);
 }
 
+void FWidgetBlueprintCompilerContext::FCreateVariableContext::AddGeneratedFunctionGraph(UEdGraph* Graph) const
+{
+	Context.GeneratedFunctionGraphs.Add(Graph);
+}
+
 UWidgetBlueprint* FWidgetBlueprintCompilerContext::FCreateVariableContext::GetWidgetBlueprint() const
 {
 	return Context.WidgetBlueprint();
diff --git a/Engine/Source/Editor/UMGEditor/Public/WidgetBlueprintCompiler.h b/Engine/Source/Editor/UMGEditor/Public/WidgetBlueprintCompiler.h
--- a/Engine/Source/Editor/UMGEditor/Public/WidgetBlueprintCompiler.h
+++ b/Engine/Source/Editor/UMGEditor/Public/WidgetBlueprintCompiler.h
@@ -12,6 +12,7 @@
 #include "Components/WidgetComponent.h"
 
 class FProperty;
+class UEdGraph;
 class UWidget;
 class UWidgetAnimation;
 class UWidgetGraphSchema;
@@ -69,7 +70,7 @@ protected:
 	/** Validates the Desired Focus name to make sure it's part of the Widget Tree. */
 	void ValidateDesiredFocusWidgetName();
 
-	// FKismetCompilerContext
+	//~ Begin FKismetCompilerContext
 	virtual UEdGraphSchema_K2* CreateSchema() override;
 	virtual void CreateFunctionList() override;
 	virtual void SpawnNewClass(const FString& NewClassName) override;
@@ -83,7 +84,7 @@ protected:
 	virtual void FinishCompilingClass(UClass* Class) override;
 	virtual bool ValidateGeneratedClass(UBlueprintGeneratedClass* Class) override;
 	virtual void OnPostCDOCompiled(const UObject::FPostCDOCompiledContext& Context) override;
-	// End FKismetCompilerContext
+	//~ End FKismetCompilerContext
 
 	void SanitizeBindings(UBlueprintGeneratedClass* Class);
 
@@ -99,6 +100,7 @@ public:
 	{
 	public:
 		FProperty* CreateVariable(const FName Name, const FEdGraphPinType& Type) const;
+		void AddGeneratedFunctionGraph(UEdGraph* Graph) const;
 		UWidgetBlueprint* GetWidgetBlueprint() const;
 		UWidgetBlueprintGeneratedClass* GetSkeletonGeneratedClass() const;
 		EKismetCompileType::Type GetCompileType() const;
