Fix 338325 build issue.

diff --git a/Engine/Plugins/Runtime/ModelViewViewModel/Source/ModelViewViewModelBlueprint/Private/MVVMBlueprintViewConversionFunction.cpp b/Engine/Plugins/Runtime/ModelViewViewModel/Source/ModelViewViewModelBlueprint/Private/MVVMBlueprintViewConversionFunction.cpp
--- a/Engine/Plugins/Runtime/ModelViewViewModel/Source/ModelViewViewModelBlueprint/Private/MVVMBlueprintViewConversionFunction.cpp
+++ b/Engine/Plugins/Runtime/ModelViewViewModel/Source/ModelViewViewModelBlueprint/Private/MVVMBlueprintViewConversionFunction.cpp
@@ -137,7 +137,7 @@ void UMVVMBlueprintViewConversionFunction::InitializeFromWrapperGraph(UBlueprint
 		CachedWrapperNode = WrapperNode;
 		SavePinValues(SelfContext);
 
-		if (bWrapperGraphTransient && bWrapperGraphTransient)
+		if (bWrapperGraphTransient && CachedWrapperNode)
 		{
 			SelfContext->FunctionGraphs.RemoveSingle(CachedWrapperGraph);
 		}
