MVVM: When a conversion function is virtual, reevaluate at runtime the function.

diff --git a/Engine/Plugins/Runtime/ModelViewViewModel/Source/ModelViewViewModel/Private/Bindings/MVVMCompiledBindingLibrary.cpp b/Engine/Plugins/Runtime/ModelViewViewModel/Source/ModelViewViewModel/Private/Bindings/MVVMCompiledBindingLibrary.cpp
--- a/Engine/Plugins/Runtime/ModelViewViewModel/Source/ModelViewViewModel/Private/Bindings/MVVMCompiledBindingLibrary.cpp
+++ b/Engine/Plugins/Runtime/ModelViewViewModel/Source/ModelViewViewModel/Private/Bindings/MVVMCompiledBindingLibrary.cpp
@@ -18,6 +18,14 @@ DECLARE_CYCLE_STAT(TEXT("Load Library"), STAT_UMG_Viewmodel_LoadLibrary, STATGRO
 
 #define LOCTEXT_NAMESPACE "MVVMCompiledBindingLibrary"
 
+namespace UE::MVVM::Private
+{
+	bool IsFunctionVirtual(UFunction* Function)
+	{
+		return !Function->HasAnyFunctionFlags(FUNC_Static | FUNC_Final);
+	}
+}
+
 /**
  *
  */
@@ -187,7 +195,7 @@ TValueOrError<void, FMVVMCompiledBindingLibrary::EExecutionFailingReason> FMVVMC
 			checkf(TestSource.GetValue().GetObjectVariant().GetData() != nullptr, TEXT("Tested in EvaluateFieldPath"));
 #endif
 
-			TValueOrError<UE::MVVM::FMVVMFieldVariant, void> FinalPath = GetFinalFieldFromPathImpl(InBinding.SourceFieldPath);
+			TValueOrError<UE::MVVM::FMVVMFieldVariant, void> FinalPath = GetFinalFieldFromPathImpl(UE::MVVM::FObjectVariant(InSource), InBinding.SourceFieldPath);
 			if (FinalPath.HasError())
 			{
 				return MakeError(EExecutionFailingReason::InvalidSource);
@@ -363,6 +371,11 @@ TValueOrError<UE::MVVM::FFieldContext, void> FMVVMCompiledBindingLibrary::Evalua
 		{
 			check(LoadedFunctions.IsValidIndex(PathIndex.Index));
 			UFunction* Function = LoadedFunctions[PathIndex.Index];
+			if (Function && UE::MVVM::Private::IsFunctionVirtual(Function) && CurrentContainer.IsUObject())
+			{
+				Function = CurrentContainer.GetUObject()->GetClass()->FindFunctionByName(Function->GetFName());
+			}
+
 			if (!Function)
 			{
 				return MakeError();
@@ -398,7 +411,7 @@ TValueOrError<UE::MVVM::FFieldContext, void> FMVVMCompiledBindingLibrary::Evalua
 		return MakeError();
 	}
 
-	TValueOrError<UE::MVVM::FMVVMFieldVariant, void> FinalPath = GetFinalFieldFromPathImpl(InFieldPath);
+	TValueOrError<UE::MVVM::FMVVMFieldVariant, void> FinalPath = GetFinalFieldFromPathImpl(CurrentContainer, InFieldPath);
 	if (FinalPath.HasError())
 	{
 		return MakeError();
@@ -407,7 +420,7 @@ TValueOrError<UE::MVVM::FFieldContext, void> FMVVMCompiledBindingLibrary::Evalua
 }
 
 
-TValueOrError<UE::MVVM::FMVVMFieldVariant, void> FMVVMCompiledBindingLibrary::GetFinalFieldFromPathImpl(const FMVVMVCompiledFieldPath& InFieldPath) const
+TValueOrError<UE::MVVM::FMVVMFieldVariant, void> FMVVMCompiledBindingLibrary::GetFinalFieldFromPathImpl(UE::MVVM::FObjectVariant CurrentContainer, const FMVVMVCompiledFieldPath& InFieldPath) const
 {
 	const int32 IterationMax = InFieldPath.StartIndex + InFieldPath.Num - 1;
 
@@ -428,6 +441,11 @@ TValueOrError<UE::MVVM::FMVVMFieldVariant, void> FMVVMCompiledBindingLibrary::Ge
 	{
 		check(LoadedFunctions.IsValidIndex(FinalPathIndex.Index));
 		UFunction* Function = LoadedFunctions[FinalPathIndex.Index];
+		if (Function && UE::MVVM::Private::IsFunctionVirtual(Function) && CurrentContainer.IsUObject())
+		{
+			Function = CurrentContainer.GetUObject()->GetClass()->FindFunctionByName(Function->GetFName());
+		}
+
 		if (!Function)
 		{
 			return MakeError();
diff --git a/Engine/Plugins/Runtime/ModelViewViewModel/Source/ModelViewViewModel/Public/Bindings/MVVMCompiledBindingLibrary.h b/Engine/Plugins/Runtime/ModelViewViewModel/Source/ModelViewViewModel/Public/Bindings/MVVMCompiledBindingLibrary.h
--- a/Engine/Plugins/Runtime/ModelViewViewModel/Source/ModelViewViewModel/Public/Bindings/MVVMCompiledBindingLibrary.h
+++ b/Engine/Plugins/Runtime/ModelViewViewModel/Source/ModelViewViewModel/Public/Bindings/MVVMCompiledBindingLibrary.h
@@ -3,6 +3,7 @@
 #pragma once
 
 #include "FieldNotificationId.h" // IWYU pragma: keep
+#include "Types/MVVMObjectVariant.h"
 #include "MVVMCompiledBindingLibrary.generated.h"
 
 namespace UE::FieldNotification { struct FFieldId; }
@@ -168,6 +169,7 @@ struct FMVVMCompiledLoadedPropertyOrFunctionIndex
 	}
 
 private:
+	/** Is the index in LoadedProperties or LoadedFunctions.*/
 	UPROPERTY()
 	int16 Index = INDEX_NONE;
 
@@ -179,7 +181,7 @@ private:
 	UPROPERTY()
 	uint8 bIsScriptStructProperty : 1;
 
-	/** Is the index in LoadedProperties or the index is in LoadedFunctions. */
+	/** Is the index in LoadedProperties or LoadedFunctions. */
 	UPROPERTY()
 	uint8 bIsProperty : 1;
 };
@@ -331,7 +333,7 @@ private:
 	TValueOrError<void, EExecutionFailingReason> ExecuteImpl(UObject* ExecutionSource, const FMVVMVCompiledBinding& Binding, UObject* Source, EConversionFunctionType ConversionType) const;
 	TValueOrError<void, EExecutionFailingReason> ExecuteImpl(UE::MVVM::FFieldContext& Source, UE::MVVM::FFieldContext& Destination, UE::MVVM::FFunctionContext& ConversionFunction, EConversionFunctionType ConversionType) const;
 
-	TValueOrError<UE::MVVM::FMVVMFieldVariant, void> GetFinalFieldFromPathImpl(const FMVVMVCompiledFieldPath& FieldPath) const;
+	TValueOrError<UE::MVVM::FMVVMFieldVariant, void> GetFinalFieldFromPathImpl(UE::MVVM::FObjectVariant CurrentContainer, const FMVVMVCompiledFieldPath& FieldPath) const;
 
 private:
 	TArray<FProperty*> LoadedProperties;
diff --git a/Engine/Plugins/Runtime/ModelViewViewModel/Source/ModelViewViewModelBlueprint/Private/Bindings/MVVMConversionFunctionHelper.cpp b/Engine/Plugins/Runtime/ModelViewViewModel/Source/ModelViewViewModelBlueprint/Private/Bindings/MVVMConversionFunctionHelper.cpp
--- a/Engine/Plugins/Runtime/ModelViewViewModel/Source/ModelViewViewModelBlueprint/Private/Bindings/MVVMConversionFunctionHelper.cpp
+++ b/Engine/Plugins/Runtime/ModelViewViewModel/Source/ModelViewViewModelBlueprint/Private/Bindings/MVVMConversionFunctionHelper.cpp
@@ -88,7 +88,7 @@ namespace Private
 		FGraphNodeCreator<UK2Node_FunctionEntry> FunctionEntryCreator(*FunctionGraph);
 		UK2Node_FunctionEntry* FunctionEntry = FunctionEntryCreator.CreateNode();
 		FunctionEntry->FunctionReference.SetSelfMember(FunctionGraph->GetFName());
-		FunctionEntry->AddExtraFlags(FUNC_BlueprintCallable | FUNC_BlueprintPure | FUNC_Const | FUNC_Protected);
+		FunctionEntry->AddExtraFlags(FUNC_BlueprintCallable | FUNC_BlueprintPure | FUNC_Const | FUNC_Protected | FUNC_Final);
 		FunctionEntry->bIsEditable = bIsEditable;
 		FunctionEntry->MetaData.Category = FText::FromName(ConversionFunctionCategory.Resolve());
 		FunctionEntry->NodePosX = -500;
diff --git a/Engine/Plugins/Runtime/ModelViewViewModel/Source/ModelViewViewModelBlueprint/Private/MVVMViewBlueprintCompiler.cpp b/Engine/Plugins/Runtime/ModelViewViewModel/Source/ModelViewViewModelBlueprint/Private/MVVMViewBlueprintCompiler.cpp
--- a/Engine/Plugins/Runtime/ModelViewViewModel/Source/ModelViewViewModelBlueprint/Private/MVVMViewBlueprintCompiler.cpp
+++ b/Engine/Plugins/Runtime/ModelViewViewModel/Source/ModelViewViewModelBlueprint/Private/MVVMViewBlueprintCompiler.cpp
@@ -626,7 +626,7 @@ void FMVVMViewBlueprintCompiler::CreateFunctionsDeclaration(const FWidgetBluepri
 				SourceCreator.SetterGraph = UE::MVVM::FunctionGraphHelper::CreateIntermediateFunctionGraph(
 					WidgetBlueprintCompilerContext
 					, SourceCreator.SetterFunctionName
-					, (FUNC_BlueprintCallable | FUNC_Public)
+					, (FUNC_BlueprintCallable | FUNC_Public | FUNC_Final)
 					, TEXT("Viewmodel")
 					, false);
 				BlueprintView->TemporaryGraph.Add(SourceCreator.SetterGraph);
