Temporary fix for UE-191378. Log an error instead of crashing until we figure out what is going on.
#rnx

diff --git a/Engine/Plugins/Runtime/ModelViewViewModel/Source/ModelViewViewModelBlueprint/Private/MVVMViewBlueprintCompiler.cpp b/Engine/Plugins/Runtime/ModelViewViewModel/Source/ModelViewViewModelBlueprint/Private/MVVMViewBlueprintCompiler.cpp
--- a/Engine/Plugins/Runtime/ModelViewViewModel/Source/ModelViewViewModelBlueprint/Private/MVVMViewBlueprintCompiler.cpp
+++ b/Engine/Plugins/Runtime/ModelViewViewModel/Source/ModelViewViewModelBlueprint/Private/MVVMViewBlueprintCompiler.cpp
@@ -1743,7 +1743,10 @@ TValueOrError<FMVVMViewBlueprintCompiler::FBindingSourceContext, FText> FMVVMVie
 		else
 		{
 			Result.UserWidgetPropertyContextIndex = CompilerUserWidgetPropertyContexts.IndexOfByPredicate([SourceName](const FCompilerUserWidgetPropertyContext& Other) { return Other.PropertyName == SourceName; });
-			check(Result.UserWidgetPropertyContextIndex != INDEX_NONE);
+			if (!CompilerUserWidgetPropertyContexts.IsValidIndex(Result.UserWidgetPropertyContextIndex))
+			{
+				return MakeError(LOCTEXT("InvalidUserWidgetPropertyContextIndexInternal", "Internal error. UserWidgetPropertyContextIndex is invalid."));
+			}
 			Result.SourceClass = CompilerUserWidgetPropertyContexts[Result.UserWidgetPropertyContextIndex].Class;
 			Result.PropertyPath = CreatePropertyPath(Class, CompilerUserWidgetPropertyContexts[Result.UserWidgetPropertyContextIndex].PropertyName, PropertyPath.GetFields(Class));
 		}
