MVVM: Fix bad refactor code introduced with CL 26614337

diff --git a/Engine/Plugins/Runtime/ModelViewViewModel/Source/ModelViewViewModelBlueprint/Private/MVVMViewBlueprintCompiler.cpp b/Engine/Plugins/Runtime/ModelViewViewModel/Source/ModelViewViewModelBlueprint/Private/MVVMViewBlueprintCompiler.cpp
--- a/Engine/Plugins/Runtime/ModelViewViewModel/Source/ModelViewViewModelBlueprint/Private/MVVMViewBlueprintCompiler.cpp
+++ b/Engine/Plugins/Runtime/ModelViewViewModel/Source/ModelViewViewModelBlueprint/Private/MVVMViewBlueprintCompiler.cpp
@@ -553,7 +553,7 @@ void FMVVMViewBlueprintCompiler::CreateSourceLists(const FWidgetBlueprintCompile
 		if (bIsForwardBinding || bIsBackwardBinding)
 		{
 			const UMVVMBlueprintViewConversionFunction* ConversionFunction = Binding.Conversion.GetConversionFunction(bIsForwardBinding);
-			if (ConversionFunction)
+			if (ConversionFunction && ConversionFunction->GetPins().Num() > 0)
 			{
 				if (Binding.BindingType == EMVVMBindingMode::TwoWay)
 				{
