MVVM: Add UX for view's settings. Seperate bindings initialization and sources initialization. Use can now select if it wants the bindings and the sources to be initialized.

diff --git a/Engine/Plugins/Runtime/ModelViewViewModel/Source/ModelViewViewModel/Private/View/MVVMView.cpp b/Engine/Plugins/Runtime/ModelViewViewModel/Source/ModelViewViewModel/Private/View/MVVMView.cpp
--- a/Engine/Plugins/Runtime/ModelViewViewModel/Source/ModelViewViewModel/Private/View/MVVMView.cpp
+++ b/Engine/Plugins/Runtime/ModelViewViewModel/Source/ModelViewViewModel/Private/View/MVVMView.cpp
@@ -16,8 +16,10 @@
 
 #include UE_INLINE_GENERATED_CPP_BY_NAME(MVVMView)
 
+DECLARE_CYCLE_STAT(TEXT("InitializeSources"), STAT_UMG_Viewmodel_InitializeSources, STATGROUP_UMG_Viewmodel);
+DECLARE_CYCLE_STAT(TEXT("UninitializeSources"), STAT_UMG_Viewmodel_UninitializeSources, STATGROUP_UMG_Viewmodel);
 DECLARE_CYCLE_STAT(TEXT("InitializeBindings"), STAT_UMG_Viewmodel_InitializeBindings, STATGROUP_UMG_Viewmodel);
-DECLARE_CYCLE_STAT(TEXT("DeinitializeBindings"), STAT_UMG_Viewmodel_DeinitializeBindings, STATGROUP_UMG_Viewmodel);
+DECLARE_CYCLE_STAT(TEXT("UninitializeBindings"), STAT_UMG_Viewmodel_UninitializeBindings, STATGROUP_UMG_Viewmodel);
 DECLARE_CYCLE_STAT(TEXT("SetSource"), STAT_UMG_Viewmodel_SetSource, STATGROUP_UMG_Viewmodel);
 DECLARE_CYCLE_STAT(TEXT("ExecuteBinding ValueChanged"), STAT_UMG_Viewmodel_ExecuteBinding_ValueChanged, STATGROUP_UMG_Viewmodel);
 DECLARE_CYCLE_STAT(TEXT("ExecuteBinding Delayed"), STAT_UMG_Viewmodel_ExecuteBinding_Delayed, STATGROUP_UMG_Viewmodel);
@@ -50,21 +52,38 @@ void UMVVMView::Construct()
 	check(ClassExtension);
 	check(bConstructed == false);
 
-	if (ClassExtension->DoesInitializeAtConstruction())
+	if (ClassExtension->InitializeSourcesOnConstruct())
 	{
-		InitializeBindings();
+		InitializeSources();
 	}
 
 	bConstructed = true;
 }
 
-void UMVVMView::InitializeBindings()
+
+void UMVVMView::Destruct()
 {
-	SCOPE_CYCLE_COUNTER(STAT_UMG_Viewmodel_InitializeBindings);
+	check(bConstructed == true);
+	bConstructed = false;
+
+#if UE_WITH_MVVM_DEBUGGING
+	UE::MVVM::FDebugging::BroadcastViewBeginDestruction(this);
+#endif
+
+	UninitializeSources(); // and bindings
+}
+
 
-	check(bInitialized == false);
+void UMVVMView::InitializeSources()
+{
+	if (bSourcesInitialized || ClassExtension == nullptr)
+	{
+		return;
+	}
+
+	SCOPE_CYCLE_COUNTER(STAT_UMG_Viewmodel_InitializeSources);
 
-	// Init ViewModel instances
+	// Init Sources/ViewModel instances
 	UUserWidget* UserWidget = GetUserWidget();
 	const TArrayView<const FMVVMViewClass_SourceCreator> AllViewModelCreators = ClassExtension->GetViewModelCreators();
 	for (int32 Index = 0; Index < AllViewModelCreators.Num(); ++Index)
@@ -96,10 +115,46 @@ void UMVVMView::InitializeBindings()
 		}
 	}
 
+	bSourcesInitialized = true;
+
 #if UE_WITH_MVVM_DEBUGGING
 	UE::MVVM::FDebugging::BroadcastViewConstructed(this);
 #endif
 
+	if (ClassExtension->InitializeBindingsOnConstruct())
+	{
+		InitializeBindings();
+	}
+}
+
+
+void UMVVMView::UninitializeSources()
+{
+	if (!bSourcesInitialized)
+	{
+		return;
+	}
+
+	SCOPE_CYCLE_COUNTER(STAT_UMG_Viewmodel_UninitializeSources);
+
+	UninitializeInternal(true);
+}
+
+
+void UMVVMView::InitializeBindings()
+{
+	if (!bSourcesInitialized)
+	{
+		InitializeSources();
+	}
+
+	if (bBindingsInitialized)
+	{
+		return;
+	}
+
+	SCOPE_CYCLE_COUNTER(STAT_UMG_Viewmodel_InitializeBindings);
+
 	bHasEveryTickBinding = false;
 
 	const TArrayView<const FMVVMViewClass_CompiledBinding> CompiledBindings = ClassExtension->GetCompiledBindings();
@@ -115,7 +170,7 @@ void UMVVMView::InitializeBindings()
 		}
 	}
 
-	bInitialized = true;
+	bBindingsInitialized = true;
 
 	if (bHasEveryTickBinding)
 	{
@@ -124,27 +179,27 @@ void UMVVMView::InitializeBindings()
 }
 
 
-void UMVVMView::Destruct()
+void UMVVMView::UninitializeBindings()
 {
-	check(bConstructed == true);
-	bConstructed = false;
-
-#if UE_WITH_MVVM_DEBUGGING
-	UE::MVVM::FDebugging::BroadcastViewBeginDestruction(this);
-#endif
-
-	if (bInitialized)
+	if (!bBindingsInitialized)
 	{
-		DeinitializeBindings();
+		return;
 	}
+
+	SCOPE_CYCLE_COUNTER(STAT_UMG_Viewmodel_UninitializeBindings);
+
+	UninitializeInternal(false);
 }
 
-void UMVVMView::DeinitializeBindings()
+
+void UMVVMView::UninitializeInternal(bool bUninitializeSources)
 {
-	SCOPE_CYCLE_COUNTER(STAT_UMG_Viewmodel_DeinitializeBindings);
+	SCOPE_CYCLE_COUNTER(STAT_UMG_Viewmodel_UninitializeBindings);
+
+	bool bUninitializeBindings = bBindingsInitialized;
 
-	check(bInitialized == true);
-	bInitialized = false;
+	bBindingsInitialized = false;
+	bSourcesInitialized = !bUninitializeSources;
 
 	UUserWidget* UserWidget = GetUserWidget();
 	check(UserWidget);
@@ -155,7 +210,7 @@ void UMVVMView::DeinitializeBindings()
 		const FMVVMViewClass_SourceCreator& Item = AllViewModelCreators[Index];
 		FMVVMViewSource& Source = Sources[Index];
 
-		if (Source.RegisteredCount > 0 && Source.Source)
+		if (Source.RegisteredCount > 0 && Source.Source && bUninitializeBindings)
 		{
 			TScriptInterface<INotifyFieldValueChanged> SourceAsInterface = Source.Source;
 			checkf(SourceAsInterface.GetInterface(), TEXT("It was added as a INotifyFieldValueChanged. It should still be."));
@@ -163,7 +218,7 @@ void UMVVMView::DeinitializeBindings()
 		}
 
 		// For GC release any object used by the view
-		if (!Source.bSetManually)
+		if (!Source.bSetManually && bUninitializeSources)
 		{
 			Item.DestroyInstance(Source.Source, this);
 
@@ -184,6 +239,7 @@ void UMVVMView::DeinitializeBindings()
 
 	if (bHasEveryTickBinding)
 	{
+		ensureMsgf(bUninitializeBindings, TEXT("The "));
 		GEngine->GetEngineSubsystem<UMVVMBindingSubsystem>()->RemoveViewWithEveryTickBinding(this);
 	}
 	bHasEveryTickBinding = false;
@@ -336,7 +392,7 @@ bool UMVVMView::SetSourceInternal(FName ViewModelName, TScriptInterface<INotifyF
 		const TArrayView<const FMVVMViewClass_CompiledBinding> CompiledBindings = ClassExtension->GetCompiledBindings();
 
 		// Unregister any bindings from that source
-		if (bInitialized)
+		if (bBindingsInitialized)
 		{
 			for (int32 Index = 0; Index < CompiledBindings.Num(); ++Index)
 			{
@@ -367,8 +423,8 @@ bool UMVVMView::SetSourceInternal(FName ViewModelName, TScriptInterface<INotifyF
 
 		bool bPreviousEveryTickBinding = bHasEveryTickBinding;
 		bHasEveryTickBinding = false;
-		// Register back any binding that was previously enabled
-		if (bInitialized && NewValue.GetObject())
+		// Register back any bindings that was previously enabled
+		if (bBindingsInitialized && NewValue.GetObject())
 		{
 			// Enabled the default bindings
 			for (int32 Index = 0; Index < CompiledBindings.Num(); ++Index)
@@ -398,7 +454,7 @@ bool UMVVMView::SetSourceInternal(FName ViewModelName, TScriptInterface<INotifyF
 			}
 		}
 
-		if (bInitialized && bPreviousEveryTickBinding != bHasEveryTickBinding)
+		if (bBindingsInitialized && bPreviousEveryTickBinding != bHasEveryTickBinding)
 		{
 			if (bHasEveryTickBinding)
 			{
@@ -473,8 +529,9 @@ void UMVVMView::ExecuteDelayedBinding(const FMVVMViewDelayedBinding& DelayedBind
 {
 	SCOPE_CYCLE_COUNTER(STAT_UMG_Viewmodel_ExecuteBinding_Delayed);
 
-	if (ensure(ClassExtension) && bInitialized)
+	if (ensure(ClassExtension) && bBindingsInitialized)
 	{
+		ensure(bSourcesInitialized);
 		if (ensure(ClassExtension->GetCompiledBindings().IsValidIndex(DelayedBinding.GetCompiledBindingIndex())))
 		{
 			if (IsLibraryBindingEnabled(DelayedBinding.GetCompiledBindingIndex()))
@@ -514,9 +571,10 @@ void UMVVMView::ExecuteEveryTickBindings() const
 	SCOPE_CYCLE_COUNTER(STAT_UMG_Viewmodel_ExecuteBinding_Tick);
 
 	ensure(bHasEveryTickBinding);
-	ensure(bInitialized);
+	ensure(bBindingsInitialized);
+	ensure(bSourcesInitialized);
 
-	if (ClassExtension && bInitialized)
+	if (ClassExtension && bBindingsInitialized)
 	{
 		const TArrayView<const FMVVMViewClass_CompiledBinding> CompiledBindings = ClassExtension->GetCompiledBindings();
 
diff --git a/Engine/Plugins/Runtime/ModelViewViewModel/Source/ModelViewViewModel/Public/View/MVVMView.h b/Engine/Plugins/Runtime/ModelViewViewModel/Source/ModelViewViewModel/Public/View/MVVMView.h
--- a/Engine/Plugins/Runtime/ModelViewViewModel/Source/ModelViewViewModel/Public/View/MVVMView.h
+++ b/Engine/Plugins/Runtime/ModelViewViewModel/Source/ModelViewViewModel/Public/View/MVVMView.h
@@ -68,8 +68,41 @@ public:
 	virtual void Destruct() override;
 	//~ End UUserWidgetExtension implementation
 
+	/**
+	 * Initialize the sources if they are not already initialized.
+	 * Initializing the sources will also initialize the bindings if the option bInitializeBindingsOnConstruct is enabled.
+	 * @note A sources can be a viewmodel or any other object used by a binding.
+	 */
+	UFUNCTION(BlueprintCallable, Category = "Viewmodel")
+	void InitializeSources();
+	/**
+	 * Uninitialize the sources if they are already initialized.
+	 * It will uninitialized the bindings.
+	 */
+	UFUNCTION(BlueprintCallable, Category = "Viewmodel")
+	void UninitializeSources();
+	/** The sources were initialized, manually or automatically. */
+	UFUNCTION(BlueprintPure, Category = "Viewmodel")
+	bool AreSourcesInitialized() const
+	{
+		return bSourcesInitialized;
+	}
+
+	/**
+	 * Initialize the bindings if they are not already initialized.
+	 * Initializing the bindings will execute them.
+	 */
+	UFUNCTION(BlueprintCallable, Category = "Viewmodel")
 	void InitializeBindings();
-	void DeinitializeBindings();
+	/** Uninitialize the bindings if they are already initialized. */
+	UFUNCTION(BlueprintCallable, Category = "Viewmodel")
+	void UninitializeBindings();
+	/** The bindings were initialized, manually or automatically. */
+	UFUNCTION(BlueprintPure, Category = "Viewmodel")
+	bool AreBindingsInitialized() const
+	{
+		return bBindingsInitialized;
+	}
 
 	const UMVVMViewClass* GetViewClass() const
 	{
@@ -107,6 +140,7 @@ public:
 private:
 	bool EvaluateSourceCreator(int32 SourceIndex);
 	bool SetSourceInternal(FName ViewModelName, TScriptInterface<INotifyFieldValueChanged> ViewModel, bool bForDynamicSource);
+	void UninitializeInternal(bool bSources);
 
 	void HandledLibraryBindingValueChanged(UObject* InViewModel, UE::FieldNotification::FFieldId InFieldId, int32 InCompiledBindingIndex) const;
 
@@ -143,7 +177,10 @@ private:
 	
 	/** Is the Initialize method called. */
 	UPROPERTY(VisibleAnywhere, Transient, Category = "Viewmodel")
-	bool bInitialized = false;
+	bool bSourcesInitialized = false;
+	/** Is the Initialize method called. */
+	UPROPERTY(VisibleAnywhere, Transient, Category = "Viewmodel")
+	bool bBindingsInitialized = false;
 
 	/** The view has at least one binding that need to be ticked every frame. */
 	UPROPERTY(VisibleAnywhere, Transient, Category = "Viewmodel")
diff --git a/Engine/Plugins/Runtime/ModelViewViewModel/Source/ModelViewViewModel/Public/View/MVVMViewClass.h b/Engine/Plugins/Runtime/ModelViewViewModel/Source/ModelViewViewModel/Public/View/MVVMViewClass.h
--- a/Engine/Plugins/Runtime/ModelViewViewModel/Source/ModelViewViewModel/Public/View/MVVMViewClass.h
+++ b/Engine/Plugins/Runtime/ModelViewViewModel/Source/ModelViewViewModel/Public/View/MVVMViewClass.h
@@ -333,9 +333,14 @@ public:
 
 public:
 	/** Should it automatically execute the bindings when the view is constructed or they will be executed manually. */
-	bool DoesInitializeAtConstruction() const
+	bool InitializeSourcesOnConstruct() const
 	{
-		return bInitializeAtConstruction;
+		return bInitializeSourcesOnConstruct;
+	}
+	/** Should it automatically execute the bindings when the view is constructed or they will be executed manually. */
+	bool InitializeBindingsOnConstruct() const
+	{
+		return bInitializeBindingsOnConstruct;
 	}
 	
 	/** Get the list of the needed ViewModel. */
@@ -395,7 +400,10 @@ private:
 
 	/** */
 	UPROPERTY()
-	bool bInitializeAtConstruction = true;
+	bool bInitializeSourcesOnConstruct = true;
+	/** */
+	UPROPERTY()
+	bool bInitializeBindingsOnConstruct = true;
 
 #if WITH_EDITORONLY_DATA
 	FDelegateHandle BluerpintCompiledHandle;
diff --git a/Engine/Plugins/Runtime/ModelViewViewModel/Source/ModelViewViewModelBlueprint/Private/MVVMBlueprintView.cpp b/Engine/Plugins/Runtime/ModelViewViewModel/Source/ModelViewViewModelBlueprint/Private/MVVMBlueprintView.cpp
--- a/Engine/Plugins/Runtime/ModelViewViewModel/Source/ModelViewViewModelBlueprint/Private/MVVMBlueprintView.cpp
+++ b/Engine/Plugins/Runtime/ModelViewViewModel/Source/ModelViewViewModelBlueprint/Private/MVVMBlueprintView.cpp
@@ -12,6 +12,11 @@
 #include UE_INLINE_GENERATED_CPP_BY_NAME(MVVMBlueprintView)
 
 
+UMVVMBlueprintView::UMVVMBlueprintView()
+{
+	Settings = CreateDefaultSubobject<UMVVMBlueprintViewSettings>("Settings");
+}
+
 FMVVMBlueprintViewModelContext* UMVVMBlueprintView::FindViewModel(FGuid ViewModelId)
 {
 	return AvailableViewModels.FindByPredicate([ViewModelId](const FMVVMBlueprintViewModelContext& Other)
diff --git a/Engine/Plugins/Runtime/ModelViewViewModel/Source/ModelViewViewModelBlueprint/Private/MVVMViewBlueprintCompiler.cpp b/Engine/Plugins/Runtime/ModelViewViewModel/Source/ModelViewViewModelBlueprint/Private/MVVMViewBlueprintCompiler.cpp
--- a/Engine/Plugins/Runtime/ModelViewViewModel/Source/ModelViewViewModelBlueprint/Private/MVVMViewBlueprintCompiler.cpp
+++ b/Engine/Plugins/Runtime/ModelViewViewModel/Source/ModelViewViewModelBlueprint/Private/MVVMViewBlueprintCompiler.cpp
@@ -860,6 +860,11 @@ bool FMVVMViewBlueprintCompiler::Compile(UWidgetBlueprintGeneratedClass* Class,
 	CompileSourceCreators(CompileResult.GetValue(), Class, BlueprintView, ViewExtension);
 	CompileBindings(CompileResult.GetValue(), Class, BlueprintView, ViewExtension);
 
+	{
+		ViewExtension->bInitializeSourcesOnConstruct = BlueprintView->GetSettings()->bInitializeSourcesOnConstruct;
+		ViewExtension->bInitializeBindingsOnConstruct = ViewExtension->bInitializeSourcesOnConstruct ? BlueprintView->GetSettings()->bInitializeBindingsOnConstruct : false;
+	}
+
 	bool bResult = bAreSourcesCreatorValid && bAreSourceContextsValid && bIsBindingsValid;
 	if (bResult)
 	{
diff --git a/Engine/Plugins/Runtime/ModelViewViewModel/Source/ModelViewViewModelBlueprint/Public/MVVMBlueprintView.h b/Engine/Plugins/Runtime/ModelViewViewModel/Source/ModelViewViewModelBlueprint/Public/MVVMBlueprintView.h
--- a/Engine/Plugins/Runtime/ModelViewViewModel/Source/ModelViewViewModelBlueprint/Public/MVVMBlueprintView.h
+++ b/Engine/Plugins/Runtime/ModelViewViewModel/Source/ModelViewViewModelBlueprint/Public/MVVMBlueprintView.h
@@ -28,6 +28,34 @@ namespace  UE::MVVM
 	};
 }
 
+/**
+ *
+ */
+UCLASS(MinimalAPI)
+class UMVVMBlueprintViewSettings : public UObject
+{
+	GENERATED_BODY()
+
+public:
+	/**
+	 * Auto initialize the view sources when the Widget is constructed.
+	 * If false, the user will have to initialize the sources manually.
+	 * It prevents the sources evaluating until you are ready.
+	 */
+	UPROPERTY(EditAnywhere, Category = "Viewmodel")
+	bool bInitializeSourcesOnConstruct = true;
+	/**
+	 * Auto initialize the view bindings when the Widget is constructed.
+	 * If false, the user will have to initialize the bindings manually.
+	 * It prevents bindings execution and improves performance when you know the widget won't be visible.
+	 * @note All bindings are executed when the view is automatically initialized or manually initialized.
+	 * @note Sources needs to be initialized before initializing the bindings.
+	 * @note When Sources is manually initialized, the bindings will also be initialized if this is true.
+	 */
+	UPROPERTY(EditAnywhere, Category = "Viewmodel", meta=(EditCondition="bInitializeSourcesOnConstruct"))
+	bool bInitializeBindingsOnConstruct = true;
+};
+
 /**
  * 
  */
@@ -37,6 +65,14 @@ class MODELVIEWVIEWMODELBLUEPRINT_API UMVVMBlueprintView : public UObject
 	GENERATED_BODY()
 
 public:
+	UMVVMBlueprintView();
+
+public:
+	UMVVMBlueprintViewSettings* GetSettings()
+	{
+		return Settings;
+	}
+
 	FMVVMBlueprintViewModelContext* FindViewModel(FGuid ViewModelId);
 	const FMVVMBlueprintViewModelContext* FindViewModel(FGuid ViewModelId) const;
 	const FMVVMBlueprintViewModelContext* FindViewModel(FName ViewModelName) const;
@@ -111,6 +147,9 @@ public:
 	TArray<TObjectPtr<UEdGraph>> TemporaryGraph;
 
 private:
+	UPROPERTY()
+	TObjectPtr<UMVVMBlueprintViewSettings> Settings;
+
 	UPROPERTY(EditAnywhere, Category = "Viewmodel")
 	TArray<FMVVMBlueprintViewBinding> Bindings;
 
diff --git a/Engine/Plugins/Runtime/ModelViewViewModel/Source/ModelViewViewModelEditor/Private/Widgets/SMVVMViewBindingPanel.cpp b/Engine/Plugins/Runtime/ModelViewViewModel/Source/ModelViewViewModelEditor/Private/Widgets/SMVVMViewBindingPanel.cpp
--- a/Engine/Plugins/Runtime/ModelViewViewModel/Source/ModelViewViewModelEditor/Private/Widgets/SMVVMViewBindingPanel.cpp
+++ b/Engine/Plugins/Runtime/ModelViewViewModel/Source/ModelViewViewModelEditor/Private/Widgets/SMVVMViewBindingPanel.cpp
@@ -26,6 +26,7 @@
 #include "Widgets/Docking/SDockTab.h"
 #include "Widgets/Input/SButton.h"
 #include "Widgets/Input/SCheckBox.h"
+#include "Widgets/Input/SComboButton.h"
 #include "Widgets/Input/SSearchBox.h"
 #include "Widgets/Images/SImage.h"
 #include "Widgets/Layout/SSpacer.h"
@@ -62,8 +63,29 @@ struct FStructDetailNotifyHook : FNotifyHook
 	FMVVMBlueprintViewBinding Binding;
 	TWeakObjectPtr<UMVVMWidgetBlueprintExtension_View> MVVMExtension;
 };
+
+void SetSelectObjectsToViewSettings(TWeakPtr<FWidgetBlueprintEditor> WeakEditor)
+{
+	if (TSharedPtr<FWidgetBlueprintEditor> Editor = WeakEditor.Pin())
+	{
+		UMVVMEditorSubsystem* Subsystem = GEditor->GetEditorSubsystem<UMVVMEditorSubsystem>();
+		if (!Subsystem)
+		{
+			return;
+		}
+
+		if (UMVVMBlueprintView* BlueprintView = Subsystem->GetView(Editor->GetWidgetBlueprintObj()))
+		{
+			Editor->CleanSelection();
+			TSet<UObject*> Selections;
+			Selections.Add(BlueprintView->GetSettings());
+			Editor->SelectObjects(Selections);
+		}
+	}
 }
 
+} //namespace UE::MVVM::Private
+
 /** */
 void SBindingsPanel::Construct(const FArguments& InArgs, TSharedPtr<FWidgetBlueprintEditor> WidgetBlueprintEditor, bool bInIsDrawerTab)
 {
@@ -425,11 +447,6 @@ void SBindingsPanel::HandleExtensionAdded(UBlueprintExtension* NewExtension)
 
 TSharedRef<SWidget> SBindingsPanel::GenerateSettingsMenu()
 {
-	UMVVMWidgetBlueprintExtension_View* MVVMExtensionPtr = MVVMExtension.Get();
-	DetailsView->SetObject(MVVMExtensionPtr && MVVMExtensionPtr->GetBlueprintView() ? MVVMExtensionPtr->GetBlueprintView() : nullptr);
-	DetailContainer->SetContent(DetailsView.ToSharedRef());
-	NotifyHook->Binding = FMVVMBlueprintViewBinding();
-
 	UWidgetBlueprintToolMenuContext* WidgetBlueprintMenuContext = NewObject<UWidgetBlueprintToolMenuContext>();
 	WidgetBlueprintMenuContext->WidgetBlueprintEditor = WeakBlueprintEditor;
 
@@ -441,6 +458,21 @@ TSharedRef<SWidget> SBindingsPanel::GenerateSettingsMenu()
 void SBindingsPanel::RegisterSettingsMenu()
 {
 	UToolMenu* Menu = UToolMenus::Get()->RegisterMenu("MVVM.ViewBindings.Toolbar");
+	FToolMenuSection& Section = Menu->FindOrAddSection("Settings");
+	Section.AddDynamicEntry("Settings", FNewToolMenuSectionDelegate::CreateLambda([](FToolMenuSection& InSection)
+		{
+			if (const UWidgetBlueprintToolMenuContext* Context = InSection.FindContext<UWidgetBlueprintToolMenuContext>())
+			{
+				InSection.AddMenuEntry(
+					"ViewSettings"
+					, LOCTEXT("ViewSettings", "View Settings")
+					, LOCTEXT("ViewSettingsTooltip", "View Settings")
+					, FSlateIcon(FMVVMEditorStyle::Get().GetStyleSetName(), "BlueprintView.TabIcon")
+					, FUIAction(FExecuteAction::CreateStatic(UE::MVVM::Private::SetSelectObjectsToViewSettings, Context->WidgetBlueprintEditor))
+					, EUserInterfaceActionType::Button
+				);
+			}
+		}));
 }
 
 
@@ -454,6 +486,12 @@ TSharedRef<SWidget> SBindingsPanel::GenerateEditViewWidget()
 		BindingsList = SNew(SBindingsList, StaticCastSharedRef<SBindingsPanel>(AsShared()), MVVMExtension.Get());
 	}
 
+	DetailContainer = SNew(SBorder)
+		.Visibility(EVisibility::Collapsed)
+		[
+			DetailsView.ToSharedRef()
+		];
+
 	TSharedPtr<SHorizontalBox> BindingPanelToolBar = SNew(SHorizontalBox);
 
 	FSlimHorizontalToolBarBuilder ToolbarBuilderGlobal(TSharedPtr<const FUICommandList>(), FMultiBoxCustomization::None);
@@ -545,6 +583,23 @@ TSharedRef<SWidget> SBindingsPanel::GenerateEditViewWidget()
 			);
 		}
 
+		ToolbarBuilderGlobal.AddWidget(
+			SNew(SComboButton)
+			.HasDownArrow(false)
+			.ContentPadding(0)
+			.ForegroundColor(FSlateColor::UseForeground())
+			.ButtonStyle(FAppStyle::Get(), "SimpleButton")
+			.MenuContent()
+			[
+				GenerateSettingsMenu()
+			]
+			.ButtonContent()
+			[
+				SNew(SImage)
+				.Image(FAppStyle::GetBrush("DetailsView.ViewOptions"))
+			]
+		);
+
 		ToolbarBuilderGlobal.EndSection();
 	}
 
@@ -647,11 +702,7 @@ TSharedRef<SWidget> SBindingsPanel::GenerateEditViewWidget()
 				+ SSplitter::Slot()
 				.Value(0.25f)
 				[
-					SAssignNew(DetailContainer, SBorder)
-					.Visibility(EVisibility::Collapsed)
-					[
-						DetailsView.ToSharedRef()
-					]
+					DetailContainer.ToSharedRef()
 				]
 			]
 		];
diff --git a/Engine/Plugins/Runtime/ModelViewViewModel/Source/ModelViewViewModelEditor/Private/Widgets/SMVVMViewModelPanel.cpp b/Engine/Plugins/Runtime/ModelViewViewModel/Source/ModelViewViewModelEditor/Private/Widgets/SMVVMViewModelPanel.cpp
--- a/Engine/Plugins/Runtime/ModelViewViewModel/Source/ModelViewViewModelEditor/Private/Widgets/SMVVMViewModelPanel.cpp
+++ b/Engine/Plugins/Runtime/ModelViewViewModel/Source/ModelViewViewModelEditor/Private/Widgets/SMVVMViewModelPanel.cpp
@@ -18,8 +18,11 @@
 
 #include "Framework/Commands/GenericCommands.h"
 #include "Framework/MultiBox/MultiBoxBuilder.h"
+#include "Styling/MVVMEditorStyle.h"
 
 #include "Widgets/Images/SImage.h"
+#include "Widgets/Input/SButton.h"
+#include "Widgets/Input/SComboButton.h"
 #include "Widgets/Layout/SBorder.h"
 #include "Widgets/Layout/SSplitter.h"
 #include "Widgets/PropertyViewer/SFieldIcon.h"
@@ -27,6 +30,7 @@
 #include "Widgets/SMVVMViewModelBindingListWidget.h"
 #include "Widgets/Text/SInlineEditableTextBlock.h"
 #include "SPositiveActionButton.h"
+#include "SWarningOrErrorBox.h"
 
 #define LOCTEXT_NAMESPACE "ViewModelPanel"
 
@@ -47,24 +51,70 @@ void UMVVMBlueprintViewModelContextWrapper::PostEditChangeProperty(struct FPrope
 	}
 }
 
+namespace UE::MVVM::Private
+{
+
+void SetSelectObjectsToView(TWeakPtr<FWidgetBlueprintEditor> WeakEditor)
+{
+	if (TSharedPtr<FWidgetBlueprintEditor> Editor = WeakEditor.Pin())
+	{
+		UMVVMEditorSubsystem* Subsystem = GEditor->GetEditorSubsystem<UMVVMEditorSubsystem>();
+		if (!Subsystem)
+		{
+			return;
+		}
+
+		if (UMVVMBlueprintView* BlueprintView = Subsystem->GetView(Editor->GetWidgetBlueprintObj()))
+		{
+			Editor->CleanSelection();
+			TSet<UObject*> Selections;
+			Selections.Add(BlueprintView->GetSettings());
+			Editor->SelectObjects(Selections);
+		}
+	}
+}
+
+} //namespace UE::MVVM::Private
+
 namespace UE::MVVM
 {
 
 void SMVVMViewModelPanel::RegisterMenu()
 {
-	UToolMenu* Menu = UToolMenus::Get()->RegisterMenu("MVVM.Viewmodels.Toolbar");
-	Menu->MenuType = EMultiBoxType::SlimHorizontalToolBar;
-	FToolMenuSection& Section = Menu->FindOrAddSection("Left");
-	Section.AddDynamicEntry("AddViewmodel", FNewToolMenuSectionDelegate::CreateLambda([](FToolMenuSection& InSection)
-		{
-			if (const UMVVMViewModelPanelToolMenuContext* Context = InSection.FindContext<UMVVMViewModelPanelToolMenuContext>())
+	{
+		UToolMenu* Menu = UToolMenus::Get()->RegisterMenu("MVVM.Viewmodels.Toolbar");
+		Menu->MenuType = EMultiBoxType::SlimHorizontalToolBar;
+		FToolMenuSection& Section = Menu->FindOrAddSection("Left");
+		Section.AddDynamicEntry("AddViewmodel", FNewToolMenuSectionDelegate::CreateLambda([](FToolMenuSection& InSection)
 			{
-				if (TSharedPtr<SMVVMViewModelPanel> ViewModelPanel = Context->ViewModelPanel.Pin())
+				if (const UMVVMViewModelPanelToolMenuContext* Context = InSection.FindContext<UMVVMViewModelPanelToolMenuContext>())
 				{
-					ViewModelPanel->BuildContextMenu(InSection);
+					if (TSharedPtr<SMVVMViewModelPanel> ViewModelPanel = Context->ViewModelPanel.Pin())
+					{
+						ViewModelPanel->BuildContextMenu(InSection);
+					}
 				}
-			}
-		}));
+			}));
+	}
+	{
+		UToolMenu* Menu = UToolMenus::Get()->RegisterMenu("MVVM.Viewmodels.Settings");
+		//Menu->MenuType = EMultiBoxType::Menu;
+		FToolMenuSection& Section = Menu->FindOrAddSection("Main");
+		Section.AddDynamicEntry("Main", FNewToolMenuSectionDelegate::CreateLambda([](FToolMenuSection& InSection)
+			{
+				if (const UWidgetBlueprintToolMenuContext* Context = InSection.FindContext<UWidgetBlueprintToolMenuContext>())
+				{
+					InSection.AddMenuEntry(
+						"ViewSettings"
+						, LOCTEXT("ViewSettings", "View Settings")
+						, LOCTEXT("ViewSettingsTooltip", "View Settings")
+						, FSlateIcon(FMVVMEditorStyle::Get().GetStyleSetName(), "BlueprintView.TabIcon")
+						, FUIAction(FExecuteAction::CreateStatic(UE::MVVM::Private::SetSelectObjectsToView, Context->WidgetBlueprintEditor))
+						, EUserInterfaceActionType::Button
+					);
+				}
+			}));
+	}
 }
 
 
@@ -133,25 +183,56 @@ void SMVVMViewModelPanel::Construct(const FArguments& InArgs, TSharedPtr<FWidget
 		.SearchBoxPreSlot()
 		[
 			UToolMenus::Get()->GenerateWidget("MVVM.Viewmodels.Toolbar", GenerateWidgetContext)
+		]
+		.SearchBoxPostSlot()
+		[
+			SNew(SComboButton)
+			.HasDownArrow(false)
+			.ContentPadding(0)
+			.ForegroundColor(FSlateColor::UseForeground())
+			.ButtonStyle(FAppStyle::Get(), "SimpleButton")
+			.MenuContent()
+			[
+				UToolMenus::Get()->GenerateWidget("MVVM.Viewmodels.Settings", GenerateWidgetContext)
+			]
+			.ButtonContent()
+			[
+				SNew(SImage)
+				.Image(FAppStyle::GetBrush("DetailsView.ViewOptions"))
+			]
 		];
 
 	FillViewModel();
 
 	ModelContextWrapper.Reset(NewObject<UMVVMBlueprintViewModelContextWrapper>());
 
-
 	ChildSlot
 	[
 		SNew(SBorder)
 		.BorderImage(FAppStyle::Get().GetBrush("Brushes.Recessed"))
 		.Padding(0)
 		[
-			SNew(SSplitter)
-			.Orientation(EOrientation::Orient_Vertical)
-			+ SSplitter::Slot()
+			SNew(SVerticalBox)
+			+ SVerticalBox::Slot()
+			.FillHeight(1.f)
 			[
 				ViewModelTreeView.ToSharedRef()
 			]
+			+ SVerticalBox::Slot()
+			.AutoHeight()
+			.Padding(FMargin(0.0f, 8.0f, 0.0f, 4.0f))
+			[
+				SNew(SWarningOrErrorBox)
+				.Visibility(this, &SMVVMViewModelPanel::GetWarningPanelVisibility)
+				.MessageStyle(EMessageStyle::Warning)
+				.Message(LOCTEXT("InitializationPanelWarningDescription", "The view will not initialize automatically. It was manually set in the View Settings."))
+				[
+					SNew(SButton)
+					.OnClicked(this, &SMVVMViewModelPanel::HandleDisableWarningPanel)
+					.TextStyle(FAppStyle::Get(), "DialogButtonText")
+					.Text(LOCTEXT("WarningDisable", "Dismiss"))
+				]
+			]
 		]
 	];
 }
@@ -236,6 +317,7 @@ void SMVVMViewModelPanel::HandleViewUpdated(UBlueprintExtension*)
 				if (ensure(CurrentBlueprintView))
 				{
 					ViewModelsUpdatedHandle = CurrentBlueprintView->OnViewModelsUpdated.AddSP(this, &SMVVMViewModelPanel::HandleViewModelsUpdated);
+
 					bViewUpdated = true;
 
 					WidgetBlueprint->OnExtensionAdded.Remove(ExtensionAddeddHandle);
@@ -692,6 +774,29 @@ bool SMVVMViewModelPanel::RenameViewModelProperty(FGuid ViewModelGuid, const FTe
 }
 
 
+EVisibility SMVVMViewModelPanel::GetWarningPanelVisibility() const
+{
+	if (!bDisableWarningPanel)
+	{
+		if (UMVVMBlueprintView* WidgetBlueprint = WeakBlueprintView.Get())
+		{
+			if (!WidgetBlueprint->GetSettings()->bInitializeSourcesOnConstruct || !WidgetBlueprint->GetSettings()->bInitializeBindingsOnConstruct)
+			{
+				return EVisibility::Visible;
+			}
+		}
+	}
+
+	return EVisibility::Collapsed;
+}
+
+
+FReply SMVVMViewModelPanel::HandleDisableWarningPanel()
+{
+	bDisableWarningPanel = true;
+	return FReply::Handled();
+}
+
 namespace Private
 {
 	static FName NAME_ViewModelName = TEXT("ViewModelName");
diff --git a/Engine/Plugins/Runtime/ModelViewViewModel/Source/ModelViewViewModelEditor/Private/Widgets/SMVVMViewModelPanel.h b/Engine/Plugins/Runtime/ModelViewViewModel/Source/ModelViewViewModelEditor/Private/Widgets/SMVVMViewModelPanel.h
--- a/Engine/Plugins/Runtime/ModelViewViewModel/Source/ModelViewViewModelEditor/Private/Widgets/SMVVMViewModelPanel.h
+++ b/Engine/Plugins/Runtime/ModelViewViewModel/Source/ModelViewViewModelEditor/Private/Widgets/SMVVMViewModelPanel.h
@@ -95,6 +95,9 @@ private:
 	void HandleRenameViewModel();
 	bool HandleCanRenameViewModel() const;
 
+	EVisibility GetWarningPanelVisibility() const;
+	FReply HandleDisableWarningPanel();
+
 	void CreateCommandList();
 	void FillViewModel();
 	bool RenameViewModelProperty(FGuid ViewModelGuid, const FText& RenameTo, bool bCommit, FText& OutErrorMessage) const;
@@ -124,6 +127,7 @@ private:
 	FGuid SelectedViewModelGuid;
 	FName PreviousViewModelPropertyName;
 	bool bIsViewModelSelecting = false;
+	bool bDisableWarningPanel = false;
 };
 
 } // namespace
