MVVM: Add viewmodel to the UMG preview mode.

diff --git a/Engine/Plugins/Runtime/ModelViewViewModel/Source/ModelViewViewModelEditor/Private/ModelViewViewModelEditorModule.cpp b/Engine/Plugins/Runtime/ModelViewViewModel/Source/ModelViewViewModelEditor/Private/ModelViewViewModelEditorModule.cpp
--- a/Engine/Plugins/Runtime/ModelViewViewModel/Source/ModelViewViewModelEditor/Private/ModelViewViewModelEditorModule.cpp
+++ b/Engine/Plugins/Runtime/ModelViewViewModel/Source/ModelViewViewModelEditor/Private/ModelViewViewModelEditorModule.cpp
@@ -19,6 +19,7 @@
 #include "PropertyEditorModule.h"
 #include "Styling/MVVMEditorStyle.h"
 #include "Tabs/MVVMBindingSummoner.h"
+#include "Tabs/MVVMPreviewSourceSummoner.h"
 #include "Tabs/MVVMViewModelSummoner.h"
 #include "ToolMenus.h"
 #include "UMGEditorModule.h"
@@ -120,6 +121,10 @@ void FModelViewViewModelEditorModule::HandleRegisterBlueprintEditorTab(const FWi
 			}
 		}
 	}
+	else if (ApplicationMode.GetModeName() == FWidgetBlueprintApplicationModes::PreviewMode)
+	{
+		TabFactories.RegisterFactory(MakeShared<UE::MVVM::FPreviewSourceSummoner>(ApplicationMode.GetBlueprintEditor()));
+	}
 }
 
 
diff --git a/Engine/Plugins/Runtime/ModelViewViewModel/Source/ModelViewViewModelEditor/Private/Tabs/MVVMPreviewSourceSummoner.cpp b/Engine/Plugins/Runtime/ModelViewViewModel/Source/ModelViewViewModelEditor/Private/Tabs/MVVMPreviewSourceSummoner.cpp
new file mode 100644
--- /dev/null
+++ b/Engine/Plugins/Runtime/ModelViewViewModel/Source/ModelViewViewModelEditor/Private/Tabs/MVVMPreviewSourceSummoner.cpp
@@ -0,0 +1,49 @@
+// Copyright Epic Games, Inc. All Rights Reserved.
+
+#include "Tabs/MVVMPreviewSourceSummoner.h"
+
+#include "Styling/MVVMEditorStyle.h"
+#include "WidgetBlueprintEditor.h"
+#include "Widgets/SMVVMPreviewSourcePanel.h"
+#include "Widgets/SNullWidget.h"
+
+#define LOCTEXT_NAMESPACE "DebugSourcePanel"
+
+namespace UE::MVVM
+{
+
+const FLazyName DebugSourcePanelTabID = "PreviewSourcePanel";
+
+FName FPreviewSourceSummoner::GetTabID()
+{
+	return DebugSourcePanelTabID.Resolve();
+}
+
+
+FPreviewSourceSummoner::FPreviewSourceSummoner(TSharedPtr<FWidgetBlueprintEditor> InBlueprintEditor)
+	: FWorkflowTabFactory(GetTabID(), InBlueprintEditor)
+	, WeakEditor(InBlueprintEditor)
+{
+	TabLabel = LOCTEXT("ViewmodelTabLabel", "Viewmodels");
+	TabIcon = FSlateIcon(FMVVMEditorStyle::Get().GetStyleSetName(), "BlueprintView.TabIcon");
+
+	bIsSingleton = true;
+
+	ViewMenuDescription = LOCTEXT("Viewmodel_Desc", "Viewmodels");
+	ViewMenuTooltip = LOCTEXT("Viewmodel_ToolTip", "Show the viewmodels panel");
+}
+
+
+TSharedRef<SWidget> FPreviewSourceSummoner::CreateTabBody(const FWorkflowTabSpawnInfo& Info) const
+{
+	if (TSharedPtr<FWidgetBlueprintEditor> BlueprintEditorPtr = WeakEditor.Pin())
+	{
+		return SNew(SPreviewSourcePanel, BlueprintEditorPtr)
+			.AddMetaData<FTagMetaData>(FTagMetaData(TEXT("PreviewSourcePanel")));
+	}
+	return SNullWidget::NullWidget;
+}
+
+} // namespace
+
+#undef LOCTEXT_NAMESPACE 
diff --git a/Engine/Plugins/Runtime/ModelViewViewModel/Source/ModelViewViewModelEditor/Private/Tabs/MVVMPreviewSourceSummoner.h b/Engine/Plugins/Runtime/ModelViewViewModel/Source/ModelViewViewModelEditor/Private/Tabs/MVVMPreviewSourceSummoner.h
new file mode 100644
--- /dev/null
+++ b/Engine/Plugins/Runtime/ModelViewViewModel/Source/ModelViewViewModelEditor/Private/Tabs/MVVMPreviewSourceSummoner.h
@@ -0,0 +1,27 @@
+// Copyright Epic Games, Inc. All Rights Reserved.
+
+#pragma once
+
+#include "WorkflowOrientedApp/WorkflowTabFactory.h"
+
+class FWidgetBlueprintEditor;
+class SWidget;
+
+namespace UE::MVVM
+{
+
+struct FPreviewSourceSummoner : public FWorkflowTabFactory
+{
+public:
+	static FName GetTabID();
+	
+public:
+	FPreviewSourceSummoner(TSharedPtr<FWidgetBlueprintEditor> InBlueprintEditor);
+	
+	virtual TSharedRef<SWidget> CreateTabBody(const FWorkflowTabSpawnInfo& Info) const override;
+	
+protected:
+	TWeakPtr<FWidgetBlueprintEditor> WeakEditor;
+};
+
+} // namespace
diff --git a/Engine/Plugins/Runtime/ModelViewViewModel/Source/ModelViewViewModelEditor/Private/Widgets/SMVVMPreviewSourcePanel.cpp b/Engine/Plugins/Runtime/ModelViewViewModel/Source/ModelViewViewModelEditor/Private/Widgets/SMVVMPreviewSourcePanel.cpp
new file mode 100644
--- /dev/null
+++ b/Engine/Plugins/Runtime/ModelViewViewModel/Source/ModelViewViewModelEditor/Private/Widgets/SMVVMPreviewSourcePanel.cpp
@@ -0,0 +1,197 @@
+// Copyright Epic Games, Inc. All Rights Reserved.
+
+#include "Widgets/SMVVMPreviewSourcePanel.h"
+
+#include "MVVMSubsystem.h"
+#include "View/MVVMView.h"
+
+#include "Preview/PreviewMode.h"
+#include "WidgetBlueprintEditor.h"
+
+#include "Styling/AppStyle.h"
+#include "Widgets/SBoxPanel.h"
+#include "Widgets/Layout/SBorder.h"
+#include "Widgets/PropertyViewer/SFieldIcon.h"
+#include "Widgets/Views/SListView.h"
+#include "Widgets/Text/STextBlock.h"
+
+#define LOCTEXT_NAMESPACE "SMVVMDebugSourcePanel"
+
+
+namespace UE::MVVM::Private
+{
+
+class SPreviewSourceView : public SListView<TSharedPtr<SPreviewSourceEntry>>
+{
+};
+
+class SPreviewSourceEntry
+{
+public:
+	SPreviewSourceEntry(UObject* InInstance, FName InName)
+		: WeakInstance(InInstance)
+		, Name(InName)
+	{}
+	
+	UClass* GetClass() const
+	{
+		UObject* Instance = WeakInstance.Get();
+		return Instance ? Instance->GetClass() : nullptr;
+	}
+
+	FText GetDisplayName() const
+	{
+		return FText::FromName(Name);
+	}
+
+	UObject* GetInstance() const
+	{
+		return WeakInstance.Get();
+	}
+
+private:
+	TWeakObjectPtr<UObject> WeakInstance;
+	FName Name;
+};
+
+} //namespace UE::MVVM::Private
+
+namespace UE::MVVM
+{
+
+void SPreviewSourcePanel::Construct(const FArguments& InArgs, TSharedPtr<FWidgetBlueprintEditor> Editor)
+{
+	WeakEditor = Editor;
+	check(Editor);
+
+	if (TSharedPtr<UE::UMG::Editor::FPreviewMode> Context = Editor->GetPreviewMode())
+	{
+		HandlePreviewWidgetChanged();
+		Context->OnPreviewWidgetChanged().AddSP(this, &SPreviewSourcePanel::HandlePreviewWidgetChanged);
+		Context->OnSelectedObjectChanged().AddSP(this, &SPreviewSourcePanel::HandleSelectedObjectChanged);
+	}
+
+	ChildSlot
+	[
+		SNew(SBorder)
+		.BorderImage(FAppStyle::GetBrush("ToolPanel.GroupBorder"))
+		.Padding(FMargin(4.f))
+		[
+			SAssignNew(SourceListView, Private::SPreviewSourceView)
+			.ListItemsSource(&SourceList)
+			.SelectionMode(ESelectionMode::Single)
+			.OnGenerateRow(this, &SPreviewSourcePanel::GenerateWidget)
+			.OnSelectionChanged(this, &SPreviewSourcePanel::HandleSourceSelectionChanged)
+		]
+	];
+}
+
+
+void SPreviewSourcePanel::HandlePreviewWidgetChanged()
+{
+	SourceList.Reset();
+
+	if (TSharedPtr<FWidgetBlueprintEditor> Editor = WeakEditor.Pin())
+	{
+		if (UUserWidget* NewWidget = Editor->GetPreviewMode() ? Editor->GetPreviewMode()->GetPreviewWidget() : nullptr)
+		{
+			if (UMVVMView* View = UMVVMSubsystem::GetViewFromUserWidget(NewWidget))
+			{
+				for (const FMVVMViewSource& Source : View->GetSources())
+				{
+					SourceList.Emplace(MakeShared<Private::SPreviewSourceEntry>(Source.Source, Source.SourceName));
+				}
+			}
+		}
+	}
+
+	if (SourceListView)
+	{
+		SourceListView->RequestListRefresh();
+	}
+}
+
+
+void SPreviewSourcePanel::HandleSelectedObjectChanged()
+{
+	if (bInternalSelection)
+	{
+		return;
+	}
+
+	TGuardValue<bool> TmpInternalSelection(bInternalSelection, true);
+
+	if (TSharedPtr<FWidgetBlueprintEditor> Editor = WeakEditor.Pin())
+	{
+		TArray<UObject*> NewSelections = Editor->GetPreviewMode() ? Editor->GetPreviewMode()->GetSelectedObjectList() : TArray<UObject*>();
+		if (NewSelections.Num() == 1)
+		{
+			TSharedPtr<Private::SPreviewSourceEntry>* Found = SourceList.FindByPredicate([ToFind = NewSelections[0]](const TSharedPtr<Private::SPreviewSourceEntry>& Other){ return Other->GetInstance() == ToFind; });
+			if (Found)
+			{
+				SourceListView->SetSelection(*Found);
+			}
+			else
+			{
+				SourceListView->SetSelection(TSharedPtr<Private::SPreviewSourceEntry>());
+			}
+		}
+		else
+		{
+			SourceListView->SetSelection(TSharedPtr<Private::SPreviewSourceEntry>());
+		}
+	}
+}
+
+
+void SPreviewSourcePanel::HandleSourceSelectionChanged(TSharedPtr<Private::SPreviewSourceEntry> Entry, ESelectInfo::Type SelectionType) const
+{
+	if (bInternalSelection)
+	{
+		return;
+	}
+
+	TGuardValue<bool> TmpInternalSelection(bInternalSelection, true);
+
+	if (TSharedPtr<FWidgetBlueprintEditor> Editor = WeakEditor.Pin())
+	{
+		if (Editor->GetPreviewMode())
+		{
+			TArray<TWeakObjectPtr<UObject>> Selection;
+			if (Entry && Entry->GetInstance())
+			{
+				Selection.Add(Entry->GetInstance());
+			}
+
+			Editor->GetPreviewMode()->SetSelectedObject(Selection);
+		}
+	}
+}
+
+
+TSharedRef<ITableRow> SPreviewSourcePanel::GenerateWidget(TSharedPtr<Private::SPreviewSourceEntry> Entry, const TSharedRef<STableViewBase>& OwnerTable) const
+{
+	typedef STableRow<TSharedPtr<Private::SPreviewSourceEntry>> RowType;
+
+	TSharedRef<RowType> NewRow = SNew(RowType, OwnerTable);
+	NewRow->SetContent(SNew(SHorizontalBox)
+		+ SHorizontalBox::Slot()
+		.AutoWidth()
+		.HAlign(HAlign_Right)
+		.VAlign(VAlign_Center)
+		[
+			SNew(UE::PropertyViewer::SFieldIcon, Entry->GetClass())
+		]
+		+ SHorizontalBox::Slot()
+		.Padding(4.0f)
+		[
+			SNew(STextBlock)
+			.Text(Entry->GetDisplayName())
+		]);
+
+	return NewRow;
+}
+
+} // namespace
+
+#undef LOCTEXT_NAMESPACE
diff --git a/Engine/Plugins/Runtime/ModelViewViewModel/Source/ModelViewViewModelEditor/Private/Widgets/SMVVMPreviewSourcePanel.h b/Engine/Plugins/Runtime/ModelViewViewModel/Source/ModelViewViewModelEditor/Private/Widgets/SMVVMPreviewSourcePanel.h
new file mode 100644
--- /dev/null
+++ b/Engine/Plugins/Runtime/ModelViewViewModel/Source/ModelViewViewModelEditor/Private/Widgets/SMVVMPreviewSourcePanel.h
@@ -0,0 +1,41 @@
+// Copyright Epic Games, Inc. All Rights Reserved.
+
+#pragma once
+
+#include "Widgets/SCompoundWidget.h"
+
+class STableViewBase;
+class ITableRow;
+class FWidgetBlueprintEditor;
+class UUserWidget;
+
+namespace UE::MVVM::Private { class SPreviewSourceView; }
+namespace UE::MVVM::Private { class SPreviewSourceEntry; }
+
+namespace UE::MVVM
+{
+
+/**
+ *
+ */
+class SPreviewSourcePanel : public SCompoundWidget
+{
+public:
+	SLATE_BEGIN_ARGS(SPreviewSourcePanel) {}
+	SLATE_END_ARGS()
+
+	void Construct(const FArguments& InArgs, TSharedPtr<FWidgetBlueprintEditor> Editor);
+
+private:
+	void HandlePreviewWidgetChanged();
+	void HandleSelectedObjectChanged();
+	void HandleSourceSelectionChanged(TSharedPtr<Private::SPreviewSourceEntry> Entry, ESelectInfo::Type SelectionType) const;
+	TSharedRef<ITableRow> GenerateWidget(TSharedPtr<Private::SPreviewSourceEntry> Entry, const TSharedRef<STableViewBase>& OwnerTable) const;
+
+	TWeakPtr<FWidgetBlueprintEditor> WeakEditor;
+	TSharedPtr<Private::SPreviewSourceView> SourceListView;
+	TArray<TSharedPtr<Private::SPreviewSourceEntry>> SourceList;
+	mutable bool bInternalSelection = false;
+};
+
+} // namespace
