MVVM: Fix bad format syntax.
#rnx

diff --git a/Engine/Plugins/Runtime/ModelViewViewModel/Source/ModelViewViewModelBlueprint/Private/Bindings/MVVMCompiledBindingLibraryCompiler.cpp b/Engine/Plugins/Runtime/ModelViewViewModel/Source/ModelViewViewModelBlueprint/Private/Bindings/MVVMCompiledBindingLibraryCompiler.cpp
--- a/Engine/Plugins/Runtime/ModelViewViewModel/Source/ModelViewViewModelBlueprint/Private/Bindings/MVVMCompiledBindingLibraryCompiler.cpp
+++ b/Engine/Plugins/Runtime/ModelViewViewModel/Source/ModelViewViewModelBlueprint/Private/Bindings/MVVMCompiledBindingLibraryCompiler.cpp
@@ -225,7 +225,7 @@ TValueOrError<FCompiledBindingLibraryCompiler::FFieldIdHandle, FText> FCompiledB
 		UE::FieldNotification::FFieldId FoundFieldId = ScriptObject->GetFieldNotificationDescriptor().GetField(SourceClass, FieldId);
 		if (!FoundFieldId.IsValid())
 		{
-			return MakeError(FText::Format(LOCTEXT("FieldNotifyNotSupported", "The FieldNotify '{0}' is not supported by '{0}'.")
+			return MakeError(FText::Format(LOCTEXT("FieldNotifyNotSupported", "The FieldNotify '{0}' is not supported by '{1}'.")
 				, FText::FromName(FieldId)
 				, SourceClass->GetDisplayNameText()));
 		}
