MVVM: Rework how nodes are read and created for complex conversion function.

diff --git a/Engine/Plugins/Runtime/ModelViewViewModel/Source/ModelViewViewModelBlueprint/Private/Bindings/MVVMConversionFunctionHelper.cpp b/Engine/Plugins/Runtime/ModelViewViewModel/Source/ModelViewViewModelBlueprint/Private/Bindings/MVVMConversionFunctionHelper.cpp
--- a/Engine/Plugins/Runtime/ModelViewViewModel/Source/ModelViewViewModelBlueprint/Private/Bindings/MVVMConversionFunctionHelper.cpp
+++ b/Engine/Plugins/Runtime/ModelViewViewModel/Source/ModelViewViewModelBlueprint/Private/Bindings/MVVMConversionFunctionHelper.cpp
@@ -239,6 +239,41 @@ UK2Node* GetWrapperNode(UEdGraph* Graph)
 
 namespace Private
 {
+	TArray<TTuple<UEdGraphNode*, UEdGraphPin*>> GetPropertyPathGraphNode(const UEdGraphPin* StartPin)
+	{
+		TArray<TTuple<UEdGraphNode*, UEdGraphPin*>> NodesInPath;
+
+		auto AddNode = [&NodesInPath](const UEdGraphPin* Pin)
+		{
+			UEdGraphNode* Result = nullptr;
+			if (Pin->Direction == EGPD_Input && Pin->LinkedTo.Num() == 1 && Pin->PinName != UEdGraphSchema_K2::PN_Execute)
+			{
+				Result = Pin->LinkedTo[0]->GetOwningNode();
+				NodesInPath.Emplace(Result, Pin->LinkedTo[0]);
+			}
+			return Result;
+		};
+
+		UEdGraphNode* CurrentNode = AddNode(StartPin);
+		while (CurrentNode)
+		{
+			TArray<UEdGraphPin*>& Pins = CurrentNode->Pins;
+			CurrentNode = nullptr;
+			for (UEdGraphPin* Pin : Pins)
+			{
+				if (UEdGraphNode* NewNode = AddNode(Pin))
+				{
+					CurrentNode = NewNode;
+					break;
+				}
+			}
+		}
+
+		Algo::Reverse(NodesInPath);
+		return NodesInPath;
+	}
+
+
 	FMVVMBlueprintPropertyPath GetPropertyPathForPin(const UBlueprint* Blueprint, const UEdGraphPin* StartPin, bool bSkipResolve)
 	{
 		if (StartPin->Direction != EGPD_Input || StartPin->PinName == UEdGraphSchema_K2::PN_Self)
@@ -252,95 +287,124 @@ namespace Private
 			return FMVVMBlueprintPropertyPath();
 		}
 
-		TArray<FMemberReference> PathParts;
-		TDeque<UEdGraphNode*> NodesToSearch;
-
-		if (StartPin->LinkedTo.Num() > 0)
-		{
-			if (UEdGraphNode* Node = StartPin->LinkedTo[0]->GetOwningNode())
-			{
-				NodesToSearch.PushLast(Node);
-			}
-		}
+		FMVVMBlueprintPropertyPath ResultPath;
 
-		while (NodesToSearch.Num() > 0)
+		auto AddRoot = [&ResultPath, BlueprintView, Blueprint, bSkipResolve](FMemberReference& Member)
 		{
-			UEdGraphNode* Node = NodesToSearch[0];
-			NodesToSearch.PopFirst();
-
-			if (UK2Node_VariableGet* Getter = Cast<UK2Node_VariableGet>(Node))
+			if (bSkipResolve)
 			{
-				PathParts.Insert(Getter->VariableReference, 0);
-			}
-			else if (UK2Node_CallFunction* Function = Cast<UK2Node_CallFunction>(Node))
-			{
-				PathParts.Insert(Function->FunctionReference, 0);
+				// if the generated class hasn't yet been generated we can blindly forge ahead and try to figure out if it's a widget or a viewmodel
+				if (const FMVVMBlueprintViewModelContext* ViewModel = BlueprintView->FindViewModel(Member.GetMemberName()))
+				{
+					ResultPath.SetViewModelId(ViewModel->GetViewModelId());
+				}
+				else
+				{
+					ResultPath.SetWidgetName(Member.GetMemberName());
+				}
 			}
-
-			for (UEdGraphPin* Pin : Node->Pins)
+			else
 			{
-				if (Pin->Direction == EGPD_Input && Pin->LinkedTo.Num() > 0)
+				if (const FObjectProperty* Property = CastField<FObjectProperty>(Member.ResolveMember<FProperty>(Blueprint->SkeletonGeneratedClass)))
 				{
-					NodesToSearch.PushLast(Pin->LinkedTo[0]->GetOwningNode());
+					if (Property->PropertyClass->IsChildOf<UWidget>() || Property->PropertyClass->IsChildOf<UBlueprint>())
+					{
+						ResultPath.SetWidgetName(Property->GetFName());
+					}
+					else if (Property->PropertyClass->ImplementsInterface(UNotifyFieldValueChanged::StaticClass()))
+					{
+						if (const FMVVMBlueprintViewModelContext* ViewModel = BlueprintView->FindViewModel(Property->GetFName()))
+						{
+							ResultPath.SetViewModelId(ViewModel->GetViewModelId());
+						}
+					}
 				}
 			}
-		}
+		};
 
-		if (PathParts.Num() == 0)
+		auto AddPropertyPath = [&ResultPath, Blueprint](FMemberReference& MemberReference)
 		{
-			return FMVVMBlueprintPropertyPath();
-		}
-
-		FMVVMBlueprintPropertyPath ResultPath;
-
-		FMemberReference RootRef = PathParts[0];
-		PathParts.RemoveAt(0);
-
-		if (bSkipResolve)
-		{
-			// if the generated class hasn't yet been generated we can blindly forge ahead and try to figure out if it's a widget or a viewmodel
-			if (const FMVVMBlueprintViewModelContext* ViewModel = BlueprintView->FindViewModel(RootRef.GetMemberName()))
+			if (UFunction* Function = MemberReference.ResolveMember<UFunction>(Blueprint->SkeletonGeneratedClass))
 			{
-				ResultPath.SetViewModelId(ViewModel->GetViewModelId());
+				ResultPath.AppendPropertyPath(Blueprint, UE::MVVM::FMVVMConstFieldVariant(Function));
 			}
-			else
+			else if (const FProperty* Property = MemberReference.ResolveMember<FProperty>(Blueprint->SkeletonGeneratedClass))
 			{
-				ResultPath.SetWidgetName(RootRef.GetMemberName());
+				ResultPath.AppendPropertyPath(Blueprint, UE::MVVM::FMVVMConstFieldVariant(Property));
 			}
-		}
-		else
+		};
+
+		auto AddBreakNode = [&ResultPath, Blueprint](UScriptStruct* Struct, FName PropertyName)
 		{
-			if (const FObjectProperty* Property = CastField<FObjectProperty>(RootRef.ResolveMember<FProperty>(Blueprint->SkeletonGeneratedClass)))
+			FProperty* FoundProperty = Struct->FindPropertyByName(PropertyName);
+			if (ensure(FoundProperty))
 			{
-				if (Property->PropertyClass->IsChildOf<UWidget>() || Property->PropertyClass->IsChildOf<UBlueprint>())
+				ResultPath.AppendPropertyPath(Blueprint, UE::MVVM::FMVVMConstFieldVariant(FoundProperty));
+			}
+		};
+
+		bool bFirst = true;
+		TArray<TTuple<UEdGraphNode*, UEdGraphPin*>> NodesToSearch = GetPropertyPathGraphNode(StartPin);
+		for (const TTuple<UEdGraphNode*, UEdGraphPin*>& NodePair : NodesToSearch)
+		{
+			UEdGraphNode* Node = NodePair.Get<UEdGraphNode*>();
+			if (UK2Node_VariableGet* GetNode = Cast<UK2Node_VariableGet>(Node))
+			{
+				if (bFirst)
+				{
+					AddRoot(GetNode->VariableReference);
+				}
+				else
 				{
-					ResultPath.SetWidgetName(Property->GetFName());
+					AddPropertyPath(GetNode->VariableReference);
 				}
-				else if (Property->PropertyClass->ImplementsInterface(UNotifyFieldValueChanged::StaticClass()))
+			}
+			else if (UK2Node_CallFunction* FunctionNode = Cast<UK2Node_CallFunction>(Node))
+			{
+				// UK2Node_CallFunction can be native break function
+				if (UFunction* Function = FunctionNode->FunctionReference.ResolveMember<UFunction>(Blueprint->SkeletonGeneratedClass))
 				{
-					if (const FMVVMBlueprintViewModelContext* ViewModel = BlueprintView->FindViewModel(Property->GetFName()))
+					bool bAddPropertyPath = true;
+					const FProperty* ArgumentProperty = UE::MVVM::BindingHelper::GetFirstArgumentProperty(Function);
+					const FStructProperty* ArgumentStructProperty = CastField<FStructProperty>(ArgumentProperty);
+					if (ArgumentStructProperty && ArgumentStructProperty->Struct)
 					{
-						ResultPath.SetViewModelId(ViewModel->GetViewModelId());
+						const FString& MetaData = ArgumentStructProperty->Struct->GetMetaData(FBlueprintMetadata::MD_NativeBreakFunction);
+						if (MetaData.Len() > 0)
+						{
+							if (const UFunction* NativeBreakFunction = FindObject<UFunction>(nullptr, *MetaData, true))
+							{
+								AddBreakNode(ArgumentStructProperty->Struct, NodePair.Get<UEdGraphPin*>()->GetFName());
+								bAddPropertyPath = false;
+							}
+						}
+					}
+
+					if (bAddPropertyPath)
+					{
+						ResultPath.AppendPropertyPath(Blueprint, UE::MVVM::FMVVMConstFieldVariant(Function));
 					}
 				}
 			}
-		}
-
-		if (PathParts.Num() == 0)
-		{
-			return ResultPath;
-		}
-
-		for (const FMemberReference& MemberReference : PathParts)
-		{
-			if (UFunction* Function = MemberReference.ResolveMember<UFunction>(Blueprint->SkeletonGeneratedClass)) 
+			else if (UK2Node_BreakStruct* StructNode = Cast<UK2Node_BreakStruct>(Node))
 			{
-				ResultPath.AppendPropertyPath(Blueprint, UE::MVVM::FMVVMConstFieldVariant(Function));
+				if (ensure(StructNode->StructType))
+				{
+					AddBreakNode(StructNode->StructType, NodePair.Get<UEdGraphPin*>()->GetFName());
+				}
 			}
-			else if (const FProperty* Property = MemberReference.ResolveMember<FProperty>(Blueprint->SkeletonGeneratedClass))
+			else if (UK2Node_Self* Self = Cast<UK2Node_Self>(Node))
 			{
-				ResultPath.AppendPropertyPath(Blueprint, UE::MVVM::FMVVMConstFieldVariant(Property));
+				if (bFirst)
+				{
+					ResultPath.SetWidgetName(Blueprint->GetFName());
+				}
+				else
+				{
+					ensure(false);
+				}
 			}
+			bFirst = false;
 		}
 
 		return ResultPath;
@@ -368,48 +432,26 @@ void SetPropertyPathForPin(const UBlueprint* Blueprint, const FMVVMBlueprintProp
 	UEdGraph* FunctionGraph = ConversionNode ? ConversionNode->GetGraph() : nullptr;
 	const UEdGraphSchema* Schema = FunctionGraph ? FunctionGraph->GetSchema() : nullptr;
 
+	UK2Node_FunctionEntry* ConverionFunctionEntry = ConversionNode ? Private::FindFunctionEntry(FunctionGraph) : nullptr;
+	UK2Node_FunctionResult* ConverionFunctionResult = ConversionNode ? Private::FindFunctionResult(FunctionGraph) : nullptr;
+
 	// Remove previous nodes
 	{
-		TArray<UEdGraphPin*> PinsToClear;
-		TArray<UEdGraphNode*> NodesToRemove;
-		PinsToClear.Add(PathPin);
-		while (PinsToClear.Num() > 0)
+		TArray<TTuple<UEdGraphNode*, UEdGraphPin*>> AllNodesForPath = Private::GetPropertyPathGraphNode(PathPin);
+		for (const TTuple<UEdGraphNode*, UEdGraphPin*>& Pair : AllNodesForPath)
 		{
-			UEdGraphPin* Pin = PinsToClear[0];
-			PinsToClear.RemoveAt(0);
-
-			PinsToClear.Append(Pin->LinkedTo);
-			Pin->BreakAllPinLinks();
-
-			// rename the node out of the way
-			UEdGraphNode* Node = Pin->GetOwningNode();
-			if (Node != ConversionNode)
-			{
-				NodesToRemove.AddUnique(Node);
-			}
+			UEdGraphNode* Node = Pair.Get<UEdGraphNode*>();
+			FunctionGraph->RemoveNode(Node, true);
 		}
+	}
 
-		if (FunctionGraph)
-		{
-			for (UEdGraphNode* Node : NodesToRemove)
-			{
-				// Break the links, and relink to the other nodes.
-				if (UK2Node_CallFunction* CallFunctionNode = Cast<UK2Node_CallFunction>(Node))
-				{
-					UEdGraphPin* ThenPin = CallFunctionNode->GetThenPin();
-					UEdGraphPin* ExecPin = CallFunctionNode->GetExecPin();
-					if (ThenPin && ThenPin->LinkedTo.Num() == 1)
-					{
-						Schema->MovePinLinks(*ExecPin, *ThenPin->LinkedTo[0], false);
-					}
-				}
-				FunctionGraph->RemoveNode(Node, true);
-			}
-		}
+	if (!FunctionGraph || !ConversionNode || !ConverionFunctionEntry || !ConverionFunctionResult)
+	{
+		return;
 	}
 
 	// Add new nodes
-	if (FunctionGraph && ConversionNode && !Path.IsEmpty())
+	if (!Path.IsEmpty())
 	{
 		const int32 ArgumentIndex = ConversionNode->Pins.IndexOfByPredicate([PathPin](const UEdGraphPin* Other){ return Other == PathPin; });
 		if (!ensure(ConversionNode->Pins.IsValidIndex(ArgumentIndex)))
@@ -421,15 +463,6 @@ void SetPropertyPathForPin(const UBlueprint* Blueprint, const FMVVMBlueprintProp
 		float PosX = ConversionNode->NodePosX - 300 * (Fields.Num() + 1);
 		const float PosY = ConversionNode->NodePosY + ArgumentIndex * 100;
 
-		UEdGraphPin* PreviousExecPin = nullptr;
-		{
-			UEdGraphNode* FunctionEntry = Private::FindFunctionEntry(FunctionGraph);
-			if (!ensure(FunctionEntry))
-			{
-				return;
-			}
-			PreviousExecPin = FunctionEntry->FindPinChecked(UEdGraphSchema_K2::PN_Then);
-		}
 		UEdGraphPin* PreviousDataPin = nullptr;
 		UClass* PreviousClass = nullptr;
 
@@ -449,12 +482,6 @@ void SetPropertyPathForPin(const UBlueprint* Blueprint, const FMVVMBlueprintProp
 					const FMVVMBlueprintViewModelContext* Context = View->FindViewModel(Path.GetViewModelId());
 					RootProperty = Blueprint->SkeletonGeneratedClass->FindPropertyByName(Context->GetViewModelName());
 				}
-
-				if (RootProperty == nullptr)
-				{
-					ensureMsgf(bCreateSelfNodeForRootProperty, TEXT("Could not resolve root property!"));
-					return;
-				}
 			}
 
 			if (bCreateSelfNodeForRootProperty)
@@ -470,6 +497,12 @@ void SetPropertyPathForPin(const UBlueprint* Blueprint, const FMVVMBlueprintProp
 			}
 			else
 			{
+				if (RootProperty == nullptr)
+				{
+					ensureMsgf(false, TEXT("Could not resolve root property!"));
+					return;
+				}
+
 				FGraphNodeCreator<UK2Node_VariableGet> RootGetterCreator(*FunctionGraph);
 				UK2Node_VariableGet* RootGetterNode = RootGetterCreator.CreateNode();
 				RootGetterNode->NodePosX = PosX;
@@ -485,32 +518,86 @@ void SetPropertyPathForPin(const UBlueprint* Blueprint, const FMVVMBlueprintProp
 		}
 
 		// create all the subsequent nodes in the path
+		const FProperty* PreviousProperty = nullptr;
 		for (int32 Index = 0; Index < Fields.Num(); ++Index)
 		{
-			if (!ensureMsgf(PreviousClass != nullptr, TEXT("Previous class not set!")))
-			{
-				return;
-			}
-
 			UEdGraphNode* NewNode = nullptr;
 			const UE::MVVM::FMVVMConstFieldVariant& Field = Fields[Index];
-			const bool bLastField = (Index == Fields.Num() - 1);
+
+			auto CanNewConnections = [Schema](UEdGraphPin* Pin, UEdGraphPin* PreviousDataPin, const UClass* Context)
+			{
+				return Pin->Direction == EGPD_Input
+					&& Pin->PinName != UEdGraphSchema_K2::PN_Execute
+					&& Schema->ArePinsCompatible(PreviousDataPin, Pin, Context);
+			};
+
+			auto FindNewOutputPin = [](UEdGraphNode* NewNode) -> UEdGraphPin*
+			{
+				// then update our previous pin pointers
+				for (UEdGraphPin* Pin : NewNode->Pins)
+				{
+					if (Pin->Direction == EGPD_Output)
+					{
+						if (Pin->PinName != UEdGraphSchema_K2::PN_Then)
+						{
+							return Pin;
+						}
+					}
+				}
+				return nullptr;
+			};
+
+			UEdGraphPin* NewPreviousDataPin = nullptr;
 			if (Field.IsProperty())
 			{
 				const FProperty* Property = Field.GetProperty();
 
-				// for structs in the middle of a path, we need to use a break node
-				const FStructProperty* StructProperty = CastField<FStructProperty>(Property);
-				if (!bLastField && StructProperty != nullptr)
+				// for struct in the middle of a path, we need to use a break node
+				if (const FStructProperty* PreviousStructProperty = CastField<FStructProperty>(PreviousProperty))
 				{
-					FGraphNodeCreator<UK2Node_BreakStruct> BreakCreator(*FunctionGraph);
-					UK2Node_BreakStruct* BreakNode = BreakCreator.CreateNode();
-					BreakNode->AllocateDefaultPins();
-					BreakCreator.Finalize();
+					const FString& MetaData = PreviousStructProperty->Struct->GetMetaData(FBlueprintMetadata::MD_NativeBreakFunction);
+					if (MetaData.Len() > 0)
+					{
+						const UFunction* Function = FindObject<UFunction>(nullptr, *MetaData, true);
+						ensure(Function);
+						FGraphNodeCreator<UK2Node_CallFunction> MakeStructCreator(*FunctionGraph);
+						UK2Node_CallFunction* FunctionNode = MakeStructCreator.CreateNode(false);
+						FunctionNode->SetFromFunction(Function);
+						MakeStructCreator.Finalize();
+
+						NewNode = FunctionNode;
+						PreviousProperty = UE::MVVM::BindingHelper::GetReturnProperty(Function);
+					}
+					else
+					{
+						FGraphNodeCreator<UK2Node_BreakStruct> BreakCreator(*FunctionGraph);
+						UK2Node_BreakStruct* BreakNode = BreakCreator.CreateNode();
+						BreakNode->StructType = PreviousStructProperty->Struct;
+						BreakNode->AllocateDefaultPins();
+						BreakCreator.Finalize();
+
+						NewNode = BreakNode;
+						PreviousProperty = Property;
+					}
 
-					NewNode = BreakNode;
+					NewPreviousDataPin = nullptr;
+					for (UEdGraphPin* Pin : NewNode->Pins)
+					{
+						if (Pin->Direction == EGPD_Output)
+						{
+							if (Pin->PinName == Property->GetFName())
+							{
+								NewPreviousDataPin = Pin;
+								break;
+							}
+						}
+					}
 
-					PreviousClass = nullptr;
+					if (!NewPreviousDataPin)
+					{
+						ensure(false);
+						return;
+					}
 				}
 				else if (PreviousClass != nullptr)
 				{
@@ -521,18 +608,18 @@ void SetPropertyPathForPin(const UBlueprint* Blueprint, const FMVVMBlueprintProp
 					GetterCreator.Finalize();
 
 					NewNode = GetterNode;
-					if (const FObjectProperty* ObjectProperty = CastField<FObjectProperty>(Property))
-					{
-						PreviousClass = ObjectProperty->PropertyClass;
-					}
-					else
-					{
-						PreviousClass = nullptr;
-					}
+					PreviousProperty = Property;
+					NewPreviousDataPin = FindNewOutputPin(NewNode);
 				}
 			}
 			else if (Field.IsFunction())
 			{
+				if (CastField<FStructProperty>(PreviousProperty))
+				{
+					ensure(false);
+					return;
+				}
+
 				const UFunction* Function = Field.GetFunction();
 
 				FGraphNodeCreator<UK2Node_CallFunction> CallFunctionCreator(*FunctionGraph);
@@ -542,16 +629,8 @@ void SetPropertyPathForPin(const UBlueprint* Blueprint, const FMVVMBlueprintProp
 				CallFunctionCreator.Finalize();
 
 				NewNode = FunctionNode;
-
-				const FProperty* ReturnProperty = UE::MVVM::BindingHelper::GetReturnProperty(Function);
-				if (const FObjectProperty* ObjectProperty = CastField<FObjectProperty>(ReturnProperty))
-				{
-					PreviousClass = ObjectProperty->PropertyClass;
-				}
-				else
-				{
-					PreviousClass = nullptr;
-				}
+				PreviousProperty = UE::MVVM::BindingHelper::GetReturnProperty(Function);
+				NewPreviousDataPin = FindNewOutputPin(NewNode);
 			}
 			else
 			{
@@ -559,37 +638,26 @@ void SetPropertyPathForPin(const UBlueprint* Blueprint, const FMVVMBlueprintProp
 				return;
 			}
 
-			// create new connections
+			// create new data connections
 			for (UEdGraphPin* Pin : NewNode->Pins)
 			{
-				if (Pin->Direction == EGPD_Input)
+				if (CanNewConnections(Pin, PreviousDataPin, PreviousClass))
 				{
-					if (Pin->PinName == UEdGraphSchema_K2::PN_Execute)
-					{
-						Schema->MovePinLinks(*PreviousExecPin, *Pin, false);
-					}
-					else if (Schema->ArePinsCompatible(PreviousDataPin, Pin, PreviousClass))
-					{
-						Pin->MakeLinkTo(PreviousDataPin);
-					}
+					Pin->MakeLinkTo(PreviousDataPin);
 				}
 			}
 
-			// then update our previous pin pointers
-			for (UEdGraphPin* Pin : NewNode->Pins)
+			if (const FObjectProperty* ObjectProperty = CastField<FObjectProperty>(PreviousProperty))
 			{
-				if (Pin->Direction == EGPD_Output)
-				{
-					if (Pin->PinName == UEdGraphSchema_K2::PN_Then)
-					{
-						PreviousExecPin = Pin;
-					}
-					else
-					{
-						PreviousDataPin = Pin;
-					}
-				}
+				PreviousClass = ObjectProperty->PropertyClass;
 			}
+			else
+			{
+				PreviousClass = nullptr;
+			}
+
+			// then update our previous pin pointers
+			PreviousDataPin = NewPreviousDataPin;
 
 			NewNode->NodePosX = PosX;
 			NewNode->NodePosY = PosY;
@@ -597,14 +665,38 @@ void SetPropertyPathForPin(const UBlueprint* Blueprint, const FMVVMBlueprintProp
 			PosX += 300;
 		}
 
-		// finish by linking to our ultimate destinations
-		UEdGraphPin* LastExecPin = ConversionNode->FindPin(UEdGraphSchema_K2::PN_Execute);
-		if (LastExecPin == nullptr)
+		// Link the last data pin to the Conversation Function Pin
+		PreviousDataPin->MakeLinkTo(PathPin);
+	}
+
+	// Link Then / Exec pin
+	{
+		check(ConversionNode && ConverionFunctionEntry && ConverionFunctionResult);
+
+		UEdGraphPin* ThenPin = ConverionFunctionEntry->FindPinChecked(UEdGraphSchema_K2::PN_Then);
+		check(ThenPin);
+		ThenPin->BreakAllPinLinks();
+
+		for (UEdGraphPin* Pin : ConversionNode->Pins)
 		{
-			LastExecPin = Private::FindFunctionResult(FunctionGraph)->FindPin(UEdGraphSchema_K2::PN_Execute);
+			TArray<TTuple<UEdGraphNode*, UEdGraphPin*>> AllNodesForPath = Private::GetPropertyPathGraphNode(Pin);
+			for (const TTuple<UEdGraphNode*, UEdGraphPin*>& Pair : AllNodesForPath)
+			{
+				UEdGraphNode* PathNode = Pair.Get<UEdGraphNode*>();
+				if (UK2Node_CallFunction* CallFunction = Cast<UK2Node_CallFunction>(PathNode))
+				{
+					// if it not a pure node
+					if (UEdGraphPin* ExecPin = CallFunction->FindPin(UEdGraphSchema_K2::PN_Execute))
+					{
+						ThenPin->BreakAllPinLinks();
+						ThenPin->MakeLinkTo(ExecPin);
+						ThenPin = CallFunction->FindPinChecked(UEdGraphSchema_K2::PN_Then);
+					}
+				}
+			}
 		}
-		PreviousExecPin->MakeLinkTo(LastExecPin);
-		PreviousDataPin->MakeLinkTo(PathPin);
+		UEdGraphPin* LastExecPin = ConverionFunctionResult->FindPinChecked(UEdGraphSchema_K2::PN_Execute);
+		ThenPin->MakeLinkTo(LastExecPin);
 	}
 }
 
