MVVM: Fix logging when the CompiledLibrary is not loaded. Output a message when the source creation is set to Manual in the editor. Nothing is done at runtime but a message is still nice.

diff --git a/Engine/Plugins/Runtime/ModelViewViewModel/Source/ModelViewViewModel/Private/View/MVVMViewClass.cpp b/Engine/Plugins/Runtime/ModelViewViewModel/Source/ModelViewViewModel/Private/View/MVVMViewClass.cpp
--- a/Engine/Plugins/Runtime/ModelViewViewModel/Source/ModelViewViewModel/Private/View/MVVMViewClass.cpp
+++ b/Engine/Plugins/Runtime/ModelViewViewModel/Source/ModelViewViewModel/Private/View/MVVMViewClass.cpp
@@ -244,6 +244,10 @@ FString FMVVMViewClass_SourceCreator::ToString(const FMVVMCompiledBindingLibrary
 				StringBuilder << FieldPathString.GetError();
 			}
 		}
+		else
+		{
+			StringBuilder << TEXT("None"); // The setting is Manual. At runtime, do nothing.
+		}
 	}
 
 	return StringBuilder.ToString();
@@ -568,7 +572,7 @@ void UMVVMViewClass::HandleBlueprintCompiled()
 #if UE_WITH_MVVM_DEBUGGING
 void UMVVMViewClass::Log(FMVVMViewClass_SourceCreator::FToStringArgs SourceArgs, FMVVMViewClass_CompiledBinding::FToStringArgs BindingArgs) const
 {
-	bool bPreviousLoad = ViewCounter >= 0;
+	bool bPreviousLoad = ViewCounter > 0;
 	if (!bPreviousLoad)
 	{
 		const_cast<UMVVMViewClass*>(this)->BindingLibrary.Load();
