MVVM: Fix bad order in ToString introduced with CL26629088

diff --git a/Engine/Plugins/Runtime/ModelViewViewModel/Source/ModelViewViewModel/Private/View/MVVMViewClass.cpp b/Engine/Plugins/Runtime/ModelViewViewModel/Source/ModelViewViewModel/Private/View/MVVMViewClass.cpp
--- a/Engine/Plugins/Runtime/ModelViewViewModel/Source/ModelViewViewModel/Private/View/MVVMViewClass.cpp
+++ b/Engine/Plugins/Runtime/ModelViewViewModel/Source/ModelViewViewModel/Private/View/MVVMViewClass.cpp
@@ -427,10 +427,6 @@ FString FMVVMViewClass_CompiledBinding::ToString(const FMVVMCompiledBindingLibra
 			StringBuilder << TEXT(", FieldId: ");
 			StringBuilder << (SourceFieldId.HasValue() ? SourceFieldId.GetValue().GetName() : FName());
 		}
-		else
-		{
-			ensureMsgf(IsOneTime(), TEXT("The field Id should be valid."));
-		}
 	}
 
 	if (Args.bAddFlags)
diff --git a/Engine/Plugins/Runtime/ModelViewViewModel/Source/ModelViewViewModelBlueprint/Private/MVVMBlueprintViewBinding.cpp b/Engine/Plugins/Runtime/ModelViewViewModel/Source/ModelViewViewModelBlueprint/Private/MVVMBlueprintViewBinding.cpp
--- a/Engine/Plugins/Runtime/ModelViewViewModel/Source/ModelViewViewModelBlueprint/Private/MVVMBlueprintViewBinding.cpp
+++ b/Engine/Plugins/Runtime/ModelViewViewModel/Source/ModelViewViewModelBlueprint/Private/MVVMBlueprintViewBinding.cpp
@@ -301,8 +301,6 @@ namespace UE::MVVM::Private
 							NameBuilder << TEXT(", ");
 						}
 
-
-
 						if (Pin.UsedPathAsValue())
 						{
 							AddPath(Pin.GetPath());
@@ -326,7 +324,7 @@ namespace UE::MVVM::Private
 		}
 		else
 		{
-			const FMVVMBlueprintPropertyPath& PropertyPath = bIsSource ? Binding.DestinationPath : Binding.SourcePath;
+			const FMVVMBlueprintPropertyPath& PropertyPath = bIsSource ? Binding.SourcePath : Binding.DestinationPath;
 			AddPath(PropertyPath);
 		}
 
