MVVM: Fix conversion function with multiples input paths.
Bindings will only be executed once at initialization. The bug happend when more than once source can activate the same binding.
Binding can now be OneWay if at least one argument is a FieldNotify. Otherwise it needs to be a OneTime.
Binding that has hardcoded value, no paths, are now supported.

diff --git a/Engine/Plugins/Runtime/ModelViewViewModel/Source/ModelViewViewModelBlueprint/Private/MVVMViewBlueprintCompiler.cpp b/Engine/Plugins/Runtime/ModelViewViewModel/Source/ModelViewViewModelBlueprint/Private/MVVMViewBlueprintCompiler.cpp
--- a/Engine/Plugins/Runtime/ModelViewViewModel/Source/ModelViewViewModelBlueprint/Private/MVVMViewBlueprintCompiler.cpp
+++ b/Engine/Plugins/Runtime/ModelViewViewModel/Source/ModelViewViewModelBlueprint/Private/MVVMViewBlueprintCompiler.cpp
@@ -259,10 +259,7 @@ void FMVVMViewBlueprintCompiler::CreateVariables(const FWidgetBlueprintCompilerC
 		if (NewProperty != nullptr)
 		{
 			NewProperty->SetPropertyFlags(CPF_BlueprintVisible | CPF_RepSkip | CPF_Transient | CPF_DuplicateTransient);
-			if (SourceContext.BlueprintSetter.IsEmpty())
-			{
-				NewProperty->SetPropertyFlags(CPF_BlueprintReadOnly);
-			}
+			NewProperty->SetPropertyFlags(CPF_BlueprintReadOnly);
 			NewProperty->SetPropertyFlags(SourceContext.bExposeOnSpawn ? CPF_ExposeOnSpawn : CPF_DisableEditOnInstance);
 
 #if WITH_EDITOR
@@ -552,33 +549,37 @@ void FMVVMViewBlueprintCompiler::CreateSourceLists(const FWidgetBlueprintCompile
 
 		if (bIsForwardBinding || bIsBackwardBinding)
 		{
-			const UMVVMBlueprintViewConversionFunction* ConversionFunction = Binding.Conversion.GetConversionFunction(bIsForwardBinding);
-			if (ConversionFunction && ConversionFunction->GetPins().Num() > 0)
+			UMVVMBlueprintViewConversionFunction* ConversionFunction = Binding.Conversion.GetConversionFunction(bIsForwardBinding);
+			if (ConversionFunction && ConversionFunction->NeedsWrapperGraph())
 			{
-				if (Binding.BindingType == EMVVMBindingMode::TwoWay)
+				ConversionFunction->SavePinValues(Context.GetWidgetBlueprint());
+				if (ConversionFunction->GetPins().Num() > 0)
 				{
-					Self->AddMessageForBinding(Binding, BlueprintView, LOCTEXT("TwoWayBindingsWithConversion", "Two-way bindings are not allowed to use conversion functions."), EBindingMessageType::Error);
-					bAreSourceContextsValid = false;
-					continue;
-				}
+					if (Binding.BindingType == EMVVMBindingMode::TwoWay)
+					{
+						Self->AddMessageForBinding(Binding, BlueprintView, LOCTEXT("TwoWayBindingsWithConversion", "Two-way bindings are not allowed to use conversion functions."), EBindingMessageType::Error);
+						bAreSourceContextsValid = false;
+						continue;
+					}
 
-				// generate sources for conversion function arguments
-				for (const FMVVMBlueprintPin& Pin: ConversionFunction->GetPins())
-				{
-					if (Pin.UsedPathAsValue())
+					// generate sources for conversion function arguments
+					for (const FMVVMBlueprintPin& Pin: ConversionFunction->GetPins())
 					{
-						bAreSourceContextsValid &= GenerateCompilerSourceContext(Pin.GetPath(), Pin.GetName());
+						if (Pin.UsedPathAsValue())
+						{
+							bAreSourceContextsValid &= GenerateCompilerSourceContext(Pin.GetPath(), Pin.GetName());
+						}
 					}
-				}
 
-				// generate destination source
-				if (bIsForwardBinding)
-				{
-					bAreSourceContextsValid &= GenerateCompilerSourceContext(Binding.DestinationPath);
-				}
-				else
-				{
-					bAreSourceContextsValid &= GenerateCompilerSourceContext(Binding.SourcePath);
+					// generate destination source
+					if (bIsForwardBinding)
+					{
+						bAreSourceContextsValid &= GenerateCompilerSourceContext(Binding.DestinationPath);
+					}
+					else
+					{
+						bAreSourceContextsValid &= GenerateCompilerSourceContext(Binding.SourcePath);
+					}
 				}
 			}
 			else
@@ -833,6 +834,8 @@ bool FMVVMViewBlueprintCompiler::PreCompile(UWidgetBlueprintGeneratedClass* Clas
 	const int32 NumBindings = BlueprintView->GetNumBindings();
 	CompilerBindings.Reset(NumBindings*2);
 	BindingSourceContexts.Reset(NumBindings*2);
+	SimpleBindingContexts.Reset();
+	ComplexConversionFunctionContexts.Reset();
 
 	FPropertyEditorPermissionList::Get().AddPermissionList(Class, FNamePermissionList(), EPropertyPermissionListRules::AllowListAllProperties);
 
@@ -907,7 +910,7 @@ bool FMVVMViewBlueprintCompiler::PreCompileBindingSources(UWidgetBlueprintGenera
 
 		bool bIsOneTimeBinding = IsOneTimeBinding(Binding.BindingType);
 
-		auto CreateSourceContextForPropertyPath = [this, &Binding, BlueprintView, Class, Index, bIsOneTimeBinding](const FMVVMBlueprintPropertyPath& Path, bool bForwardBinding, FName ArgumentName) -> bool
+		auto CreateSourceContextForPropertyPath = [this, &Binding, BlueprintView, Class, Index, bIsOneTimeBinding](const FMVVMBlueprintPropertyPath& Path, bool bForwardBinding, int32 ComplexConversionFunctionContextIndex, FName ArgumentName) -> bool
 		{
 			const TValueOrError<FBindingSourceContext, FText> CreatedBindingSourceContext = CreateBindingSourceContext(BlueprintView, Class, Path, bIsOneTimeBinding);
 			if (CreatedBindingSourceContext.HasError())
@@ -951,40 +954,66 @@ bool FMVVMViewBlueprintCompiler::PreCompileBindingSources(UWidgetBlueprintGenera
 
 			BindingSourceContext.BindingIndex = Index;
 			BindingSourceContext.bIsForwardBinding = bForwardBinding;
-			BindingSourceContext.bIsComplexBinding = !ArgumentName.IsNone();
+			BindingSourceContext.ComplexConversionFunctionContextIndex = ComplexConversionFunctionContextIndex;
 
 			this->BindingSourceContexts.Add(MoveTemp(BindingSourceContext));
 			return true;
 		};
 
 		enum class ECreateSourcesForConversionFunctionResult : uint8 { Valid, Failed, Continue };
-		auto CreateSourcesForConversionFunction = [this, &Binding, BlueprintView, &CreateSourceContextForPropertyPath](bool bForwardBinding)
+		auto CreateSourcesForConversionFunction = [this, &Binding, BlueprintView, Index, &CreateSourceContextForPropertyPath](bool bForwardBinding)
 		{
 			ECreateSourcesForConversionFunctionResult Result = ECreateSourcesForConversionFunctionResult::Continue;
 			if (const UMVVMBlueprintViewConversionFunction* ConversionFunction = Binding.Conversion.GetConversionFunction(bForwardBinding))
 			{
-				for (const FMVVMBlueprintPin& Pin : ConversionFunction->GetPins())
+				if (ConversionFunction->NeedsWrapperGraph())
 				{
-					if (Pin.UsedPathAsValue())
+					const int32 ComplexConversionFunctionContextIndex = ComplexConversionFunctionContexts.AddDefaulted();
 					{
-						if (Pin.GetName().IsNone())
-						{
-							AddMessageForBinding(Binding, BlueprintView,
-								FText::Format(LOCTEXT("InvalidArgumentPathName", "The conversion function {0} has an invalid argument"), FText::FromString(Binding.GetDisplayNameString(WidgetBlueprintCompilerContext.WidgetBlueprint()))),
-								EBindingMessageType::Error
-							);
-							Result = ECreateSourcesForConversionFunctionResult::Failed;
-						}
-						else if (CreateSourceContextForPropertyPath(Pin.GetPath(), bForwardBinding, Pin.GetName()))
-						{
-							Result = ECreateSourcesForConversionFunctionResult::Valid;
-						}
-						else
+						FComplexConversionFunctionContext& NewConversionFunctionContext = ComplexConversionFunctionContexts[ComplexConversionFunctionContextIndex];
+						NewConversionFunctionContext.BindingIndex = Index;
+						NewConversionFunctionContext.bIsForwardBinding = bForwardBinding;
+					}
+
+					for (const FMVVMBlueprintPin& Pin : ConversionFunction->GetPins())
+					{
+						if (Pin.UsedPathAsValue())
 						{
-							Result = ECreateSourcesForConversionFunctionResult::Failed;
-							break;
+							if (Pin.GetName().IsNone())
+							{
+								AddMessageForBinding(Binding, BlueprintView,
+									FText::Format(LOCTEXT("InvalidArgumentPathName", "The conversion function {0} has an invalid argument"), FText::FromString(Binding.GetDisplayNameString(WidgetBlueprintCompilerContext.WidgetBlueprint()))),
+									EBindingMessageType::Error
+								);
+								Result = ECreateSourcesForConversionFunctionResult::Failed;
+							}
+							else
+							{
+								const FMVVMBlueprintPropertyPath& Path = Pin.GetPath();
+								if (CreateSourceContextForPropertyPath(Path, bForwardBinding, ComplexConversionFunctionContextIndex, Pin.GetName()))
+								{
+									Result = ECreateSourcesForConversionFunctionResult::Valid;
+								}
+								else
+								{
+									Result = ECreateSourcesForConversionFunctionResult::Failed;
+									break;
+								}
+							}
 						}
 					}
+
+					if (Result == ECreateSourcesForConversionFunctionResult::Continue)
+					{
+						// The bindings doesn't have a path but could be onetime with hard codded value.
+						FSimpleBindingContext& NewBindingContext = SimpleBindingContexts.AddDefaulted_GetRef();
+						NewBindingContext.BindingIndex = Index;
+						NewBindingContext.bIsForwardBinding = bForwardBinding;
+						NewBindingContext.ComplexConversionFunctionContextIndex = ComplexConversionFunctionContextIndex;
+						ComplexConversionFunctionContexts[ComplexConversionFunctionContextIndex].bNeedsValidSource = false;
+
+						Result = ECreateSourcesForConversionFunctionResult::Valid;
+					}
 				}
 			}
 			return Result;
@@ -1054,7 +1083,7 @@ bool FMVVMViewBlueprintCompiler::PreCompileBindingSources(UWidgetBlueprintGenera
 					{
 						AddWarningForPropertyWithMVVMAndLegacyBinding(Binding.DestinationPath);
 					}
-					if (!CreateSourceContextForPropertyPath(Binding.SourcePath, true, FName()))
+					if (!CreateSourceContextForPropertyPath(Binding.SourcePath, true, INDEX_NONE, FName()))
 					{
 						bIsBindingsValid = false;
 						continue;
@@ -1086,7 +1115,7 @@ bool FMVVMViewBlueprintCompiler::PreCompileBindingSources(UWidgetBlueprintGenera
 					{
 						AddWarningForPropertyWithMVVMAndLegacyBinding(Binding.SourcePath);
 					}
-					if (!CreateSourceContextForPropertyPath(Binding.DestinationPath, false, FName()))
+					if (!CreateSourceContextForPropertyPath(Binding.DestinationPath, false, INDEX_NONE, FName()))
 					{
 						bIsBindingsValid = false;
 						continue;
@@ -1334,6 +1363,86 @@ bool FMVVMViewBlueprintCompiler::PreCompileBindings(UWidgetBlueprintGeneratedCla
 		return false;
 	}
 
+	auto GetConversionFunction = [this, Class, BlueprintView](FMVVMBlueprintViewBinding& Binding, bool bIsForwardBinding) -> TValueOrError<const UFunction*, void>
+	{
+		const UFunction* ConversionFunction = nullptr;
+		const UMVVMBlueprintViewConversionFunction* ViewConversionFunction = Binding.Conversion.GetConversionFunction(bIsForwardBinding);
+		if (ViewConversionFunction)
+		{
+			ConversionFunction = ViewConversionFunction->GetCompiledFunction(Class);
+			if (ConversionFunction == nullptr)
+			{
+				AddMessageForBinding(Binding, BlueprintView, FText::Format(LOCTEXT("ConversionFunctionNotFound", "The conversion function '{0}' could not be found."),
+					FText::FromName(ViewConversionFunction->GetCompiledFunctionName())),
+					EBindingMessageType::Error
+				);
+				bIsBindingsValid = false;
+				return MakeError();
+			}
+		}
+		return MakeValue(ConversionFunction);
+	};
+
+	auto TestExecutionMode = [this, BlueprintView](FMVVMBlueprintViewBinding& Binding)
+	{
+		if (Binding.bOverrideExecutionMode)
+		{
+			if (!GetDefault<UMVVMDeveloperProjectSettings>()->IsExecutionModeAllowed(Binding.OverrideExecutionMode))
+			{
+				AddMessageForBinding(Binding, BlueprintView, LOCTEXT("NotAllowedExecutionMode", "The binding has a restricted execution mode."), EBindingMessageType::Error);
+				return false;
+			}
+		}
+		return true;
+	};
+
+	auto GetSetterPath = [this, Class, BlueprintView](FMVVMBlueprintViewBinding& Binding, int32 BindingIndex, bool bIsForwardBinding) -> TValueOrError<TArray<UE::MVVM::FMVVMConstFieldVariant>, void>
+	{
+
+		const FBindingDestinationContext* DestinationContext = BindingDestinationContexts.FindByPredicate([BindingIndex, bIsForwardBinding](const FBindingDestinationContext& Other)
+			{
+				return Other.BindingIndex == BindingIndex && Other.bIsForwardBinding == bIsForwardBinding;
+			});
+
+		if (DestinationContext == nullptr)
+		{
+			AddMessageForBinding(Binding, BlueprintView, LOCTEXT("CouldNotFindDestination", "Could not find the pre compiled destination."), EBindingMessageType::Error);
+			return MakeError();
+		}
+
+		TArray<UE::MVVM::FMVVMConstFieldVariant> SetterPath;
+		if (!DestinationContext->bCanBeSetInNative)
+		{
+			check(!DestinationContext->GeneratedFunctionName.IsNone());
+			UFunction* FoundFunction = Class->FindFunctionByName(DestinationContext->GeneratedFunctionName);
+			if (FoundFunction == nullptr)
+			{
+				AddMessageForBinding(Binding, BlueprintView, LOCTEXT("CouldNotFindDestinationBPFunction", "Could not find the generated destination function."), EBindingMessageType::Error);
+				return MakeError();
+			}
+			FMVVMBlueprintPropertyPath DestinationPPropertyPath;
+			DestinationPPropertyPath.SetWidgetName(Class->ClassGeneratedBy->GetFName());
+			DestinationPPropertyPath.SetPropertyPath(BlueprintView->GetOuterUMVVMWidgetBlueprintExtension_View()->GetWidgetBlueprint(), UE::MVVM::FMVVMConstFieldVariant(FoundFunction));
+			SetterPath = CreateBindingDestinationPath(BlueprintView, Class, DestinationPPropertyPath);
+		}
+		else
+		{
+			const FMVVMBlueprintPropertyPath& DestinationPath = bIsForwardBinding ? Binding.DestinationPath : Binding.SourcePath;
+			SetterPath = CreateBindingDestinationPath(BlueprintView, Class, DestinationPath);
+		}
+
+		if (!IsPropertyPathValid(SetterPath))
+		{
+			const FMVVMBlueprintPropertyPath& DestinationPath = bIsForwardBinding ? Binding.DestinationPath : Binding.SourcePath;
+			AddMessageForBinding(Binding, BlueprintView, FText::Format(LOCTEXT("PropertyPathIsInvalid", "The property path '{0}' is invalid."),
+				PropertyPathToText(Class, BlueprintView, DestinationPath)),
+				EBindingMessageType::Error
+			);
+			return MakeError();
+		}
+		return MakeValue(SetterPath);
+	};
+
 	for (const FBindingSourceContext& BindingSourceContext : BindingSourceContexts)
 	{
 		FMVVMBlueprintViewBinding* BindingPtr = BlueprintView->GetBindingAt(BindingSourceContext.BindingIndex);
@@ -1350,89 +1459,132 @@ bool FMVVMViewBlueprintCompiler::PreCompileBindings(UWidgetBlueprintGeneratedCla
 			return MakeValue(FCompiledBindingLibraryCompiler::FFieldIdHandle());
 		};
 
-		if (Binding.bOverrideExecutionMode)
+		if (!TestExecutionMode(Binding))
 		{
-			if (!GetDefault<UMVVMDeveloperProjectSettings>()->IsExecutionModeAllowed(Binding.OverrideExecutionMode))
-			{
-				AddMessageForBinding(Binding, BlueprintView, LOCTEXT("NotAllowedExecutionMode", "The binding has a restricted execution mode."), EBindingMessageType::Error);
-			}
+			bIsBindingsValid = false;
+			continue;
 		}
 
+		const bool bIsComplexBinding = BindingSourceContext.ComplexConversionFunctionContextIndex != INDEX_NONE;
+		bool bFieldIdNeeded = !IsOneTimeBinding(Binding.BindingType);
+
 		TValueOrError<FCompiledBindingLibraryCompiler::FFieldIdHandle, FText> AddFieldResult = AddFieldId(BindingSourceContext.SourceClass, true, Binding.BindingType, BindingSourceContext.FieldId.GetFieldName());
 		if (AddFieldResult.HasError())
 		{
-			AddMessageForBinding(Binding, BlueprintView, FText::Format(LOCTEXT("CouldNotCreateSource", "Could not create source. {0}"),
-				AddFieldResult.GetError()), EBindingMessageType::Error);
-			bIsBindingsValid = false;
-			continue;
+			bFieldIdNeeded = false;
+			// For complex binding, at least one must be true
+			if (!bIsComplexBinding)
+			{
+				AddMessageForBinding(Binding, BlueprintView, FText::Format(LOCTEXT("CouldNotCreateSource", "Could not create source. {0}"),
+					AddFieldResult.GetError()), EBindingMessageType::Error);
+				bIsBindingsValid = false;
+				continue;
+			}
+		}
+		else if (bIsComplexBinding)
+		{
+			if (ensure(ComplexConversionFunctionContexts.IsValidIndex(BindingSourceContext.ComplexConversionFunctionContextIndex)))
+			{
+				ComplexConversionFunctionContexts[BindingSourceContext.ComplexConversionFunctionContextIndex].bHasValidFieldId = true;
+			}
 		}
 
 		TArray<UE::MVVM::FMVVMConstFieldVariant> SetterPath;
 		{
-			const FBindingDestinationContext* DestinationContext = BindingDestinationContexts.FindByPredicate([BindingIndex = BindingSourceContext.BindingIndex, bIsForwardBinding = BindingSourceContext.bIsForwardBinding](const FBindingDestinationContext& Other)
-				{
-					return Other.BindingIndex == BindingIndex && Other.bIsForwardBinding == bIsForwardBinding;
-				});
-
-			if (DestinationContext == nullptr)
+			TValueOrError<TArray<UE::MVVM::FMVVMConstFieldVariant>, void> SetterPathResult = GetSetterPath(Binding, BindingSourceContext.BindingIndex, BindingSourceContext.bIsForwardBinding);
+			if (SetterPathResult.HasError())
 			{
-				AddMessageForBinding(Binding, BlueprintView, LOCTEXT("CouldNotFindDestination", "Could not find the pre compiled destination."), EBindingMessageType::Error);
 				bIsBindingsValid = false;
 				continue;
 			}
+			SetterPath = SetterPathResult.StealValue();
+		}
 
-			if (!DestinationContext->bCanBeSetInNative)
-			{
-				check(!DestinationContext->GeneratedFunctionName.IsNone());
-				UFunction* FoundFunction = Class->FindFunctionByName(DestinationContext->GeneratedFunctionName);
-				if (FoundFunction == nullptr)
-				{
-					AddMessageForBinding(Binding, BlueprintView, LOCTEXT("CouldNotFindDestinationBPFunction", "Could not find the generated destination function."), EBindingMessageType::Error);
-					bIsBindingsValid = false;
-					continue;
-				}
-				FMVVMBlueprintPropertyPath DestinationPPropertyPath;
-				DestinationPPropertyPath.SetWidgetName(Class->ClassGeneratedBy->GetFName());
-				DestinationPPropertyPath.SetPropertyPath(BlueprintView->GetOuterUMVVMWidgetBlueprintExtension_View()->GetWidgetBlueprint(), UE::MVVM::FMVVMConstFieldVariant(FoundFunction));
-				SetterPath = CreateBindingDestinationPath(BlueprintView, Class, DestinationPPropertyPath);
-			}
-			else
+		const UFunction* ConversionFunction = nullptr;
+		{
+			TValueOrError<const UFunction*, void> ConversionFunctionResult = GetConversionFunction(Binding, BindingSourceContext.bIsForwardBinding);
+			if (ConversionFunctionResult.HasError())
 			{
-				const FMVVMBlueprintPropertyPath& DestinationPath = BindingSourceContext.bIsForwardBinding ? Binding.DestinationPath : Binding.SourcePath;
-				SetterPath = CreateBindingDestinationPath(BlueprintView, Class, DestinationPath);
+				bIsBindingsValid = false;
+				continue;
 			}
+			ConversionFunction = ConversionFunctionResult.StealValue();
+		}
+
+		TValueOrError<FCompiledBinding, FText> AddBindingResult = CreateCompiledBinding(Class, BindingSourceContext.PropertyPath, SetterPath, ConversionFunction, bIsComplexBinding);
+		if (AddBindingResult.HasError())
+		{
+			AddMessageForBinding(Binding, BlueprintView,
+				FText::Format(LOCTEXT("CouldNotCreateBinding", "Could not create binding. {0}"), AddBindingResult.GetError()),
+				EBindingMessageType::Error
+			);
+			bIsBindingsValid = false;
+			continue;
+		}
+
+		FCompilerBinding NewBinding;
+		NewBinding.BindingIndex = BindingSourceContext.BindingIndex;
+		NewBinding.UserWidgetPropertyContextIndex = BindingSourceContext.UserWidgetPropertyContextIndex;
+		NewBinding.SourceCreatorContextIndex = BindingSourceContext.SourceCreatorContextIndex;
+		NewBinding.ComplexConversionFunctionContextIndex = BindingSourceContext.ComplexConversionFunctionContextIndex;
+		NewBinding.bSourceIsUserWidget = BindingSourceContext.bIsRootWidget;
+		NewBinding.bFieldIdNeeded = bFieldIdNeeded;
+		NewBinding.bIsForwardBinding = BindingSourceContext.bIsForwardBinding;
+
+		NewBinding.CompiledBinding = AddBindingResult.StealValue();
+		if (AddFieldResult.HasValue())
+		{
+			NewBinding.FieldIdHandle = AddFieldResult.StealValue();
+		}
+
+		CompilerBindings.Emplace(NewBinding);
+	}
+
 
-			if (!IsPropertyPathValid(SetterPath))
+	for (const FSimpleBindingContext& BindingContext : SimpleBindingContexts)
+	{
+		FMVVMBlueprintViewBinding* BindingPtr = BlueprintView->GetBindingAt(BindingContext.BindingIndex);
+		check(BindingPtr);
+		FMVVMBlueprintViewBinding& Binding = *BindingPtr;
+
+		if (!TestExecutionMode(Binding))
+		{
+			bIsBindingsValid = false;
+			continue;
+		}
+
+		if (!IsOneTimeBinding(Binding.BindingType))
+		{
+			AddMessageForBinding(Binding, BlueprintView, LOCTEXT("FunctionRequiresOneTimeBinding", "There is no source. The binding must be a OneTime binding."), EBindingMessageType::Error);
+			bIsBindingsValid = false;
+			continue;
+		}
+
+		TArray<UE::MVVM::FMVVMConstFieldVariant> SetterPath;
+		{
+			TValueOrError<TArray<UE::MVVM::FMVVMConstFieldVariant>, void> SetterPathResult = GetSetterPath(Binding, BindingContext.BindingIndex, BindingContext.bIsForwardBinding);
+			if (SetterPathResult.HasError())
 			{
-				const FMVVMBlueprintPropertyPath& DestinationPath = BindingSourceContext.bIsForwardBinding ? Binding.DestinationPath : Binding.SourcePath;
-				AddMessageForBinding(Binding, BlueprintView, FText::Format(LOCTEXT("PropertyPathIsInvalid", "The property path '{0}' is invalid."),
-					PropertyPathToText(Class, BlueprintView, DestinationPath)),
-					EBindingMessageType::Error
-				);
 				bIsBindingsValid = false;
 				continue;
 			}
+			SetterPath = SetterPathResult.StealValue();
 		}
 
 		const UFunction* ConversionFunction = nullptr;
 		{
-			const UMVVMBlueprintViewConversionFunction* ViewConversionFunction = Binding.Conversion.GetConversionFunction(BindingSourceContext.bIsForwardBinding);
-			if (ViewConversionFunction)
+			TValueOrError<const UFunction*, void> ConversionFunctionResult = GetConversionFunction(Binding, BindingContext.bIsForwardBinding);
+			if (ConversionFunctionResult.HasError())
 			{
-				ConversionFunction = ViewConversionFunction->GetCompiledFunction(Class);
-				if (ConversionFunction == nullptr)
-				{
-					AddMessageForBinding(Binding, BlueprintView, FText::Format(LOCTEXT("ConversionFunctionNotFound", "The conversion function '{0}' could not be found."), 
-						FText::FromName(ViewConversionFunction->GetCompiledFunctionName())),
-						EBindingMessageType::Error
-					);
-					bIsBindingsValid = false;
-					continue;
-				}
+				bIsBindingsValid = false;
+				continue;
 			}
+			ConversionFunction = ConversionFunctionResult.StealValue();
 		}
 
-		TValueOrError<FCompiledBinding, FText> AddBindingResult = CreateCompiledBinding(Class, BindingSourceContext.PropertyPath, SetterPath, ConversionFunction, BindingSourceContext.bIsComplexBinding);
+		TArrayView<const UE::MVVM::FMVVMConstFieldVariant> GetterFields;
+		const bool bIsComplexBinding = BindingContext.ComplexConversionFunctionContextIndex != INDEX_NONE;
+		TValueOrError<FCompiledBinding, FText> AddBindingResult = CreateCompiledBinding(Class, GetterFields, SetterPath, ConversionFunction, bIsComplexBinding);
 		if (AddBindingResult.HasError())
 		{
 			AddMessageForBinding(Binding, BlueprintView,
@@ -1444,19 +1596,31 @@ bool FMVVMViewBlueprintCompiler::PreCompileBindings(UWidgetBlueprintGeneratedCla
 		}
 
 		FCompilerBinding NewBinding;
-		NewBinding.BindingIndex = BindingSourceContext.BindingIndex;
-		NewBinding.UserWidgetPropertyContextIndex = BindingSourceContext.UserWidgetPropertyContextIndex;
-		NewBinding.SourceCreatorContextIndex = BindingSourceContext.SourceCreatorContextIndex;
-		NewBinding.bSourceIsUserWidget = BindingSourceContext.bIsRootWidget;
+		NewBinding.BindingIndex = BindingContext.BindingIndex;
 		NewBinding.bFieldIdNeeded = !IsOneTimeBinding(Binding.BindingType);
-		NewBinding.bIsForwardBinding = BindingSourceContext.bIsForwardBinding;
+		NewBinding.bIsForwardBinding = BindingContext.bIsForwardBinding;
+		NewBinding.bNeedsValidSource = false;
 
 		NewBinding.CompiledBinding = AddBindingResult.StealValue();
-		NewBinding.FieldIdHandle = AddFieldResult.StealValue();
 
 		CompilerBindings.Emplace(NewBinding);
 	}
 
+	if (bIsBindingsValid)
+	{
+		// Confirm that at least one source is valid for each pin of a complex conversion function.
+		for (const FComplexConversionFunctionContext& ComplexConversionFunctionContext : ComplexConversionFunctionContexts)
+		{
+			if (!ComplexConversionFunctionContext.bHasValidFieldId && ComplexConversionFunctionContext.bNeedsValidSource)
+			{
+				FMVVMBlueprintViewBinding* BindingPtr = BlueprintView->GetBindingAt(ComplexConversionFunctionContext.BindingIndex);
+				check(BindingPtr);
+				AddMessageForBinding(*BindingPtr, BlueprintView, LOCTEXT("CouldNotCreateSourceForConversionFunction", "There is no source. The binding must be a OneTime binding."), EBindingMessageType::Error);
+				bIsBindingsValid = false;
+			}
+		}
+	}
+
 	return bIsBindingsValid;
 }
 
@@ -1496,9 +1660,10 @@ bool FMVVMViewBlueprintCompiler::CompileBindings(const FCompiledBindingLibraryCo
 			{
 				NewBinding.SourcePropertyName = CompilerSourceCreatorContexts[CompileBinding.SourceCreatorContextIndex].ViewModelContext.GetViewModelName();
 			}
-			else
+			else if (CompileBinding.bNeedsValidSource)
 			{
 				WidgetBlueprintCompilerContext.MessageLog.Error(*LOCTEXT("InvalidSourceInternal", "Internal error. The binding doesn't have a valid source.").ToString());
+				bIsBindingsValid = false;
 				return false;
 			}
 		}
@@ -1591,11 +1756,18 @@ bool FMVVMViewBlueprintCompiler::CompileBindings(const FCompiledBindingLibraryCo
 				continue;
 			}
 
+			bool bExecuteAtInitialization = CompileBinding.bIsForwardBinding;
+			if (bExecuteAtInitialization && CompileBinding.ComplexConversionFunctionContextIndex != INDEX_NONE)
+			{
+				bExecuteAtInitialization = !ComplexConversionFunctionContexts[CompileBinding.ComplexConversionFunctionContextIndex].bExecAtInitGenerated;
+				ComplexConversionFunctionContexts[CompileBinding.ComplexConversionFunctionContextIndex].bExecAtInitGenerated = true;
+			}
+
 			NewBinding.ExecutionMode = ViewBinding->bOverrideExecutionMode ? ViewBinding->OverrideExecutionMode : (EMVVMExecutionMode)CVarDefaultExecutionMode->GetInt();;
 			NewBinding.EditorId = ViewBinding->BindingId;
 
 			NewBinding.Flags |= (ViewBinding->bEnabled) ? FMVVMViewClass_CompiledBinding::EBindingFlags::EnabledByDefault : 0;
-			NewBinding.Flags |= (CompileBinding.bIsForwardBinding) ? FMVVMViewClass_CompiledBinding::EBindingFlags::ExecuteAtInitialization : 0;
+			NewBinding.Flags |= (bExecuteAtInitialization) ? FMVVMViewClass_CompiledBinding::EBindingFlags::ExecuteAtInitialization : 0;
 			NewBinding.Flags |= (IsOneTimeBinding(ViewBinding->BindingType)) ? FMVVMViewClass_CompiledBinding::EBindingFlags::OneTime : 0;
 			NewBinding.Flags |= (bIsOptional) ? FMVVMViewClass_CompiledBinding::EBindingFlags::ViewModelOptional : 0;
 			NewBinding.Flags |= (ViewBinding->bOverrideExecutionMode) ? FMVVMViewClass_CompiledBinding::EBindingFlags::OverrideExecuteMode : 0;
diff --git a/Engine/Plugins/Runtime/ModelViewViewModel/Source/ModelViewViewModelBlueprint/Private/MVVMViewBlueprintCompiler.h b/Engine/Plugins/Runtime/ModelViewViewModel/Source/ModelViewViewModelBlueprint/Private/MVVMViewBlueprintCompiler.h
--- a/Engine/Plugins/Runtime/ModelViewViewModel/Source/ModelViewViewModelBlueprint/Private/MVVMViewBlueprintCompiler.h
+++ b/Engine/Plugins/Runtime/ModelViewViewModel/Source/ModelViewViewModelBlueprint/Private/MVVMViewBlueprintCompiler.h
@@ -169,9 +169,11 @@ private:
 		int32 BindingIndex = INDEX_NONE;
 		int32 UserWidgetPropertyContextIndex = INDEX_NONE;
 		int32 SourceCreatorContextIndex = INDEX_NONE;
+		int32 ComplexConversionFunctionContextIndex = INDEX_NONE;
 		bool bSourceIsUserWidget = false;
 		bool bFieldIdNeeded = false;
 		bool bIsForwardBinding = false;
+		bool bNeedsValidSource = true;	// if false, the binding is a onetime or executed manually (onclick event).
 
 		FCompiledBindingLibraryCompiler::FFieldIdHandle FieldIdHandle;
 		FCompiledBinding CompiledBinding;
@@ -179,6 +181,19 @@ private:
 	};
 	TArray<FCompilerBinding> CompilerBindings;
 
+	/**
+	 * Binding that do not requires a source to execute.
+	 * They could be a OneTime ComplexConversionFunction or an event (onclick).
+	 */
+	struct FSimpleBindingContext
+	{
+		int32 BindingIndex = INDEX_NONE;
+		bool bIsForwardBinding = false;
+
+		int32 ComplexConversionFunctionContextIndex = INDEX_NONE;
+	};
+	TArray<FSimpleBindingContext> SimpleBindingContexts;
+
 	/**
 	 * Source path of a binding that contains a FieldId that we can register/bind to
 	 */
@@ -186,8 +201,6 @@ private:
 	{
 		int32 BindingIndex = INDEX_NONE;
 		bool bIsForwardBinding = false;
-		// Complex binding do not need to add the PropertyPath to the binding info.
-		bool bIsComplexBinding = false;
 
 		const UClass* SourceClass = nullptr;
 		// The property that are registering to.
@@ -198,6 +211,8 @@ private:
 		// or Field.SubProperty.SubProperty
 		TArray<UE::MVVM::FMVVMConstFieldVariant> PropertyPath;
 
+		// Complex binding do not need to add the PropertyPath to the binding info.
+		int32 ComplexConversionFunctionContextIndex = INDEX_NONE;
 		// The source if it's a property
 		int32 UserWidgetPropertyContextIndex = INDEX_NONE;
 		// The source if it's a property
@@ -222,6 +237,25 @@ private:
 	};
 	TArray<FBindingDestinationContext> BindingDestinationContexts;
 
+	/**
+	 * Complex Conversion function shared data
+	 */
+	struct FComplexConversionFunctionContext
+	{
+		int32 BindingIndex = INDEX_NONE;
+		bool bIsForwardBinding = false;
+
+		// More than one binding can use the same conversion function.
+		// Only one should execute the binding at initialization.
+		bool bExecAtInitGenerated = false;
+		// More than one binding can use the same conversion function.
+		// At least one must have a valid source to trigger the binding.
+		//If not, then it should be a FSimpleBindingContext
+		bool bHasValidFieldId = false;
+		bool bNeedsValidSource = true;
+	};
+	TArray<FComplexConversionFunctionContext> ComplexConversionFunctionContexts;
+
 	TMap<FName, UWidget*> WidgetNameToWidgetPointerMap;
 	FWidgetBlueprintCompilerContext& WidgetBlueprintCompilerContext;
 	FCompiledBindingLibraryCompiler BindingLibraryCompiler;
diff --git a/Engine/Plugins/Runtime/ModelViewViewModel/Source/ModelViewViewModelBlueprint/Private/MVVMWidgetBlueprintExtension_View.cpp b/Engine/Plugins/Runtime/ModelViewViewModel/Source/ModelViewViewModelBlueprint/Private/MVVMWidgetBlueprintExtension_View.cpp
--- a/Engine/Plugins/Runtime/ModelViewViewModel/Source/ModelViewViewModelBlueprint/Private/MVVMWidgetBlueprintExtension_View.cpp
+++ b/Engine/Plugins/Runtime/ModelViewViewModel/Source/ModelViewViewModelBlueprint/Private/MVVMWidgetBlueprintExtension_View.cpp
@@ -139,8 +139,8 @@ void UMVVMWidgetBlueprintExtension_View::HandleFinishCompilingClass(UWidgetBluep
 		if (bCompiled && UE::MVVM::Private::GAllowViewClass)
 		{
 			check(ViewExtension);
-			// Does it have any bindings or any sources
-			if (ViewExtension->GetViewModelCreators().Num() > 0 || const_cast<const UMVVMViewClass*>(ViewExtension)->GetCompiledBindings().Num() > 0)
+			// Does it have any bindings
+			if (const_cast<const UMVVMViewClass*>(ViewExtension)->GetCompiledBindings().Num() > 0)
 			{
 				// Test if parent also has a view
 				if (Class->GetExtension<UMVVMViewClass>(true))
