MVVM: Undo modification made to make property blueprintreadonly. Introduced with CL 26845876

diff --git a/Engine/Plugins/Runtime/ModelViewViewModel/Source/ModelViewViewModelBlueprint/Private/MVVMViewBlueprintCompiler.cpp b/Engine/Plugins/Runtime/ModelViewViewModel/Source/ModelViewViewModelBlueprint/Private/MVVMViewBlueprintCompiler.cpp
--- a/Engine/Plugins/Runtime/ModelViewViewModel/Source/ModelViewViewModelBlueprint/Private/MVVMViewBlueprintCompiler.cpp
+++ b/Engine/Plugins/Runtime/ModelViewViewModel/Source/ModelViewViewModelBlueprint/Private/MVVMViewBlueprintCompiler.cpp
@@ -259,7 +259,10 @@ void FMVVMViewBlueprintCompiler::CreateVariables(const FWidgetBlueprintCompilerC
 		if (NewProperty != nullptr)
 		{
 			NewProperty->SetPropertyFlags(CPF_BlueprintVisible | CPF_RepSkip | CPF_Transient | CPF_DuplicateTransient);
-			NewProperty->SetPropertyFlags(CPF_BlueprintReadOnly);
+			if (SourceContext.BlueprintSetter.IsEmpty())
+			{
+				NewProperty->SetPropertyFlags(CPF_BlueprintReadOnly);
+			}
 			NewProperty->SetPropertyFlags(SourceContext.bExposeOnSpawn ? CPF_ExposeOnSpawn : CPF_DisableEditOnInstance);
 
 #if WITH_EDITOR
