MVVM: Add flag on viewmodel context to allow edit, remove and rename in the editor.

diff --git a/Engine/Plugins/Runtime/ModelViewViewModel/Source/ModelViewViewModelBlueprint/Public/MVVMBlueprintViewModelContext.h b/Engine/Plugins/Runtime/ModelViewViewModel/Source/ModelViewViewModelBlueprint/Public/MVVMBlueprintViewModelContext.h
--- a/Engine/Plugins/Runtime/ModelViewViewModel/Source/ModelViewViewModelBlueprint/Public/MVVMBlueprintViewModelContext.h
+++ b/Engine/Plugins/Runtime/ModelViewViewModel/Source/ModelViewViewModelBlueprint/Public/MVVMBlueprintViewModelContext.h
@@ -92,7 +92,7 @@ private:
 	UPROPERTY(VisibleAnywhere, AdvancedDisplay, Category = "Viewmodel", meta = (DisplayName = "Viewmodel Context Id"))
 	FGuid ViewModelContextId;
 
-	UPROPERTY(VisibleAnywhere, Category = "Viewmodel", NoClear, meta = (AllowedClasses = "/Script/UMG.NotifyFieldValueChanged", DisallowedClasses = "/Script/UMG.Widget"))
+	UPROPERTY(VisibleAnywhere, Category = "Viewmodel", NoClear, meta = (DisallowCreateNew, AllowedClasses = "/Script/UMG.NotifyFieldValueChanged", DisallowedClasses = "/Script/UMG.Widget"))
 	TObjectPtr<UClass> NotifyFieldValueClass = nullptr;
 
 	UPROPERTY()
@@ -111,29 +111,41 @@ public:
 	EMVVMBlueprintViewModelContextCreationType CreationType = EMVVMBlueprintViewModelContextCreationType::CreateInstance;
 
 	/** Identifier of an already registered viewmodel. */
-	UPROPERTY(EditAnywhere, Category = "Viewmodel", AdvancedDisplay, meta = (DisplayName = "Global Viewmodel Identifier", EditCondition = "CreationType == EMVVMBlueprintViewModelContextCreationType::GlobalViewModelCollection", EditConditionHides))
+	UPROPERTY(EditAnywhere, Category = "Viewmodel", AdvancedDisplay, meta = (DisplayName = "Global Viewmodel Identifier", EditCondition = "CreationType == EMVVMBlueprintViewModelContextCreationType::GlobalViewModelCollection && bCanEdit", EditConditionHides))
 	FName GlobalViewModelIdentifier;
 
 	/** The Path to get the viewmodel instance. */
-	UPROPERTY(EditAnywhere, Category = "Viewmodel", AdvancedDisplay, meta = (DisplayName = "Viewmodel Property Path", EditCondition = "CreationType == EMVVMBlueprintViewModelContextCreationType::PropertyPath", EditConditionHides))
+	UPROPERTY(EditAnywhere, Category = "Viewmodel", AdvancedDisplay, meta = (DisplayName = "Viewmodel Property Path", EditCondition = "CreationType == EMVVMBlueprintViewModelContextCreationType::PropertyPath && bCanEdit", EditConditionHides))
 	FString ViewModelPropertyPath;
 
-	UPROPERTY(EditAnywhere, Category = "Viewmodel", Instanced, AdvancedDisplay, meta = (EditInline, EditCondition = "CreationType == EMVVMBlueprintViewModelContextCreationType::Resolver", EditConditionHides))
+	UPROPERTY(EditAnywhere, Category = "Viewmodel", Instanced, AdvancedDisplay, meta = (EditInline, EditCondition = "CreationType == EMVVMBlueprintViewModelContextCreationType::Resolver && bCanEdit", EditConditionHides))
 	TObjectPtr<UMVVMViewModelContextResolver> Resolver = nullptr;
 
 	/**
 	 * Generate a public setter for this viewmodel.
 	 * @note Always true when the Creation Type is Manual.
 	 */
-	UPROPERTY(EditAnywhere, Category = "Viewmodel", AdvancedDisplay, meta = (DisplanName="Create Public Setter", EditCondition = "CreationType != EMVVMBlueprintViewModelContextCreationType::Manual", EditConditionHides))
+	UPROPERTY(EditAnywhere, Category = "Viewmodel", AdvancedDisplay, meta = (DisplanName="Create Public Setter", EditCondition = "CreationType != EMVVMBlueprintViewModelContextCreationType::Manual && bCanEdit", EditConditionHides))
 	bool bCreateSetterFunction = false;
 
 	/**
 	 * Optional. Will not warn if the instance is not set or found.
 	 * @note Always true when the Creation Type is Manual.
 	 */
-	UPROPERTY(EditAnywhere, Category = "Viewmodel", AdvancedDisplay, meta = (EditCondition = "CreationType == EMVVMBlueprintViewModelContextCreationType::GlobalViewModelCollection || CreationType == EMVVMBlueprintViewModelContextCreationType::PropertyPath", EditConditionHides))
+	UPROPERTY(EditAnywhere, Category = "Viewmodel", AdvancedDisplay, meta = (EditCondition = "(CreationType == EMVVMBlueprintViewModelContextCreationType::GlobalViewModelCollection || CreationType == EMVVMBlueprintViewModelContextCreationType::PropertyPath) && bCanEdit", EditConditionHides))
 	bool bOptional = false;
+
+	/** Can change the name in the editor. */
+	UPROPERTY()
+	bool bCanRename = true;
+
+	/** Can change properties in the editor. */
+	UPROPERTY()
+	bool bCanEdit = true;
+
+	/** Can remove the viewmodel in the editor. */
+	UPROPERTY()
+	bool bCanRemove = true;
 };
 
 template<>
diff --git a/Engine/Plugins/Runtime/ModelViewViewModel/Source/ModelViewViewModelEditor/Private/Customizations/MVVMBlueprintViewModelContextCustomization.cpp b/Engine/Plugins/Runtime/ModelViewViewModel/Source/ModelViewViewModelEditor/Private/Customizations/MVVMBlueprintViewModelContextCustomization.cpp
--- a/Engine/Plugins/Runtime/ModelViewViewModel/Source/ModelViewViewModelEditor/Private/Customizations/MVVMBlueprintViewModelContextCustomization.cpp
+++ b/Engine/Plugins/Runtime/ModelViewViewModel/Source/ModelViewViewModelEditor/Private/Customizations/MVVMBlueprintViewModelContextCustomization.cpp
@@ -8,6 +8,7 @@
 #include "IDetailChildrenBuilder.h"
 #include "IPropertyAccessEditor.h"
 #include "IPropertyTypeCustomization.h"
+#include "MVVMBlueprintView.h"
 #include "MVVMDeveloperProjectSettings.h"
 #include "MVVMEditorSubsystem.h"
 #include "PropertyHandle.h"
@@ -184,6 +185,21 @@ void FBlueprintViewModelContextDetailCustomization::CustomizeChildren(TSharedRef
 		}
 	}
 
+	bool bCanEdit = false;
+	bool bCanRename = false;
+	{
+
+		TSharedPtr<FWidgetBlueprintEditor> Editor = WidgetBlueprintEditor.Pin();
+		const UWidgetBlueprint* WidgetBP = Editor ? Editor->GetWidgetBlueprintObj() : nullptr;
+		const UMVVMBlueprintView* BlueprintView = WidgetBP ? GEditor->GetEditorSubsystem<UMVVMEditorSubsystem>()->GetView(WidgetBP) : nullptr;
+		const FMVVMBlueprintViewModelContext* ViewModel = BlueprintView ? BlueprintView->FindViewModel(ViewModelContextId) : nullptr;
+		if (ViewModel)
+		{
+			bCanEdit = ViewModel->bCanEdit;
+			bCanRename = ViewModel->bCanRename;
+		}
+	}
+
 	if (NotifyFieldValueClassHandle)
 	{
 		for (uint32 ChildIndex = 0; ChildIndex < NumChildren; ++ChildIndex)
@@ -219,9 +235,10 @@ void FBlueprintViewModelContextDetailCustomization::CustomizeChildren(TSharedRef
 							.OnTextCommitted(this, &FBlueprintViewModelContextDetailCustomization::HandleNameTextCommitted)
 							.OnVerifyTextChanged(this, &FBlueprintViewModelContextDetailCustomization::HandleNameVerifyTextChanged)
 							.SelectAllTextOnCommit(true)
-							.IsEnabled(ViewModelContextId.IsValid())
+							.IsEnabled(bCanRename && ViewModelContextId.IsValid())
 						]
-					];
+					]
+					.IsEnabled(bCanRename);
 				}
 			}
 			else if (PropertyName == Name_ViewModelPropertyPath)
@@ -276,18 +293,21 @@ void FBlueprintViewModelContextDetailCustomization::CustomizeChildren(TSharedRef
 							.Text(this, &FBlueprintViewModelContextDetailCustomization::GetCreationTypeValue)
 							.ToolTipText(this, &FBlueprintViewModelContextDetailCustomization::GetExecutionTypeValueToolTip)
 						]
-					];
+					]
+					.IsEnabled(bCanEdit);
 			}
 			else if (PropertyName == Name_CreateSetterFunction)
 			{
 				if (GetDefault<UMVVMDeveloperProjectSettings>()->bAllowGeneratedViewModelSetter)
 				{
-					ChildBuilder.AddProperty(ChildHandle.ToSharedRef());
+					ChildBuilder.AddProperty(ChildHandle.ToSharedRef())
+						.IsEnabled(bCanEdit);
 				}
 			}
 			else
 			{
-				ChildBuilder.AddProperty(ChildHandle.ToSharedRef());
+				ChildBuilder.AddProperty(ChildHandle.ToSharedRef())
+					.IsEnabled(bCanEdit);;
 			}
 		}
 	}
diff --git a/Engine/Plugins/Runtime/ModelViewViewModel/Source/ModelViewViewModelEditor/Private/MVVMEditorSubsystem.cpp b/Engine/Plugins/Runtime/ModelViewViewModel/Source/ModelViewViewModelEditor/Private/MVVMEditorSubsystem.cpp
--- a/Engine/Plugins/Runtime/ModelViewViewModel/Source/ModelViewViewModelEditor/Private/MVVMEditorSubsystem.cpp
+++ b/Engine/Plugins/Runtime/ModelViewViewModel/Source/ModelViewViewModelEditor/Private/MVVMEditorSubsystem.cpp
@@ -133,6 +133,10 @@ FName UMVVMEditorSubsystem::AddViewModel(UWidgetBlueprint* WidgetBlueprint, cons
 		if (UMVVMBlueprintView* View = GetView(WidgetBlueprint))
 		{
 			FString ClassName = ViewModelClass->ClassGeneratedBy != nullptr ? ViewModelClass->ClassGeneratedBy->GetName() : ViewModelClass->GetAuthoredName();
+			if (Cast<UBlueprintGeneratedClass>(ViewModelClass) && ClassName.EndsWith(TEXT("_C")))
+			{
+				ClassName.RemoveAt(ClassName.Len() - 2, 2);
+			}
 			FString ViewModelName = ClassName;
 			FKismetNameValidator NameValidator(WidgetBlueprint);
 
@@ -168,7 +172,10 @@ void UMVVMEditorSubsystem::RemoveViewModel(UWidgetBlueprint* WidgetBlueprint, FN
 	{
 		if (const FMVVMBlueprintViewModelContext* ViewModelContext = View->FindViewModel(ViewModel))
 		{
-			View->RemoveViewModel(ViewModelContext->GetViewModelId());
+			if (ViewModelContext->bCanRemove)
+			{
+				View->RemoveViewModel(ViewModelContext->GetViewModelId());
+			}
 		}
 	}
 }
@@ -205,9 +212,14 @@ bool UMVVMEditorSubsystem::RenameViewModel(UWidgetBlueprint* WidgetBlueprint, FN
 		return false;
 	}
 
-	const FScopedTransaction Transaction(LOCTEXT("RenameViewModel", "Rename Viewmodel"));
-	View->Modify();
-	return View->RenameViewModel(ViewModel, NewViewModel);
+	const FMVVMBlueprintViewModelContext* ViewModelContext = View->FindViewModel(ViewModel);
+	if (ViewModelContext && ViewModelContext->bCanRename)
+	{
+		const FScopedTransaction Transaction(LOCTEXT("RenameViewModel", "Rename Viewmodel"));
+		View->Modify();
+		return View->RenameViewModel(ViewModel, NewViewModel);
+	}
+	return false;
 }
 
 FMVVMBlueprintViewBinding& UMVVMEditorSubsystem::AddBinding(UWidgetBlueprint* WidgetBlueprint)
diff --git a/Engine/Plugins/Runtime/ModelViewViewModel/Source/ModelViewViewModelEditor/Private/Widgets/SMVVMViewModelPanel.cpp b/Engine/Plugins/Runtime/ModelViewViewModel/Source/ModelViewViewModelEditor/Private/Widgets/SMVVMViewModelPanel.cpp
--- a/Engine/Plugins/Runtime/ModelViewViewModel/Source/ModelViewViewModelEditor/Private/Widgets/SMVVMViewModelPanel.cpp
+++ b/Engine/Plugins/Runtime/ModelViewViewModel/Source/ModelViewViewModelEditor/Private/Widgets/SMVVMViewModelPanel.cpp
@@ -436,6 +436,7 @@ TSharedRef<SWidget> SMVVMViewModelPanel::HandleGenerateContainer(UE::PropertyVie
 				{
 					TSharedRef<SInlineEditableTextBlock> EditableTextBlock = SNew(SInlineEditableTextBlock)
 						.Text(ViewModelContext->GetDisplayName())
+						.IsReadOnly(this, &SMVVMViewModelPanel::HandleCanRename, VMGuid)
 						.OnVerifyTextChanged(this, &SMVVMViewModelPanel::HandleVerifyNameTextChanged, VMGuid)
 						.OnTextCommitted(this, &SMVVMViewModelPanel::HandleNameTextCommited, VMGuid);
 					EditableTextBlocks.Add(VMGuid, EditableTextBlock);
@@ -481,6 +482,15 @@ TSharedRef<SWidget> SMVVMViewModelPanel::HandleGenerateContainer(UE::PropertyVie
 }
 
 
+bool SMVVMViewModelPanel::HandleCanRename(FGuid ViewModelGuid) const
+{
+	const UMVVMBlueprintView* BlueprintView = WeakBlueprintView.Get();
+	const FMVVMBlueprintViewModelContext* ViewModelContext = BlueprintView ? BlueprintView->FindViewModel(ViewModelGuid) : nullptr;
+	TSharedPtr<FWidgetBlueprintEditor> BlueprintEditor = WeakBlueprintEditor.Pin();
+	return ViewModelContext && ViewModelContext->bCanRename && BlueprintEditor && BlueprintEditor->InEditingMode();
+}
+
+
 bool SMVVMViewModelPanel::HandleVerifyNameTextChanged(const FText& InText, FText& OutErrorMessage, FGuid ViewModelGuid)
 {
 	return RenameViewModelProperty(ViewModelGuid, InText, false, OutErrorMessage);
@@ -496,6 +506,7 @@ void SMVVMViewModelPanel::HandleNameTextCommited(const FText& InText, ETextCommi
 	}
 }
 
+
 FReply SMVVMViewModelPanel::HandleDragDetected(const FGeometry& MyGeometry, const FPointerEvent& MouseEvent, UE::PropertyViewer::SPropertyViewer::FHandle ContainerHandle, TArrayView<const FFieldVariant> Fields) const
 {
 	TSharedPtr<FWidgetBlueprintEditor> WidgetBlueprintEditor = WeakBlueprintEditor.Pin();
@@ -529,6 +540,7 @@ FReply SMVVMViewModelPanel::HandleDragDetected(const FGeometry& MyGeometry, cons
 	return FReply::Unhandled();
 }
 
+
 void SMVVMViewModelPanel::CreateCommandList()
 {
 	CommandList = MakeShared<FUICommandList>();
@@ -593,7 +605,8 @@ void SMVVMViewModelPanel::HandleDeleteViewModel()
 		{
 			if (FGuid* VMGuidPtr = PropertyViewerHandles.Find(SelectedItem.Handle))
 			{
-				if (const FMVVMBlueprintViewModelContext* ViewModelContext = BlueprintView->FindViewModel(*VMGuidPtr))
+				const FMVVMBlueprintViewModelContext* ViewModelContext = BlueprintView->FindViewModel(*VMGuidPtr);
+				if (ViewModelContext && ViewModelContext->bCanRemove)
 				{
 					FGuid BindingId = GEditor->GetEditorSubsystem<UMVVMEditorSubsystem>()->GetFirstBindingThatUsesViewModel(WidgetBP, *VMGuidPtr);
 					if (BindingId.IsValid())
@@ -614,9 +627,17 @@ void SMVVMViewModelPanel::HandleDeleteViewModel()
 
 bool SMVVMViewModelPanel::HandleCanDeleteViewModel() const
 {
-	if (TSharedPtr<FWidgetBlueprintEditor> BlueprintEditor = WeakBlueprintEditor.Pin())
+	TArray<UE::PropertyViewer::SPropertyViewer::FSelectedItem> Items = ViewModelTreeView->GetSelectedItems();
+	if (Items.Num() == 1 && Items[0].bIsContainerSelected)
 	{
-		return BlueprintEditor->InEditingMode();
+		const FGuid* VMGuidPtr = PropertyViewerHandles.Find(Items[0].Handle);
+		const UMVVMBlueprintView* BlueprintView = WeakBlueprintView.Get();
+		const FMVVMBlueprintViewModelContext* ViewModelContext = BlueprintView && VMGuidPtr ? BlueprintView->FindViewModel(*VMGuidPtr) : nullptr;
+		TSharedPtr<FWidgetBlueprintEditor> BlueprintEditor = WeakBlueprintEditor.Pin();
+		if (ViewModelContext && ViewModelContext->bCanRemove && BlueprintEditor)
+		{
+			return BlueprintEditor->InEditingMode() && ViewModelContext->bCanRemove;
+		}
 	}
 	return false;
 }
@@ -627,7 +648,8 @@ void SMVVMViewModelPanel::HandleRenameViewModel()
 	TArray<UE::PropertyViewer::SPropertyViewer::FSelectedItem> Items = ViewModelTreeView->GetSelectedItems();
 	if (Items.Num() == 1 && Items[0].bIsContainerSelected)
 	{
-		if (const FGuid* VMGuidPtr = PropertyViewerHandles.Find(Items[0].Handle))
+		const FGuid* VMGuidPtr = PropertyViewerHandles.Find(Items[0].Handle);
+		if (VMGuidPtr && HandleCanRename(*VMGuidPtr))
 		{
 			if (TSharedPtr<SInlineEditableTextBlock>* TextBlockPtr = EditableTextBlocks.Find(*VMGuidPtr))
 			{
@@ -640,9 +662,11 @@ void SMVVMViewModelPanel::HandleRenameViewModel()
 
 bool SMVVMViewModelPanel::HandleCanRenameViewModel() const
 {
-	if (TSharedPtr<FWidgetBlueprintEditor> BlueprintEditor = WeakBlueprintEditor.Pin())
+	TArray<UE::PropertyViewer::SPropertyViewer::FSelectedItem> Items = ViewModelTreeView->GetSelectedItems();
+	if (Items.Num() == 1 && Items[0].bIsContainerSelected)
 	{
-		return BlueprintEditor->InEditingMode();
+		const FGuid* VMGuidPtr = PropertyViewerHandles.Find(Items[0].Handle);
+		return VMGuidPtr && HandleCanRename(*VMGuidPtr);
 	}
 	return false;
 }
@@ -758,13 +782,16 @@ bool SMVVMViewModelPanel::RenameViewModelProperty(FGuid ViewModelGuid, const FTe
 			{
 				if (const FMVVMBlueprintViewModelContext* ViewModelContext = View->FindViewModel(ViewModelGuid))
 				{
-					if (bCommit)
-					{
-						return GEditor->GetEditorSubsystem<UMVVMEditorSubsystem>()->RenameViewModel(WidgetBP, ViewModelContext->GetViewModelName(), *NewNameString, OutErrorMessage);
-					}
-					else
+					if (ViewModelContext->bCanRename)
 					{
-						return GEditor->GetEditorSubsystem<UMVVMEditorSubsystem>()->VerifyViewModelRename(WidgetBP, ViewModelContext->GetViewModelName(), *NewNameString, OutErrorMessage);
+						if (bCommit)
+						{
+							return GEditor->GetEditorSubsystem<UMVVMEditorSubsystem>()->RenameViewModel(WidgetBP, ViewModelContext->GetViewModelName(), *NewNameString, OutErrorMessage);
+						}
+						else
+						{
+							return GEditor->GetEditorSubsystem<UMVVMEditorSubsystem>()->VerifyViewModelRename(WidgetBP, ViewModelContext->GetViewModelName(), *NewNameString, OutErrorMessage);
+						}
 					}
 				}
 			}
diff --git a/Engine/Plugins/Runtime/ModelViewViewModel/Source/ModelViewViewModelEditor/Private/Widgets/SMVVMViewModelPanel.h b/Engine/Plugins/Runtime/ModelViewViewModel/Source/ModelViewViewModelEditor/Private/Widgets/SMVVMViewModelPanel.h
--- a/Engine/Plugins/Runtime/ModelViewViewModel/Source/ModelViewViewModelEditor/Private/Widgets/SMVVMViewModelPanel.h
+++ b/Engine/Plugins/Runtime/ModelViewViewModel/Source/ModelViewViewModelEditor/Private/Widgets/SMVVMViewModelPanel.h
@@ -86,6 +86,7 @@ private:
 	TSharedPtr<SWidget> HandleContextMenuOpening(UE::PropertyViewer::SPropertyViewer::FHandle ContainerHandle, TArrayView<const FFieldVariant> Field) const;
 	void HandleSelectionChanged(UE::PropertyViewer::SPropertyViewer::FHandle, TArrayView<const FFieldVariant>, ESelectInfo::Type);
 	void HandleEditorSelectionChanged();
+	bool HandleCanRename(FGuid ViewModelGuid) const;
 	bool HandleVerifyNameTextChanged(const FText& InText, FText& OutErrorMessage, FGuid ViewModelGuid);
 	void HandleNameTextCommited(const FText& InText, ETextCommit::Type CommitInfo, FGuid ViewModelGuid);
 	FReply HandleDragDetected(const FGeometry& MyGeometry, const FPointerEvent& MouseEvent, UE::PropertyViewer::SPropertyViewer::FHandle ContainerHandle, TArrayView<const FFieldVariant> Fields) const;
