MVVM: Add setting to show the View settings.

diff --git a/Engine/Plugins/Runtime/ModelViewViewModel/Source/ModelViewViewModelBlueprint/Public/MVVMDeveloperProjectSettings.h b/Engine/Plugins/Runtime/ModelViewViewModel/Source/ModelViewViewModelBlueprint/Public/MVVMDeveloperProjectSettings.h
--- a/Engine/Plugins/Runtime/ModelViewViewModel/Source/ModelViewViewModelBlueprint/Public/MVVMDeveloperProjectSettings.h
+++ b/Engine/Plugins/Runtime/ModelViewViewModel/Source/ModelViewViewModelBlueprint/Public/MVVMDeveloperProjectSettings.h
@@ -119,6 +119,10 @@ public:
 	/** For the binding list widget, allow the user to edit the binding in the detail view. */
 	UPROPERTY(EditAnywhere, config, Category = "Viewmodel")
 	bool bShowDetailViewOptionInBindingPanel = true;
+	
+	/** For the binding list widget and the viewmodel panel, allow the user to edit the view settings in the detail view. */
+	UPROPERTY(EditAnywhere, config, Category = "Viewmodel")
+	bool bShowViewSettings = true;
 
 	/** When a conversion function requires a wrapper graph, add and save the generated graph to the blueprint. */
 	UPROPERTY(EditAnywhere, config, Category = "Viewmodel")
diff --git a/Engine/Plugins/Runtime/ModelViewViewModel/Source/ModelViewViewModelEditor/Private/Widgets/SMVVMViewBindingPanel.cpp b/Engine/Plugins/Runtime/ModelViewViewModel/Source/ModelViewViewModelEditor/Private/Widgets/SMVVMViewBindingPanel.cpp
--- a/Engine/Plugins/Runtime/ModelViewViewModel/Source/ModelViewViewModelEditor/Private/Widgets/SMVVMViewBindingPanel.cpp
+++ b/Engine/Plugins/Runtime/ModelViewViewModel/Source/ModelViewViewModelEditor/Private/Widgets/SMVVMViewBindingPanel.cpp
@@ -463,14 +463,17 @@ void SBindingsPanel::RegisterSettingsMenu()
 		{
 			if (const UWidgetBlueprintToolMenuContext* Context = InSection.FindContext<UWidgetBlueprintToolMenuContext>())
 			{
-				InSection.AddMenuEntry(
-					"ViewSettings"
-					, LOCTEXT("ViewSettings", "View Settings")
-					, LOCTEXT("ViewSettingsTooltip", "View Settings")
-					, FSlateIcon(FMVVMEditorStyle::Get().GetStyleSetName(), "BlueprintView.TabIcon")
-					, FUIAction(FExecuteAction::CreateStatic(UE::MVVM::Private::SetSelectObjectsToViewSettings, Context->WidgetBlueprintEditor))
-					, EUserInterfaceActionType::Button
-				);
+				if (GetDefault<UMVVMDeveloperProjectSettings>()->bShowViewSettings)
+				{
+					InSection.AddMenuEntry(
+						"ViewSettings"
+						, LOCTEXT("ViewSettings", "View Settings")
+						, LOCTEXT("ViewSettingsTooltip", "View Settings")
+						, FSlateIcon(FMVVMEditorStyle::Get().GetStyleSetName(), "BlueprintView.TabIcon")
+						, FUIAction(FExecuteAction::CreateStatic(UE::MVVM::Private::SetSelectObjectsToViewSettings, Context->WidgetBlueprintEditor))
+						, EUserInterfaceActionType::Button
+					);
+				}
 			}
 		}));
 }
diff --git a/Engine/Plugins/Runtime/ModelViewViewModel/Source/ModelViewViewModelEditor/Private/Widgets/SMVVMViewModelPanel.cpp b/Engine/Plugins/Runtime/ModelViewViewModel/Source/ModelViewViewModelEditor/Private/Widgets/SMVVMViewModelPanel.cpp
--- a/Engine/Plugins/Runtime/ModelViewViewModel/Source/ModelViewViewModelEditor/Private/Widgets/SMVVMViewModelPanel.cpp
+++ b/Engine/Plugins/Runtime/ModelViewViewModel/Source/ModelViewViewModelEditor/Private/Widgets/SMVVMViewModelPanel.cpp
@@ -6,6 +6,7 @@
 #include "Dialogs/Dialogs.h"
 #include "IStructureDetailsView.h"
 #include "MVVMBlueprintView.h"
+#include "MVVMDeveloperProjectSettings.h"
 #include "MVVMEditorSubsystem.h"
 #include "MVVMWidgetBlueprintExtension_View.h"
 #include "PropertyEditorModule.h"
@@ -104,14 +105,17 @@ void SMVVMViewModelPanel::RegisterMenu()
 			{
 				if (const UWidgetBlueprintToolMenuContext* Context = InSection.FindContext<UWidgetBlueprintToolMenuContext>())
 				{
-					InSection.AddMenuEntry(
-						"ViewSettings"
-						, LOCTEXT("ViewSettings", "View Settings")
-						, LOCTEXT("ViewSettingsTooltip", "View Settings")
-						, FSlateIcon(FMVVMEditorStyle::Get().GetStyleSetName(), "BlueprintView.TabIcon")
-						, FUIAction(FExecuteAction::CreateStatic(UE::MVVM::Private::SetSelectObjectsToView, Context->WidgetBlueprintEditor))
-						, EUserInterfaceActionType::Button
-					);
+					if (GetDefault<UMVVMDeveloperProjectSettings>()->bShowViewSettings)
+					{
+						InSection.AddMenuEntry(
+							"ViewSettings"
+							, LOCTEXT("ViewSettings", "View Settings")
+							, LOCTEXT("ViewSettingsTooltip", "View Settings")
+							, FSlateIcon(FMVVMEditorStyle::Get().GetStyleSetName(), "BlueprintView.TabIcon")
+							, FUIAction(FExecuteAction::CreateStatic(UE::MVVM::Private::SetSelectObjectsToView, Context->WidgetBlueprintEditor))
+							, EUserInterfaceActionType::Button
+						);
+					}
 				}
 			}));
 	}
