MVVM: Fix conversion function allow test when the graph is a wrapper.

diff --git a/Engine/Plugins/Runtime/ModelViewViewModel/Source/ModelViewViewModelBlueprint/Private/Bindings/MVVMCompiledBindingLibraryCompiler.cpp b/Engine/Plugins/Runtime/ModelViewViewModel/Source/ModelViewViewModelBlueprint/Private/Bindings/MVVMCompiledBindingLibraryCompiler.cpp
--- a/Engine/Plugins/Runtime/ModelViewViewModel/Source/ModelViewViewModelBlueprint/Private/Bindings/MVVMCompiledBindingLibraryCompiler.cpp
+++ b/Engine/Plugins/Runtime/ModelViewViewModel/Source/ModelViewViewModelBlueprint/Private/Bindings/MVVMCompiledBindingLibraryCompiler.cpp
@@ -474,11 +474,6 @@ TValueOrError<FCompiledBindingLibraryCompiler::FFieldPathHandle, FText> FCompile
 		}
 	}
 
-	if (!GetDefault<UMVVMDeveloperProjectSettings>()->IsConversionFunctionAllowed(Function))
-	{
-		return MakeError(LOCTEXT("ConversionFunctionNotAllow", "The conversion function is not allowed."));
-	}
-
 	TArray<int32> RawFieldIndexes;
 	RawFieldIndexes.Add(Impl->AddUniqueField(FMVVMConstFieldVariant(Function)));
 	const int32 FoundFieldPath = Impl->FieldPaths.IndexOfByPredicate([&RawFieldIndexes](const Private::FRawFieldPath& Other)
diff --git a/Engine/Plugins/Runtime/ModelViewViewModel/Source/ModelViewViewModelBlueprint/Private/MVVMViewBlueprintCompiler.cpp b/Engine/Plugins/Runtime/ModelViewViewModel/Source/ModelViewViewModelBlueprint/Private/MVVMViewBlueprintCompiler.cpp
--- a/Engine/Plugins/Runtime/ModelViewViewModel/Source/ModelViewViewModelBlueprint/Private/MVVMViewBlueprintCompiler.cpp
+++ b/Engine/Plugins/Runtime/ModelViewViewModel/Source/ModelViewViewModelBlueprint/Private/MVVMViewBlueprintCompiler.cpp
@@ -1382,6 +1382,26 @@ bool FMVVMViewBlueprintCompiler::PreCompileBindings(UWidgetBlueprintGeneratedCla
 				bIsBindingsValid = false;
 				return MakeError();
 			}
+
+			TVariant<const UFunction*, TSubclassOf<UK2Node>> FunctionOrWrapperFunction = ViewConversionFunction->GetConversionFunction(Class);
+			if (FunctionOrWrapperFunction.IsType<const UFunction*>())
+			{
+				if (!GetDefault<UMVVMDeveloperProjectSettings>()->IsConversionFunctionAllowed(FunctionOrWrapperFunction.Get<const UFunction*>()))
+				{
+					AddMessageForBinding(Binding, BlueprintView, FText::Format(LOCTEXT("ConversionFunctionNotAllow", "The conversion function {0} is not allowed."),
+						FText::FromName(FunctionOrWrapperFunction.Get<const UFunction*>()->GetFName())),
+						EBindingMessageType::Error
+					);
+					bIsBindingsValid = false;
+					return MakeError();
+				}
+			}
+			else
+			{
+				AddMessageForBinding(Binding, BlueprintView, LOCTEXT("ConversionFunctionNodeNotAllow", "The conversion function node is not allowed."), EBindingMessageType::Error);
+				bIsBindingsValid = false;
+				return MakeError();
+			}
 		}
 		return MakeValue(ConversionFunction);
 	};
