MVVM: If we are not allowed to add the view, remove all the generated functions added by the view.

diff --git a/Engine/Plugins/Runtime/ModelViewViewModel/Source/ModelViewViewModelBlueprint/Private/MVVMWidgetBlueprintExtension_View.cpp b/Engine/Plugins/Runtime/ModelViewViewModel/Source/ModelViewViewModelBlueprint/Private/MVVMWidgetBlueprintExtension_View.cpp
--- a/Engine/Plugins/Runtime/ModelViewViewModel/Source/ModelViewViewModelBlueprint/Private/MVVMWidgetBlueprintExtension_View.cpp
+++ b/Engine/Plugins/Runtime/ModelViewViewModel/Source/ModelViewViewModelBlueprint/Private/MVVMWidgetBlueprintExtension_View.cpp
@@ -151,6 +151,15 @@ void UMVVMWidgetBlueprintExtension_View::HandleFinishCompilingClass(UWidgetBluep
 				CurrentCompilerContext->AddExtension(Class, ViewExtension);
 			}
 		}
+		else if (bCompiled)
+		{
+			// If we are not allowed to add the view class, add the transient flags on added conversion graph.
+			for (TFieldIterator<UFunction> FunctionIter(Class, EFieldIteratorFlags::ExcludeSuper); FunctionIter; ++FunctionIter)
+			{
+				UFunction* Function = *FunctionIter;
+				Function->SetFlags(RF_Transient);
+			}
+		}
 	}
 	else
 	{
