MVVM: Backend for binding event. Didn't use the ComponentDelegateBinding because of initialization order is different. Create a new entry in the ViewClass and reuse the binding library to find the delegate property. At editor time, a BP function is gnerated, similar to Complex Conversion function. When the event is broadcast at runtim, the BP function is executed.
The generated function can assign a property value on a viewmodel, a widget or any property of the UserWidget.

diff --git a/Engine/Plugins/Runtime/ModelViewViewModel/Source/ModelViewViewModel/Private/View/MVVMView.cpp b/Engine/Plugins/Runtime/ModelViewViewModel/Source/ModelViewViewModel/Private/View/MVVMView.cpp
--- a/Engine/Plugins/Runtime/ModelViewViewModel/Source/ModelViewViewModel/Private/View/MVVMView.cpp
+++ b/Engine/Plugins/Runtime/ModelViewViewModel/Source/ModelViewViewModel/Private/View/MVVMView.cpp
@@ -56,6 +56,7 @@ void UMVVMView::Construct()
 	{
 		InitializeSources();
 	}
+	InitializeEvents();
 
 	bConstructed = true;
 }
@@ -70,6 +71,7 @@ void UMVVMView::Destruct()
 	UE::MVVM::FDebugging::BroadcastViewBeginDestruction(this);
 #endif
 
+	UninitializeEvents();
 	UninitializeSources(); // and bindings
 }
 
@@ -390,6 +392,7 @@ bool UMVVMView::SetSourceInternal(FName ViewModelName, TScriptInterface<INotifyF
 	if (PreviousValue != NewValue.GetObject())
 	{
 		const TArrayView<const FMVVMViewClass_CompiledBinding> CompiledBindings = ClassExtension->GetCompiledBindings();
+		const TArrayView<const FMVVMViewClass_CompiledEvent> CompiledEvents = ClassExtension->GetCompiledEvents();
 
 		// Unregister any bindings from that source
 		if (bBindingsInitialized)
@@ -404,6 +407,19 @@ bool UMVVMView::SetSourceInternal(FName ViewModelName, TScriptInterface<INotifyF
 			}
 		}
 
+		// Unregisterer any events from that source
+		if (bConstructed)
+		{
+			for (int32 Index = 0; Index < CompiledEvents.Num(); ++Index)
+			{
+				const FMVVMViewClass_CompiledEvent& CompiledEvent = CompiledEvents[Index];
+				if (IsLibraryEventEnabled(Index) && CompiledEvent.GetSourceName() == ViewModelName)
+				{
+					DisableLibraryEvent(CompiledEvent, Index);
+				}
+			}
+		}
+
 		if (!ViewSource.bSetManually && ViewSource.Source)
 		{
 			SourceCreator.DestroyInstance(ViewSource.Source, this);
@@ -423,6 +439,7 @@ bool UMVVMView::SetSourceInternal(FName ViewModelName, TScriptInterface<INotifyF
 
 		bool bPreviousEveryTickBinding = bHasEveryTickBinding;
 		bHasEveryTickBinding = false;
+
 		// Register back any bindings that was previously enabled
 		if (bBindingsInitialized && NewValue.GetObject())
 		{
@@ -454,6 +471,20 @@ bool UMVVMView::SetSourceInternal(FName ViewModelName, TScriptInterface<INotifyF
 			}
 		}
 
+		// Register back any events that was previously enabled
+		if (bConstructed && NewValue.GetObject())
+		{
+			// Enabled the default events
+			for (int32 Index = 0; Index < CompiledEvents.Num(); ++Index)
+			{
+				const FMVVMViewClass_CompiledEvent& CompiledEvent = CompiledEvents[Index];
+				if (CompiledEvent.GetSourceName() == ViewModelName)
+				{
+					EnableLibraryEvent(CompiledEvent, Index);
+				}
+			}
+		}
+
 		if (bBindingsInitialized && bPreviousEveryTickBinding != bHasEveryTickBinding)
 		{
 			if (bHasEveryTickBinding)
@@ -862,6 +893,133 @@ void UMVVMView::UnregisterLibraryBinding(const FMVVMViewClass_CompiledBinding& B
 }
 
 
+void UMVVMView::InitializeEvents()
+{
+	UUserWidget* UserWidget = GetUserWidget();
+	check(UserWidget);
+	check(ClassExtension);
+
+	const TArrayView<const FMVVMViewClass_CompiledEvent>& CompiledEvent = ClassExtension->GetCompiledEvents();
+
+	ensure(RegisteredLibraryEvents.IsEmpty());
+	RegisteredLibraryEvents.Reset();
+	RegisteredLibraryEvents.Add(false, CompiledEvent.Num());
+
+	for (int32 Index = 0; Index < CompiledEvent.Num(); ++Index)
+	{
+		const FMVVMViewClass_CompiledEvent& Event = CompiledEvent[Index];
+		EnableLibraryEvent(Event, Index);
+	}
+}
+
+
+void UMVVMView::UninitializeEvents()
+{
+	check(ClassExtension);
+
+	const TArrayView<const FMVVMViewClass_CompiledEvent>& CompiledEvent = ClassExtension->GetCompiledEvents();
+	for (int32 Index = 0; Index < CompiledEvent.Num(); ++Index)
+	{
+		const FMVVMViewClass_CompiledEvent& Event = CompiledEvent[Index];
+		if (RegisteredLibraryEvents[Index])
+		{
+			DisableLibraryEvent(Event, Index);
+		}
+	}
+}
+
+
+void UMVVMView::EnableLibraryEvent(const FMVVMViewClass_CompiledEvent& Event, int32 EventIndex)
+{
+	UUserWidget* UserWidget = GetUserWidget();
+	check(UserWidget);
+
+	auto LogMessage = [this, &Event, EventIndex](const FText& Message)
+	{
+		UE::MVVM::FMessageLog Log(GetUserWidget());
+		Log.Error(FText::Format(LOCTEXT("EnableLibraryEventFailed_Format", "Widget '{0}' can't register event '{1}'. {2}")
+			, FText::FromString(GetFullName())
+#if UE_WITH_MVVM_DEBUGGING
+			, FText::FromString(Event.ToString(ClassExtension->GetBindingLibrary(), FMVVMViewClass_CompiledEvent::FToStringArgs()))
+#else
+			, FText::AsNumber(EventIndex)
+#endif
+			, Message
+			));
+	};
+
+	const UFunction* FunctionToBind = UserWidget->GetClass()->FindFunctionByName(Event.GetUserWidgetFunctionName());
+	if (FunctionToBind == nullptr)
+	{
+		LogMessage(LOCTEXT("EnableLibraryEventFailed_FunctionNotFound", "Function not found"));
+		return;
+	}
+
+	const FMVVMCompiledBindingLibrary& Library = ClassExtension->GetBindingLibrary();
+	TValueOrError<UE::MVVM::FFieldContext, void> FieldPathResult = Library.EvaluateFieldPath(UserWidget, Event.GetMulticastDelegatePath());
+	if (FieldPathResult.HasError())
+	{
+		LogMessage(LOCTEXT("EnableLibraryEventFailed_CantEvaluateBindingPath", "The binding path couldn't be evaluated."));
+		return;
+	}
+
+	UE::MVVM::FFieldContext FieldContext = FieldPathResult.StealValue();
+	if (FieldContext.GetObjectVariant().IsNull() || !FieldContext.GetObjectVariant().IsUObject())
+	{
+		LogMessage(LOCTEXT("EnableLibraryEventFailed_OwnerInvalid", "The owner is invalid."));
+		return;
+	}
+
+
+	FMulticastDelegateProperty* MulticastDelegateProp = FieldContext.GetFieldVariant().IsProperty() ? CastField<FMulticastDelegateProperty>(FieldContext.GetFieldVariant().GetProperty()) : nullptr;
+	if (!MulticastDelegateProp)
+	{
+		LogMessage(LOCTEXT("EnableLibraryEventFailed_NotAProperty", "The path doesn't point to a multicast delegate property."));
+		return;
+	}
+
+	FScriptDelegate Delegate;
+	Delegate.BindUFunction(UserWidget, Event.GetUserWidgetFunctionName());
+	MulticastDelegateProp->AddDelegate(MoveTemp(Delegate), FieldContext.GetObjectVariant().GetUObject());
+	RegisteredLibraryEvents[EventIndex] = true;
+}
+
+
+void UMVVMView::DisableLibraryEvent(const FMVVMViewClass_CompiledEvent& Event, int32 EventIndex)
+{
+	check(IsLibraryEventEnabled(EventIndex));
+
+	RegisteredLibraryEvents[EventIndex] = false;
+
+	UUserWidget* UserWidget = GetUserWidget();
+	check(UserWidget);
+	const FMVVMCompiledBindingLibrary& Library = ClassExtension->GetBindingLibrary();
+
+	TValueOrError<UE::MVVM::FFieldContext, void> FieldPathResult = Library.EvaluateFieldPath(UserWidget, Event.GetMulticastDelegatePath());
+	if (FieldPathResult.HasError())
+	{
+		return;
+	}
+
+	UE::MVVM::FFieldContext FieldContext = FieldPathResult.StealValue();
+	if (FieldContext.GetObjectVariant().IsNull())
+	{
+		return;
+	}
+
+	FMulticastDelegateProperty* MulticastDelegateProp = CastFieldChecked<FMulticastDelegateProperty>(FieldContext.GetFieldVariant().GetProperty());
+	FScriptDelegate Delegate;
+	Delegate.BindUFunction(UserWidget, Event.GetUserWidgetFunctionName());
+	MulticastDelegateProp->RemoveDelegate(Delegate, FieldContext.GetObjectVariant().GetUObject());
+}
+
+
+bool UMVVMView::IsLibraryEventEnabled(int32 EventIndex) const
+{
+	return ensure(RegisteredLibraryEvents.IsValidIndex(EventIndex)) && RegisteredLibraryEvents[EventIndex];
+}
+
+
 TScriptInterface<INotifyFieldValueChanged> UMVVMView::FindSource(const FMVVMViewClass_CompiledBinding& Binding, int32 BindingIndex, bool bAllowNull) const
 {
 	UUserWidget* UserWidget = GetUserWidget();
diff --git a/Engine/Plugins/Runtime/ModelViewViewModel/Source/ModelViewViewModel/Private/View/MVVMViewClass.cpp b/Engine/Plugins/Runtime/ModelViewViewModel/Source/ModelViewViewModel/Private/View/MVVMViewClass.cpp
--- a/Engine/Plugins/Runtime/ModelViewViewModel/Source/ModelViewViewModel/Private/View/MVVMViewClass.cpp
+++ b/Engine/Plugins/Runtime/ModelViewViewModel/Source/ModelViewViewModel/Private/View/MVVMViewClass.cpp
@@ -490,6 +490,41 @@ FString FMVVMViewClass_CompiledBinding::ToString(const FMVVMCompiledBindingLibra
 }
 #endif
 
+///////////////////////////////////////////////////////////////////////
+// 
+///////////////////////////////////////////////////////////////////////
+#if UE_WITH_MVVM_DEBUGGING
+FString FMVVMViewClass_CompiledEvent::ToString(const FMVVMCompiledBindingLibrary& BindingLibrary, FToStringArgs Args) const
+{
+	TStringBuilder<1024> StringBuilder;
+
+	if (GetMulticastDelegatePath().IsValid())
+	{
+		StringBuilder << TEXT("Event: ");
+
+		if (Args.bUseDisplayName)
+		{
+			StringBuilder << TEXT('"');
+		}
+
+		TValueOrError<FString, FString> SourceFieldPathString = BindingLibrary.FieldPathToString(GetMulticastDelegatePath(), Args.bUseDisplayName);
+		StringBuilder << (SourceFieldPathString.HasValue() ? SourceFieldPathString.StealValue() : SourceFieldPathString.StealError());
+
+		if (Args.bUseDisplayName)
+		{
+			StringBuilder << TEXT("\"\n");
+		}
+	}
+
+	StringBuilder << TEXT("SourceName: ");
+	StringBuilder << GetSourceName();
+
+	StringBuilder << TEXT(", Function: ");
+	StringBuilder << GetUserWidgetFunctionName();
+
+	return StringBuilder.ToString();
+}
+#endif
 
 ///////////////////////////////////////////////////////////////////////
 // 
@@ -595,6 +630,17 @@ void UMVVMViewClass::Log(FMVVMViewClass_SourceCreator::FToStringArgs SourceArgs,
 	}
 	UE_LOG(LogMVVM, Log, TEXT("%s"), Builder.ToString());
 
+	Builder.Reset();
+
+	Builder << TEXT("Compiled Events for: ");
+	Builder << GetOutermost()->GetFName();
+	for (const FMVVMViewClass_CompiledEvent& Event : GetCompiledEvents())
+	{
+		Builder << TEXT("\n");
+		Builder << Event.ToString(GetBindingLibrary(), FMVVMViewClass_CompiledEvent::FToStringArgs());
+	}
+	UE_LOG(LogMVVM, Log, TEXT("%s"), Builder.ToString());
+
 	if (!bPreviousLoad)
 	{
 #if WITH_EDITOR
diff --git a/Engine/Plugins/Runtime/ModelViewViewModel/Source/ModelViewViewModel/Public/View/MVVMView.h b/Engine/Plugins/Runtime/ModelViewViewModel/Source/ModelViewViewModel/Public/View/MVVMView.h
--- a/Engine/Plugins/Runtime/ModelViewViewModel/Source/ModelViewViewModel/Public/View/MVVMView.h
+++ b/Engine/Plugins/Runtime/ModelViewViewModel/Source/ModelViewViewModel/Public/View/MVVMView.h
@@ -13,6 +13,7 @@ template <typename InterfaceType> class TScriptInterface;
 
 class UMVVMViewClass;
 struct FMVVMViewClass_CompiledBinding;
+struct FMVVMViewClass_CompiledEvent;
 struct FMVVMViewDelayedBinding;
 class UMVVMViewModelBase;
 class UWidget;
@@ -21,6 +22,7 @@ namespace UE::MVVM
 	class FDebugging;
 }
 
+/** */
 USTRUCT()
 struct FMVVMViewSource
 {
@@ -153,6 +155,12 @@ private:
 	FDelegateHandle RegisterLibraryBinding(const FMVVMViewClass_CompiledBinding& Binding, int32 BindingIndex);
 	void UnregisterLibraryBinding(const FMVVMViewClass_CompiledBinding& Binding, FDelegateHandle Handle, int32 BindingIndex);
 
+	void InitializeEvents();
+	void UninitializeEvents();
+	void EnableLibraryEvent(const FMVVMViewClass_CompiledEvent& Item, int32 EventIndex);
+	void DisableLibraryEvent(const FMVVMViewClass_CompiledEvent& Item, int32 EventIndex);
+	bool IsLibraryEventEnabled(int32 EventIndex) const;
+
 	TScriptInterface<INotifyFieldValueChanged> FindSource(const FMVVMViewClass_CompiledBinding& Binding, int32 BindingIndex, bool bAllowNull) const;
 	FMVVMViewSource* FindViewSource(const FName SourceName);
 	const FMVVMViewSource* FindViewSource(const FName SourceName) const;
@@ -166,6 +174,9 @@ private:
 
 	/** The binding that are registered by the view to the sources. */
 	TArray<FDelegateHandle> RegisteredLibraryBindings;
+	
+	/** The event that are registered by the view to the sources. */
+	TBitArray<> RegisteredLibraryEvents;
 
 	/** Should log when a binding is executed. */
 	UPROPERTY(EditAnywhere, Transient, Category = "Viewmodel")
diff --git a/Engine/Plugins/Runtime/ModelViewViewModel/Source/ModelViewViewModel/Public/View/MVVMViewClass.h b/Engine/Plugins/Runtime/ModelViewViewModel/Source/ModelViewViewModel/Public/View/MVVMViewClass.h
--- a/Engine/Plugins/Runtime/ModelViewViewModel/Source/ModelViewViewModel/Public/View/MVVMViewClass.h
+++ b/Engine/Plugins/Runtime/ModelViewViewModel/Source/ModelViewViewModel/Public/View/MVVMViewClass.h
@@ -311,6 +311,54 @@ private:
 };
 
 
+/**
+ * A compiled and shared delegate bindings
+ */
+USTRUCT()
+struct MODELVIEWVIEWMODEL_API FMVVMViewClass_CompiledEvent
+{
+	GENERATED_BODY()
+
+	friend UE::MVVM::Private::FMVVMViewBlueprintCompiler;
+
+public:
+	/** @return The unique name of the source object that owns the multicast delegate. */
+	FName GetSourceName() const
+	{
+		return SourceName;
+	}
+
+	const FMVVMVCompiledFieldPath& GetMulticastDelegatePath() const
+	{
+		return FieldPath;
+	}
+
+	const FName GetUserWidgetFunctionName() const
+	{
+		return FunctionName;
+	}
+
+#if UE_WITH_MVVM_DEBUGGING
+	struct FToStringArgs
+	{
+		bool bUseDisplayName = true;
+	};
+	/** @return a human readable version of the binding that can be use for debugging purposes. */
+	FString ToString(const FMVVMCompiledBindingLibrary& CompiledBindingLibrary, FToStringArgs Args) const;
+#endif
+
+private:
+	UPROPERTY()
+	FMVVMVCompiledFieldPath FieldPath;
+
+	UPROPERTY()
+	FName FunctionName;
+
+	UPROPERTY()
+	FName SourceName;
+};
+
+
 /**
  * Shared between every instances of the same View class.
  */
@@ -366,6 +414,12 @@ public:
 	{
 		return CompiledBindings[Index];
 	}
+	
+	/**  */
+	const TArrayView<const FMVVMViewClass_CompiledEvent> GetCompiledEvents() const
+	{
+		return MakeArrayView(CompiledEvents);
+	}
 
 #if UE_WITH_MVVM_DEBUGGING
 	void Log(FMVVMViewClass_SourceCreator::FToStringArgs SourceArgs, FMVVMViewClass_CompiledBinding::FToStringArgs BindingArgs) const;
@@ -377,6 +431,12 @@ private:
 	{
 		return MakeArrayView(CompiledBindings);
 	}
+	
+	/**  */
+	TArrayView<FMVVMViewClass_CompiledEvent> GetCompiledEvents()
+	{
+		return MakeArrayView(CompiledEvents);
+	}
 
 #if WITH_EDITOR
 	void HandleBlueprintCompiled();
@@ -390,6 +450,10 @@ private:
 	/** */
 	UPROPERTY()
 	TArray<FMVVMViewClass_CompiledBinding> CompiledBindings;
+	
+	/** */
+	UPROPERTY()
+	TArray<FMVVMViewClass_CompiledEvent> CompiledEvents;
 
 	/** All the bindings shared between all the View instance. */
 	UPROPERTY()
diff --git a/Engine/Plugins/Runtime/ModelViewViewModel/Source/ModelViewViewModelBlueprint/Private/Bindings/MVVMConversionFunctionHelper.cpp b/Engine/Plugins/Runtime/ModelViewViewModel/Source/ModelViewViewModelBlueprint/Private/Bindings/MVVMConversionFunctionHelper.cpp
--- a/Engine/Plugins/Runtime/ModelViewViewModel/Source/ModelViewViewModelBlueprint/Private/Bindings/MVVMConversionFunctionHelper.cpp
+++ b/Engine/Plugins/Runtime/ModelViewViewModel/Source/ModelViewViewModelBlueprint/Private/Bindings/MVVMConversionFunctionHelper.cpp
@@ -66,7 +66,7 @@ namespace Private
 		return nullptr;
 	}
 
-	TTuple<UEdGraph*, UK2Node_FunctionEntry*, UK2Node_FunctionResult*> CreateGraph(UBlueprint* Blueprint, FName GraphName, bool bIsEditable, bool bAddToBlueprint)
+	TTuple<UEdGraph*, UK2Node_FunctionEntry*, UK2Node_FunctionResult*> CreateGraph(UBlueprint* Blueprint, FName GraphName, const UFunction* FunctionEntryDefinition, bool bIsEditable, bool bAddToBlueprint)
 	{
 		FName UniqueFunctionName = FBlueprintEditorUtils::FindUniqueKismetName(Blueprint, GraphName.ToString());
 		UEdGraph* FunctionGraph = FBlueprintEditorUtils::CreateNewGraph(Blueprint, UniqueFunctionName, UEdGraph::StaticClass(), UEdGraphSchema_K2::StaticClass());
@@ -87,7 +87,16 @@ namespace Private
 		// function entry node
 		FGraphNodeCreator<UK2Node_FunctionEntry> FunctionEntryCreator(*FunctionGraph);
 		UK2Node_FunctionEntry* FunctionEntry = FunctionEntryCreator.CreateNode();
-		FunctionEntry->FunctionReference.SetSelfMember(FunctionGraph->GetFName());
+		if (FunctionEntryDefinition)
+		{
+			UClass* OwnerClass = FunctionEntryDefinition->GetOwnerClass();
+			FunctionEntry->FunctionReference.SetExternalMember(FunctionEntryDefinition->GetFName(), OwnerClass);
+			FunctionEntry->CustomGeneratedFunctionName = GraphName;
+		}
+		else
+		{
+			FunctionEntry->FunctionReference.SetSelfMember(FunctionGraph->GetFName());
+		}
 		FunctionEntry->AddExtraFlags(FUNC_BlueprintCallable | FUNC_BlueprintPure | FUNC_Const | FUNC_Protected | FUNC_Final);
 		FunctionEntry->bIsEditable = bIsEditable;
 		FunctionEntry->MetaData.Category = FText::FromName(ConversionFunctionCategory.Resolve());
@@ -132,7 +141,7 @@ FName CreateWrapperName(const FMVVMBlueprintViewBinding& Binding, bool bSourceTo
 	return FName(StringBuilder.ToString());
 }
 
-TPair<UEdGraph*, UK2Node*> CreateGraph(UBlueprint* Blueprint, FName GraphName, const UFunction* ConversionFunction, bool bTransient)
+TPair<UEdGraph*, UK2Node*> CreateGraph(UBlueprint* Blueprint, FName GraphName, const UFunction* Signature, const UFunction* FunctionToWrap, bool bTransient)
 {
 	bool bIsEditable = false;
 	bool bAddToBlueprint = !bTransient;
@@ -141,16 +150,16 @@ TPair<UEdGraph*, UK2Node*> CreateGraph(UBlueprint* Blueprint, FName GraphName, c
 	UK2Node_FunctionEntry* FunctionEntry = nullptr;
 	UK2Node_FunctionResult* FunctionResult = nullptr;
 	{
-		TTuple<UEdGraph*, UK2Node_FunctionEntry*, UK2Node_FunctionResult*> NewGraph = Private::CreateGraph(Blueprint, GraphName, bIsEditable, bAddToBlueprint);
+		TTuple<UEdGraph*, UK2Node_FunctionEntry*, UK2Node_FunctionResult*> NewGraph = Private::CreateGraph(Blueprint, GraphName, Signature, bIsEditable, bAddToBlueprint);
 		FunctionGraph = NewGraph.Get<0>();
 		FunctionEntry = NewGraph.Get<1>();
 		FunctionResult = NewGraph.Get<2>();
 	}
 
-	const FProperty* ReturnProperty = UE::MVVM::BindingHelper::GetReturnProperty(ConversionFunction);
-	check(ReturnProperty);
+	const FProperty* ReturnProperty = UE::MVVM::BindingHelper::GetReturnProperty(FunctionToWrap);
 
 	// create return value pin
+	if (ReturnProperty)
 	{
 		TSharedPtr<FUserPinInfo> PinInfo = MakeShared<FUserPinInfo>();
 		GetDefault<UEdGraphSchema_K2>()->ConvertPropertyToPinType(ReturnProperty, PinInfo->PinType);
@@ -164,7 +173,7 @@ TPair<UEdGraph*, UK2Node*> CreateGraph(UBlueprint* Blueprint, FName GraphName, c
 	{
 		FGraphNodeCreator<UK2Node_CallFunction> CallFunctionCreator(*FunctionGraph);
 		CallFunctionNode = CallFunctionCreator.CreateNode();
-		CallFunctionNode->SetFromFunction(ConversionFunction);
+		CallFunctionNode->SetFromFunction(FunctionToWrap);
 		CallFunctionNode->NodePosX = 0;
 		CallFunctionCreator.Finalize();
 		Private::MarkAsConversionFunction(CallFunctionNode, FunctionGraph);
@@ -172,13 +181,13 @@ TPair<UEdGraph*, UK2Node*> CreateGraph(UBlueprint* Blueprint, FName GraphName, c
 
 	// Make link Entry -> CallFunction || Entry -> Return
 	{
-		UEdGraphPin* FunctionEntryThenPin = FunctionEntry->FindPin(UEdGraphSchema_K2::PN_Then, EGPD_Output);
+		UEdGraphPin* FunctionEntryThenPin = FunctionEntry->GetThenPin();
 		UEdGraphPin* FunctionResultExecPin = FunctionResult->GetExecPin();
 		
 		if (!CallFunctionNode->IsNodePure())
 		{
 			UEdGraphPin* CallFunctionExecPin = CallFunctionNode->GetExecPin();
-			UEdGraphPin* CallFunctionThenPin = CallFunctionNode->FindPin(UEdGraphSchema_K2::PN_Then, EGPD_Output);
+			UEdGraphPin* CallFunctionThenPin = CallFunctionNode->GetThenPin();
 
 			FunctionEntryThenPin->MakeLinkTo(CallFunctionExecPin);
 			CallFunctionThenPin->MakeLinkTo(FunctionResultExecPin);
@@ -192,6 +201,7 @@ TPair<UEdGraph*, UK2Node*> CreateGraph(UBlueprint* Blueprint, FName GraphName, c
 		}
 	}
 
+	if (ReturnProperty)
 	{
 		UEdGraphPin* FunctionReturnPin = CallFunctionNode->FindPin(ReturnProperty->GetName(), EGPD_Output);
 		UEdGraphPin* FunctionResultPin = FunctionResult->FindPin(ReturnProperty->GetFName(), EGPD_Input);
@@ -202,6 +212,57 @@ TPair<UEdGraph*, UK2Node*> CreateGraph(UBlueprint* Blueprint, FName GraphName, c
 	return { FunctionGraph , CallFunctionNode };
 }
 
+
+TPair<UEdGraph*, UK2Node*> CreateGraph(UBlueprint* Blueprint, FName GraphName, const UFunction* Signature, TSubclassOf<UK2Node> NodeType, bool bTransient, TFunctionRef<void(UK2Node*)> InitNodeCallback)
+{
+	bool bIsEditable = false;
+	bool bAddToBlueprint = !bTransient;
+
+	UEdGraph* FunctionGraph = nullptr;
+	UK2Node_FunctionEntry* FunctionEntry = nullptr;
+	UK2Node_FunctionResult* FunctionResult = nullptr;
+	{
+		TTuple<UEdGraph*, UK2Node_FunctionEntry*, UK2Node_FunctionResult*> NewGraph = Private::CreateGraph(Blueprint, GraphName, Signature, bIsEditable, bAddToBlueprint);
+		FunctionGraph = NewGraph.Get<0>();
+		FunctionEntry = NewGraph.Get<1>();
+		FunctionResult = NewGraph.Get<2>();
+	}
+
+	UK2Node* CallFunctionNode = nullptr;
+	{
+		FGraphNodeCreator<UK2Node> CallFunctionCreator(*FunctionGraph);
+		CallFunctionNode = CallFunctionCreator.CreateNode(true, NodeType);
+		InitNodeCallback(CallFunctionNode);
+		CallFunctionNode->NodePosX = 0;
+		CallFunctionCreator.Finalize();
+		Private::MarkAsConversionFunction(CallFunctionNode, FunctionGraph);
+	}
+
+	// Make link Entry -> CallFunction || Entry -> Return
+	{
+		UEdGraphPin* FunctionEntryThenPin = FunctionEntry->GetThenPin();
+		UEdGraphPin* FunctionResultExecPin = FunctionResult->GetExecPin();
+
+		if (!CallFunctionNode->IsNodePure())
+		{
+			UEdGraphPin* CallFunctionExecPin = CallFunctionNode->GetExecPin();
+			UEdGraphPin* CallFunctionThenPin = CallFunctionNode->GetThenPin();
+
+			FunctionEntryThenPin->MakeLinkTo(CallFunctionExecPin);
+			CallFunctionThenPin->MakeLinkTo(FunctionResultExecPin);
+
+			CallFunctionNode->NodePosY = 0;
+		}
+		else
+		{
+			FunctionEntryThenPin->MakeLinkTo(FunctionResultExecPin);
+			CallFunctionNode->NodePosY = 100;
+		}
+	}
+
+	return { FunctionGraph , CallFunctionNode };
+}
+
 UK2Node* GetWrapperNode(UEdGraph* Graph)
 {
 	if (Graph == nullptr)
@@ -592,12 +653,6 @@ void SetPropertyPathForPin(const UBlueprint* Blueprint, const FMVVMBlueprintProp
 							}
 						}
 					}
-
-					if (!NewPreviousDataPin)
-					{
-						ensure(false);
-						return;
-					}
 				}
 				else if (PreviousClass != nullptr)
 				{
@@ -638,6 +693,12 @@ void SetPropertyPathForPin(const UBlueprint* Blueprint, const FMVVMBlueprintProp
 				return;
 			}
 
+			if (!NewPreviousDataPin)
+			{
+				ensureMsgf(false, TEXT("A node in path doesn't have a return value."));
+				return;
+			}
+
 			// create new data connections
 			for (UEdGraphPin* Pin : NewNode->Pins)
 			{
@@ -695,8 +756,17 @@ void SetPropertyPathForPin(const UBlueprint* Blueprint, const FMVVMBlueprintProp
 				}
 			}
 		}
-		UEdGraphPin* LastExecPin = ConverionFunctionResult->FindPinChecked(UEdGraphSchema_K2::PN_Execute);
-		ThenPin->MakeLinkTo(LastExecPin);
+
+		// Make pin to the conversion node or to the return node
+		if (UEdGraphPin* CallFunctionExecPin = ConversionNode->FindPin(UEdGraphSchema_K2::PN_Execute))
+		{
+			ThenPin->MakeLinkTo(CallFunctionExecPin);
+		}
+		else
+		{
+			UEdGraphPin* FunctionResultExecPin = ConverionFunctionResult->GetExecPin();
+			ThenPin->MakeLinkTo(FunctionResultExecPin);
+		}
 	}
 }
 
diff --git a/Engine/Plugins/Runtime/ModelViewViewModel/Source/ModelViewViewModelBlueprint/Private/MVVMBlueprintView.cpp b/Engine/Plugins/Runtime/ModelViewViewModel/Source/ModelViewViewModelBlueprint/Private/MVVMBlueprintView.cpp
--- a/Engine/Plugins/Runtime/ModelViewViewModel/Source/ModelViewViewModelBlueprint/Private/MVVMBlueprintView.cpp
+++ b/Engine/Plugins/Runtime/ModelViewViewModel/Source/ModelViewViewModelBlueprint/Private/MVVMBlueprintView.cpp
@@ -5,6 +5,7 @@
 #include "Components/Widget.h"
 #include "Kismet2/BlueprintEditorUtils.h"
 #include "MVVMBlueprintViewConversionFunction.h"
+#include "MVVMBlueprintViewEvent.h"
 #include "MVVMWidgetBlueprintExtension_View.h"
 #include "WidgetBlueprint.h"
 #include "UObject/FortniteMainBranchObjectVersion.h"
@@ -196,6 +197,25 @@ const FMVVMBlueprintViewBinding* UMVVMBlueprintView::GetBinding(FGuid Id) const
 	return Bindings.FindByPredicate([Id](const FMVVMBlueprintViewBinding& Binding) { return Id == Binding.BindingId; });
 }
 
+UMVVMBlueprintViewEvent* UMVVMBlueprintView::AddDefaultEvent()
+{
+	UMVVMBlueprintViewEvent* Event = NewObject<UMVVMBlueprintViewEvent>(this);
+	Events.Add(Event);
+
+	OnBindingsAdded.Broadcast();
+	OnBindingsUpdated.Broadcast();
+	return Event;
+}
+
+void UMVVMBlueprintView::RemoveEvent(UMVVMBlueprintViewEvent* Event)
+{
+	if (Events.RemoveAll([Event](TObjectPtr<UMVVMBlueprintViewEvent>& Other){ return Other == Event; }) > 0)
+	{
+		Event->RemoveWrapperGraph();
+		OnBindingsUpdated.Broadcast();
+	}
+}
+
 TArray<FText> UMVVMBlueprintView::GetBindingMessages(FGuid Id, UE::MVVM::EBindingMessageType InMessageType) const
 {
 	TArray<FText> Results;
@@ -293,6 +313,10 @@ void UMVVMBlueprintView::PreSave(FObjectPreSaveContext Context)
 	{
 		Binding.Conversion.SavePinValues(WidgetBlueprint);
 	}
+	for (UMVVMBlueprintViewEvent* Event : Events)
+	{
+		Event->SavePinValues();
+	}
 
 	Super::PreSave(Context);
 }
@@ -321,6 +345,10 @@ void UMVVMBlueprintView::PostEditChangeChainProperty(FPropertyChangedChainEvent&
 	{
 		OnViewModelsUpdated.Broadcast();
 	}
+	if (PropertyChainEvent.PropertyChain.Contains(UMVVMBlueprintView::StaticClass()->FindPropertyByName(GET_MEMBER_NAME_CHECKED(UMVVMBlueprintView, Events))))
+	{
+		OnViewModelsUpdated.Broadcast();
+	}
 }
 
 void UMVVMBlueprintView::AddAssetTags(TArray<FAssetRegistryTag>& OutTags) const
diff --git a/Engine/Plugins/Runtime/ModelViewViewModel/Source/ModelViewViewModelBlueprint/Private/MVVMBlueprintViewBinding.cpp b/Engine/Plugins/Runtime/ModelViewViewModel/Source/ModelViewViewModelBlueprint/Private/MVVMBlueprintViewBinding.cpp
--- a/Engine/Plugins/Runtime/ModelViewViewModel/Source/ModelViewViewModelBlueprint/Private/MVVMBlueprintViewBinding.cpp
+++ b/Engine/Plugins/Runtime/ModelViewViewModel/Source/ModelViewViewModelBlueprint/Private/MVVMBlueprintViewBinding.cpp
@@ -74,150 +74,7 @@ FName FMVVMBlueprintViewBinding::GetFName() const
 
 namespace UE::MVVM::Private
 {
-	FText GetDisplayNameForField(const FMVVMConstFieldVariant& Field)
-	{
-		if (!Field.IsEmpty())
-		{
-			if (Field.IsProperty())
-			{
-				return Field.GetProperty()->GetDisplayNameText();
-			}
-			if (Field.IsFunction())
-			{
-				return Field.GetFunction()->GetDisplayNameText();
-			}
-		}
-		return LOCTEXT("None", "<None>");
-	}
-
-	FString GetDisplayNameForWidget(const UWidgetBlueprint* WidgetBlueprint, const FName& WidgetName)
-	{
-		if (UWidgetTree* WidgetTree = WidgetBlueprint->WidgetTree)
-		{
-			if (UWidget* FoundWidget = WidgetTree->FindWidget(WidgetName))
-			{
-				return FoundWidget->GetDisplayLabel().IsEmpty() ? WidgetName.ToString() : FoundWidget->GetDisplayLabel();
-			}
-		}
-
-		return WidgetName.IsNone() ? TEXT("<none>") : WidgetName.ToString();
-	}
-
-	void AppendViewModelPathString(const UWidgetBlueprint* WidgetBlueprint, const UMVVMBlueprintView* BlueprintView, const FMVVMBlueprintPropertyPath& ViewModelPath, FStringBuilderBase& PathBuilder, FStringBuilderBase& FunctionKeywordsBuilder, bool bUseDisplayName = false, bool bAppendFunctionKeywords = false)
-	{
-		if (ViewModelPath.IsEmpty())
-		{
-			PathBuilder << TEXT("<none>");
-			return;
-		}
-
-		const FMVVMBlueprintViewModelContext* ViewModel = BlueprintView != nullptr ? BlueprintView->FindViewModel(ViewModelPath.GetViewModelId()) : nullptr;
-		if (ViewModel)
-		{
-			PathBuilder << ViewModel->GetDisplayName().ToString();
-		}
-		else if (ViewModelPath.IsFromWidget())
-		{
-			if (bUseDisplayName)
-			{
-				PathBuilder << GetDisplayNameForWidget(WidgetBlueprint, ViewModelPath.GetWidgetName());
-			}
-			else
-			{
-				PathBuilder << ViewModelPath.GetWidgetName();
-			}
-		}
-		else
-		{
-			PathBuilder << TEXT("<none>");
-		}
-
-		TArray<UE::MVVM::FMVVMConstFieldVariant> Fields = ViewModelPath.GetFields(WidgetBlueprint->SkeletonGeneratedClass);
-		for (const UE::MVVM::FMVVMConstFieldVariant& Field : Fields)
-		{
-			PathBuilder << TEXT(".");
-			if (bUseDisplayName)
-			{
-				PathBuilder << GetDisplayNameForField(Field).ToString();
-			}
-			else
-			{
-				PathBuilder << Field.GetName().ToString();
-			}
-
-			if (Field.IsFunction() && bAppendFunctionKeywords)
-			{
-				FString FunctionKeywords = Field.GetFunction()->GetMetaData(FBlueprintMetadata::MD_FunctionKeywords);
-				if (!FunctionKeywords.IsEmpty())
-				{
-					FunctionKeywordsBuilder << TEXT(".");
-					FunctionKeywordsBuilder << FunctionKeywords;
-				}
-			}
-		}
-	}
-	
-	void AppendWidgetPathString(const UWidgetBlueprint* WidgetBlueprint, const FMVVMBlueprintPropertyPath& WidgetPath, FStringBuilderBase& PathBuilder, FStringBuilderBase& FunctionKeywordsBuilder, bool bUseDisplayName = false, bool bAppendFunctionKeywords = false)
-	{
-		if (WidgetBlueprint == nullptr)
-		{
-			PathBuilder << TEXT("<none>");
-			return;
-		}
-
-		if (WidgetPath.GetWidgetName().IsNone() && !bUseDisplayName)
-		{
-			PathBuilder << TEXT("<none>");
-		}
-		else if (WidgetBlueprint->GetFName() == WidgetPath.GetWidgetName())
-		{
-			if (bUseDisplayName)
-			{
-				PathBuilder << TEXT("self");
-			}
-			else
-			{
-				PathBuilder << WidgetBlueprint->GetName();
-			}
-		}
-		else
-		{
-			if (bUseDisplayName)
-			{
-				PathBuilder << GetDisplayNameForWidget(WidgetBlueprint, WidgetPath.GetWidgetName());
-			}
-			else
-			{
-				PathBuilder << WidgetPath.GetWidgetName();
-			}
-		}
-
-		TArray<UE::MVVM::FMVVMConstFieldVariant> Fields = WidgetPath.GetFields(WidgetBlueprint->SkeletonGeneratedClass);
-		for (const UE::MVVM::FMVVMConstFieldVariant& Field : Fields)
-		{
-			PathBuilder << TEXT(".");
-			if (bUseDisplayName)
-			{
-				PathBuilder << GetDisplayNameForField(Field).ToString();
-			}
-			else
-			{
-				PathBuilder << Field.GetName().ToString();
-			}
-
-			if (Field.IsFunction() && bAppendFunctionKeywords)
-			{
-				FString FunctionKeywords = Field.GetFunction()->GetMetaData(FBlueprintMetadata::MD_FunctionKeywords);
-				if (!FunctionKeywords.IsEmpty())
-				{
-					FunctionKeywordsBuilder << TEXT(".");
-					FunctionKeywordsBuilder << FunctionKeywords;
-				}
-			}
-		}
-	}
-
-	FString GetBindingPathName(const UWidgetBlueprint* WidgetBlueprint, const FMVVMBlueprintViewBinding& Binding, bool bIsSource, FStringBuilderBase& FunctionKeywordsBuilder, bool bUseDisplayName = false, bool bAppendFunctionKeywords = false)
+	FString GetBindingPathName(const UWidgetBlueprint* WidgetBlueprint, const FMVVMBlueprintViewBinding& Binding, bool bIsSource, bool bUseDisplayName, bool bAppendFunctionKeywords)
 	{
 		UMVVMWidgetBlueprintExtension_View* ExtensionView = UMVVMWidgetBlueprintExtension_View::GetExtension<UMVVMWidgetBlueprintExtension_View>(WidgetBlueprint);
 		UMVVMBlueprintView* BlueprintView = ExtensionView->GetBlueprintView();
@@ -226,14 +83,7 @@ namespace UE::MVVM::Private
 
 		auto AddPath = [&](const FMVVMBlueprintPropertyPath& ArgumentPath)
 		{
-			if (ArgumentPath.IsFromViewModel())
-			{
-				AppendViewModelPathString(WidgetBlueprint, BlueprintView, ArgumentPath, NameBuilder, FunctionKeywordsBuilder, bUseDisplayName, bAppendFunctionKeywords);
-			}
-			else
-			{
-				AppendWidgetPathString(WidgetBlueprint, ArgumentPath, NameBuilder, FunctionKeywordsBuilder, bUseDisplayName, bAppendFunctionKeywords);
-			}
+			NameBuilder << ArgumentPath.ToString(WidgetBlueprint, bUseDisplayName, bAppendFunctionKeywords);
 		};
 
 		if (UMVVMBlueprintViewConversionFunction* ViewConversionFunction = Binding.Conversion.GetConversionFunction(bIsSource))
@@ -257,8 +107,8 @@ namespace UE::MVVM::Private
 						FString FunctionKeywords = ConversionFunctionPtr->GetMetaData(FBlueprintMetadata::MD_FunctionKeywords);
 						if (!FunctionKeywords.IsEmpty())
 						{
-							FunctionKeywordsBuilder << TEXT(".");
-							FunctionKeywordsBuilder << FunctionKeywords;
+							NameBuilder << TEXT(".");
+							NameBuilder << FunctionKeywords;
 						}
 					}
 				}
@@ -343,9 +193,8 @@ FString FMVVMBlueprintViewBinding::GetDisplayNameString(const UWidgetBlueprint*
 	check(WidgetBlueprint);
 
 	TStringBuilder<256> NameBuilder;
-	TStringBuilder<2> FunctionKeywordsBuilder; // This is only passed to GetBindingPathName but never used in this function.
-
-	NameBuilder << UE::MVVM::Private::GetBindingPathName(WidgetBlueprint, *this, false, FunctionKeywordsBuilder, bUseDisplayName);
+	bool bAppendFunctionKeywords = false;
+	NameBuilder << UE::MVVM::Private::GetBindingPathName(WidgetBlueprint, *this, false, bUseDisplayName, bAppendFunctionKeywords);
 
 	if (BindingType == EMVVMBindingMode::TwoWay)
 	{
@@ -364,7 +213,7 @@ FString FMVVMBlueprintViewBinding::GetDisplayNameString(const UWidgetBlueprint*
 		NameBuilder << TEXT(" ??? "); // shouldn't happen
 	}
 
-	NameBuilder << UE::MVVM::Private::GetBindingPathName(WidgetBlueprint, *this, true, FunctionKeywordsBuilder, bUseDisplayName);
+	NameBuilder << UE::MVVM::Private::GetBindingPathName(WidgetBlueprint, *this, true, bUseDisplayName, bAppendFunctionKeywords);
 
 	return NameBuilder.ToString();
 }
@@ -373,31 +222,21 @@ FString FMVVMBlueprintViewBinding::GetSearchableString(const UWidgetBlueprint* W
 {
 	check(WidgetBlueprint);
 
-	FString SearchString;
-
-	// Get the binding string with variable names.
-	SearchString.Append(GetDisplayNameString(WidgetBlueprint, false));
+	// Create the function keywords string.
+	TStringBuilder<258> Builder;
+	bool bDisplayName = true;
+	bool bAppendFunctionKeywords = true;
+	Builder << UE::MVVM::Private::GetBindingPathName(WidgetBlueprint, *this, false, bDisplayName, bAppendFunctionKeywords);
+	Builder << UE::MVVM::Private::GetBindingPathName(WidgetBlueprint, *this, true, bDisplayName, bAppendFunctionKeywords);
 
+	FString Result = Builder.ToString();
 	// Remove the extra formatting that we don't need for search.
 	// We will include the formatted string as well in the second call to GetBindingPathName.
-	SearchString.ReplaceInline(TEXT(" "), TEXT(""));
-	SearchString.ReplaceInline(TEXT(")"), TEXT(""));
-	SearchString.ReplaceInline(TEXT("("), TEXT("."));
-	SearchString.ReplaceInline(TEXT(","), TEXT("."));
-
-	SearchString.Append(TEXT("."));
-
-	// Get the binding string with display names.
-	SearchString.Append(GetDisplayNameString(WidgetBlueprint, true));
-
-	// Create the function keywords string.
-	TStringBuilder<128> FunctionKeywordsBuilder;
-	UE::MVVM::Private::GetBindingPathName(WidgetBlueprint, *this, false, FunctionKeywordsBuilder, false, true);
-	UE::MVVM::Private::GetBindingPathName(WidgetBlueprint, *this, true, FunctionKeywordsBuilder, false, true);
-
-	SearchString.Append(FunctionKeywordsBuilder);
+	Result.ReplaceInline(TEXT(")"), TEXT(" "));
+	Result.ReplaceInline(TEXT("("), TEXT(" "));
+	Result.ReplaceInline(TEXT(","), TEXT(" "));
 
-	return SearchString;
+	return Result;
 }
 
 #undef LOCTEXT_NAMESPACE
\ No newline at end of file
diff --git a/Engine/Plugins/Runtime/ModelViewViewModel/Source/ModelViewViewModelBlueprint/Private/MVVMBlueprintViewConversionFunction.cpp b/Engine/Plugins/Runtime/ModelViewViewModel/Source/ModelViewViewModelBlueprint/Private/MVVMBlueprintViewConversionFunction.cpp
--- a/Engine/Plugins/Runtime/ModelViewViewModel/Source/ModelViewViewModelBlueprint/Private/MVVMBlueprintViewConversionFunction.cpp
+++ b/Engine/Plugins/Runtime/ModelViewViewModel/Source/ModelViewViewModelBlueprint/Private/MVVMBlueprintViewConversionFunction.cpp
@@ -242,7 +242,7 @@ UEdGraph* UMVVMBlueprintViewConversionFunction::GetOrCreateWrapperGraphInternal(
 
 UEdGraph* UMVVMBlueprintViewConversionFunction::GetOrCreateWrapperGraphInternal(UBlueprint* Blueprint, const UFunction* Function) const
 {
-	TPair<UEdGraph*, UK2Node*> Result = UE::MVVM::ConversionFunctionHelper::CreateGraph(Blueprint, GraphName, Function, bWrapperGraphTransient);
+	TPair<UEdGraph*, UK2Node*> Result = UE::MVVM::ConversionFunctionHelper::CreateGraph(Blueprint, GraphName, nullptr, Function, bWrapperGraphTransient);
 	CachedWrapperGraph = Result.Get<0>();
 	CachedWrapperNode = Result.Get<1>();
 	LoadPinValuesInternal(Blueprint);
diff --git a/Engine/Plugins/Runtime/ModelViewViewModel/Source/ModelViewViewModelBlueprint/Private/MVVMBlueprintViewEvent.cpp b/Engine/Plugins/Runtime/ModelViewViewModel/Source/ModelViewViewModelBlueprint/Private/MVVMBlueprintViewEvent.cpp
new file mode 100644
--- /dev/null
+++ b/Engine/Plugins/Runtime/ModelViewViewModel/Source/ModelViewViewModelBlueprint/Private/MVVMBlueprintViewEvent.cpp
@@ -0,0 +1,392 @@
+// Copyright Epic Games, Inc. All Rights Reserved.
+
+#include "MVVMBlueprintViewEvent.h"
+#include "MVVMBlueprintView.h"
+
+#include "Bindings/MVVMConversionFunctionHelper.h"
+#include "Bindings/MVVMBindingHelper.h"
+#include "Engine/Blueprint.h"
+#include "Kismet2/BlueprintEditorUtils.h"
+#include "MVVMWidgetBlueprintExtension_View.h"
+#include "WidgetBlueprint.h"
+
+#include "EdGraphSchema_K2.h"
+#include "EdGraph/EdGraphPin.h"
+#include "K2Node_CallFunction.h"
+#include "K2Node_FunctionEntry.h"
+#include "K2Node_FunctionResult.h"
+#include "K2Node_VariableSet.h"
+#include "KismetCompiler.h"
+
+#include UE_INLINE_GENERATED_CPP_BY_NAME(MVVMBlueprintViewEvent)
+
+#define LOCTEXT_NAMESPACE "MVVMBlueprintViewEvent"
+
+void UMVVMBlueprintViewEvent::SetEventPath(FMVVMBlueprintPropertyPath InEventPath)
+{
+	if (InEventPath == EventPath)
+	{
+		return;
+	}
+
+	RemoveWrapperGraph();
+
+	EventPath = MoveTemp(InEventPath);
+	GraphName = FName();
+
+	if (!EventPath.IsEmpty())
+	{
+		TStringBuilder<256> StringBuilder;
+		StringBuilder << TEXT("__");
+		StringBuilder << FGuid::NewGuid().ToString(EGuidFormats::DigitsWithHyphensLower);
+		GraphName = StringBuilder.ToString();
+	}
+
+	CreateWrapperGraphInternal();
+	SavePinValues();
+}
+
+void UMVVMBlueprintViewEvent::SetDestinationPath(FMVVMBlueprintPropertyPath InDestinationPath)
+{
+	if (InDestinationPath == DestinationPath)
+	{
+		return;
+	}
+
+	RemoveWrapperGraph();
+
+	DestinationPath = MoveTemp(InDestinationPath);
+
+	CreateWrapperGraphInternal();
+	SavePinValues();
+}
+	
+UEdGraph* UMVVMBlueprintViewEvent::GetOrCreateWrapperGraph()
+{
+	if (CachedWrapperGraph)
+	{
+		return CachedWrapperGraph;
+	}
+
+	CreateWrapperGraphInternal();
+	return CachedWrapperGraph;
+}
+
+void UMVVMBlueprintViewEvent::RemoveWrapperGraph()
+{
+	if (CachedWrapperGraph)
+	{
+		FBlueprintEditorUtils::RemoveGraph(GetWidgetBlueprintInternal(), CachedWrapperGraph);
+		CachedWrapperGraph = nullptr;
+		CachedWrapperNode = nullptr;
+	}
+
+	Messages.Empty();
+	SavedPins.Empty();
+}
+
+UEdGraphPin* UMVVMBlueprintViewEvent::GetOrCreateGraphPin(FName PinName)
+{
+	GetOrCreateWrapperGraph();
+	if (CachedWrapperNode)
+	{
+		return CachedWrapperNode->FindPin(PinName);
+	}
+	return nullptr;
+}
+
+void UMVVMBlueprintViewEvent::SavePinValues()
+{
+	if (CachedWrapperNode)
+	{
+		SavedPins.Reset();
+		UWidgetBlueprint* Blueprint = GetWidgetBlueprintInternal();
+		for (UEdGraphPin* Pin : CachedWrapperNode->Pins)
+		{
+			if (Pin->PinName != UEdGraphSchema_K2::PN_Self && Pin->PinName != UEdGraphSchema_K2::PN_Execute && Pin->Direction == EGPD_Input)
+			{
+				SavedPins.Add(FMVVMBlueprintPin::CreateFromPin(Blueprint, Pin));
+			}
+		}
+	}
+}
+
+FMVVMBlueprintPropertyPath UMVVMBlueprintViewEvent::GetPinPath(FName PinName) const
+{
+	const FMVVMBlueprintPin* ViewPin = SavedPins.FindByPredicate([PinName](const FMVVMBlueprintPin& Other) { return PinName == Other.GetName(); });
+	return ViewPin ? ViewPin->GetPath() : FMVVMBlueprintPropertyPath();
+}
+
+void UMVVMBlueprintViewEvent::SetPinPath(FName PinName, const FMVVMBlueprintPropertyPath& Path)
+{
+	UEdGraphPin* GraphPin = GetOrCreateGraphPin(PinName);
+
+	if (GraphPin)
+	{
+		UBlueprint* Blueprint = GetWidgetBlueprintInternal();
+		// Set the value and make the blueprint as dirty before creating the pin.
+		FMVVMBlueprintPin* ViewPin = SavedPins.FindByPredicate([PinName](const FMVVMBlueprintPin& Other) { return PinName == Other.GetName(); });
+		if (!ViewPin)
+		{
+			ViewPin = &SavedPins.Add_GetRef(FMVVMBlueprintPin::CreateFromPin(Blueprint, GraphPin));
+		}
+
+		//A property (viewmodel or widget) may not be created yet and the skeletal needs to be recreated.
+		FBlueprintEditorUtils::MarkBlueprintAsStructurallyModified(Blueprint);
+
+		UE::MVVM::ConversionFunctionHelper::SetPropertyPathForPin(Blueprint, Path, GraphPin);
+
+		// Take the path built in BP, it may had some errors
+		ViewPin->SetPath(UE::MVVM::ConversionFunctionHelper::GetPropertyPathForPin(Blueprint, GraphPin, false));
+	}
+}
+
+UWidgetBlueprint* UMVVMBlueprintViewEvent::GetWidgetBlueprintInternal() const
+{
+	return GetOuterUMVVMBlueprintView()->GetOuterUMVVMWidgetBlueprintExtension_View()->GetWidgetBlueprint();
+}
+
+bool UMVVMBlueprintViewEvent::Supports(const UWidgetBlueprint* WidgetBlueprint, const FMVVMBlueprintPropertyPath& PropertyPath)
+{
+	return GetEventSignature(WidgetBlueprint, PropertyPath) != nullptr;
+}
+
+const UFunction* UMVVMBlueprintViewEvent::GetEventSignature() const
+{
+	return GetEventSignature(GetWidgetBlueprintInternal(), EventPath);
+}
+
+const UFunction* UMVVMBlueprintViewEvent::GetEventSignature(const UWidgetBlueprint* WidgetBlueprint, const FMVVMBlueprintPropertyPath& PropertyPath)
+{
+	if (!PropertyPath.IsEmpty() && PropertyPath.GetFieldPaths().Num() > 0)
+	{
+		const FMVVMBlueprintFieldPath& LastPath = PropertyPath.GetFieldPaths().Last();
+		UE::MVVM::FMVVMConstFieldVariant LastField = LastPath.GetField(WidgetBlueprint->SkeletonGeneratedClass);
+		if (LastField.IsProperty())
+		{
+			if (const FMulticastDelegateProperty* Property = CastField<FMulticastDelegateProperty>(LastField.GetProperty()))
+			{
+				return Property->SignatureFunction.Get();
+			}
+		}
+	}
+	return nullptr;
+}
+
+UEdGraph* UMVVMBlueprintViewEvent::CreateWrapperGraphInternal()
+{
+	if (GraphName.IsNone() || DestinationPath.IsEmpty() || EventPath.IsEmpty())
+	{
+		return nullptr;
+	}
+
+	const UFunction* DelegateSignature = GetEventSignature();
+	if (DelegateSignature == nullptr)
+	{
+		return nullptr;
+	}
+
+	if (DestinationPath.GetFieldPaths().Num() > 0)
+	{
+		const FMVVMBlueprintFieldPath& LastPath = DestinationPath.GetFieldPaths().Last();
+		UE::MVVM::FMVVMConstFieldVariant LastField = LastPath.GetField(GetWidgetBlueprintInternal()->SkeletonGeneratedClass);
+		if (LastField.IsFunction() && LastField.GetFunction() != nullptr)
+		{
+			return CreateWrapperGraphInternal(DelegateSignature, LastField.GetFunction());
+		}
+		else if (LastField.IsProperty() && LastField.GetProperty() != nullptr)
+		{
+			return CreateWrapperGraphInternal(DelegateSignature, LastField.GetProperty());
+		}
+	}
+	else if (DestinationPath.GetViewModelId().IsValid())
+	{
+		const FMVVMBlueprintViewModelContext* Context = GetOuterUMVVMBlueprintView()->FindViewModel(DestinationPath.GetViewModelId());
+		const FProperty* ViewModel = Context ? GetWidgetBlueprintInternal()->SkeletonGeneratedClass->FindPropertyByName(Context->GetViewModelName()) : nullptr;
+		if (ViewModel)
+		{
+			return CreateWrapperGraphInternal(DelegateSignature, ViewModel);
+		}
+	}
+	return nullptr;
+}
+
+UEdGraph* UMVVMBlueprintViewEvent::CreateWrapperGraphInternal(const UFunction* DelegateSignature, const UFunction* Function)
+{
+	UWidgetBlueprint* WidgetBlueprint = GetWidgetBlueprintInternal();
+	bool bTransient = true;
+	TPair<UEdGraph*, UK2Node*> Result = UE::MVVM::ConversionFunctionHelper::CreateGraph(WidgetBlueprint, GraphName, DelegateSignature, Function, bTransient);
+	CachedWrapperGraph = Result.Get<0>();
+	CachedWrapperNode = Result.Get<1>();
+
+	// Add SelfNode if needed
+	if (CachedWrapperNode)
+	{
+		if (UEdGraphPin* SelfPin = CachedWrapperNode->FindPin(UEdGraphSchema_K2::PSC_Self))
+		{
+			FMVVMBlueprintPropertyPath PathToSet = DestinationPath;
+			TArray<UE::MVVM::FMVVMConstFieldVariant> FieldPaths = PathToSet.GetFields(WidgetBlueprint->SkeletonGeneratedClass);
+			PathToSet.ResetPropertyPath();
+			for (int32 Index = 0; Index < FieldPaths.Num() - 1; ++Index)
+			{
+				PathToSet.AppendPropertyPath(WidgetBlueprint, FieldPaths[Index]);
+			}
+			;
+			UE::MVVM::ConversionFunctionHelper::SetPropertyPathForPin(WidgetBlueprint, PathToSet, SelfPin);
+		}
+	}
+
+	LoadPinValuesInternal();
+	return CachedWrapperGraph;
+}
+
+UEdGraph* UMVVMBlueprintViewEvent::CreateWrapperGraphInternal(const UFunction* DelegateSignature, const FProperty* Property)
+{
+	UWidgetBlueprint* WidgetBlueprint = GetWidgetBlueprintInternal();
+	bool bTransient = true;
+	TPair<UEdGraph*, UK2Node*> Result = UE::MVVM::ConversionFunctionHelper::CreateGraph(WidgetBlueprint, GraphName, DelegateSignature, UK2Node_VariableSet::StaticClass(), bTransient,
+		[WidgetBlueprint, Property](UK2Node* NewNode)
+		{
+			UK2Node_VariableSet* VariableNode = CastChecked<UK2Node_VariableSet>(NewNode);
+			UEdGraphSchema_K2::ConfigureVarNode(VariableNode, Property->GetFName(), Property->GetOwnerStruct(), WidgetBlueprint);
+		});
+	CachedWrapperGraph = Result.Get<0>();
+	CachedWrapperNode = Result.Get<1>();
+
+	// Add SelfNode
+	if (CachedWrapperNode)
+	{
+		FMVVMBlueprintPropertyPath PathToSet = DestinationPath;
+		TArray<UE::MVVM::FMVVMConstFieldVariant> FieldPaths = PathToSet.GetFields(WidgetBlueprint->SkeletonGeneratedClass);
+		PathToSet.ResetPropertyPath();
+		for (int32 Index = 0; Index < FieldPaths.Num() - 1; ++Index)
+		{
+			PathToSet.AppendPropertyPath(WidgetBlueprint, FieldPaths[Index]);
+		}
+		UEdGraphPin* SelfPin = CachedWrapperNode->FindPinChecked(UEdGraphSchema_K2::PSC_Self);
+		UE::MVVM::ConversionFunctionHelper::SetPropertyPathForPin(WidgetBlueprint, PathToSet, SelfPin);
+	}
+
+	LoadPinValuesInternal();
+	return CachedWrapperGraph;
+}
+
+void UMVVMBlueprintViewEvent::LoadPinValuesInternal()
+{
+	if (CachedWrapperNode)
+	{
+		TArray<UEdGraphPin*> AllPins;
+		AllPins.Reserve(CachedWrapperNode->Pins.Num());
+		for (UEdGraphPin* Pin : CachedWrapperNode->Pins)
+		{
+			if (Pin->PinName != UEdGraphSchema_K2::PN_Self && Pin->PinName != UEdGraphSchema_K2::PN_Execute && Pin->Direction == EGPD_Input)
+			{
+				AllPins.Add(Pin);
+			}
+		}
+
+		UBlueprint* Blueprint = GetWidgetBlueprintInternal();
+		for (int32 Index = SavedPins.Num() - 1; Index >= 0; --Index)
+		{
+			const FMVVMBlueprintPin& Pin = SavedPins[Index];
+			if (UEdGraphPin* GraphPin = CachedWrapperNode->FindPin(Pin.GetName()))
+			{
+				AllPins.RemoveSingleSwap(GraphPin);
+				Pin.CopyTo(Blueprint, GraphPin);
+			}
+			else
+			{
+				// pin doesn't exist anymore
+				SavedPins.RemoveAt(Index);
+			}
+		}
+
+		// Create the reminding pin.
+		for (UEdGraphPin* Pin : AllPins)
+		{
+			SavedPins.Add(FMVVMBlueprintPin::CreateFromPin(Blueprint, Pin));
+		}
+	}
+}
+
+TArray<FText> UMVVMBlueprintViewEvent::GetCompilationMessages(EMessageType InMessageType) const
+{
+	TArray<FText> Result;
+	Result.Reset(Messages.Num());
+	for (const FMessage& Msg : Messages)
+	{
+		if (Msg.MessageType == InMessageType)
+		{
+			Result.Add(Msg.MessageText);
+		}
+	}
+	return Result;
+}
+
+bool UMVVMBlueprintViewEvent::HasCompilationMessage(EMessageType InMessageType) const
+{
+	return Messages.ContainsByPredicate([InMessageType](const FMessage& Other)
+	{
+		return Other.MessageType == InMessageType;
+	});
+}
+
+void UMVVMBlueprintViewEvent::AddCompilationToBinding(FMessage MessageToAdd)
+{
+	Messages.Add(MoveTemp(MessageToAdd));
+}
+
+void UMVVMBlueprintViewEvent::ResetCompilationMessages()
+{
+	Messages.Reset();
+}
+
+FText UMVVMBlueprintViewEvent::GetDisplayName(bool bUseDisplayName) const
+{
+	TArray<FText> JoinArgs;
+	for (const FMVVMBlueprintPin& Pin : GetPins())
+	{
+		if (Pin.UsedPathAsValue())
+		{
+			JoinArgs.Add(Pin.GetPath().ToText(GetWidgetBlueprintInternal(), bUseDisplayName));
+		}
+	}
+
+	return FText::Format(LOCTEXT("BlueprintViewEventDisplayNameFormat", "{0} => {1}({2})")
+		, EventPath.ToText(GetWidgetBlueprintInternal(), bUseDisplayName)
+		, DestinationPath.ToText(GetWidgetBlueprintInternal(), bUseDisplayName)
+		, FText::Join(LOCTEXT("PathDelimiter", ", "), JoinArgs)
+		);
+}
+
+FString UMVVMBlueprintViewEvent::GetSearchableString() const
+{
+	TStringBuilder<256> Builder;
+	Builder << EventPath.ToString(GetWidgetBlueprintInternal(), true, true);
+	Builder << TEXT(' ');
+	Builder << DestinationPath.ToString(GetWidgetBlueprintInternal(), true, true);
+	Builder << TEXT('(');
+	bool bFirst = true;
+	for (const FMVVMBlueprintPin& Pin : GetPins())
+	{
+		if (!bFirst)
+		{
+			Builder << TEXT(", ");
+		}
+		if (Pin.UsedPathAsValue())
+		{
+			Builder << Pin.GetPath().ToString(GetWidgetBlueprintInternal(), true, true);
+		}
+		bFirst = false;
+	}
+	Builder << TEXT(')');
+	return Builder.ToString();
+}
+
+void UMVVMBlueprintViewEvent::PostEditChangeChainProperty(FPropertyChangedChainEvent& PropertyChainEvent)
+{
+	Super::PostEditChangeChainProperty(PropertyChainEvent);
+	GetOuterUMVVMBlueprintView()->OnViewModelsUpdated.Broadcast();
+}
+
+#undef LOCTEXT_NAMESPACE
diff --git a/Engine/Plugins/Runtime/ModelViewViewModel/Source/ModelViewViewModelBlueprint/Private/MVVMFunctionGraphHelper.cpp b/Engine/Plugins/Runtime/ModelViewViewModel/Source/ModelViewViewModelBlueprint/Private/MVVMFunctionGraphHelper.cpp
--- a/Engine/Plugins/Runtime/ModelViewViewModel/Source/ModelViewViewModelBlueprint/Private/MVVMFunctionGraphHelper.cpp
+++ b/Engine/Plugins/Runtime/ModelViewViewModel/Source/ModelViewViewModelBlueprint/Private/MVVMFunctionGraphHelper.cpp
@@ -21,6 +21,24 @@
 namespace UE::MVVM::FunctionGraphHelper
 {
 
+namespace Private
+{
+
+UK2Node_FunctionEntry* FindFunctionEntry(const UEdGraph* Graph)
+{
+	for (UEdGraphNode* Node : Graph->Nodes)
+	{
+		UK2Node_FunctionEntry* FunctionEntry = Cast<UK2Node_FunctionEntry>(Node);
+		if (FunctionEntry)
+		{
+			return FunctionEntry;
+		}
+	}
+	return nullptr;
+}
+
+} //namespace
+
 UEdGraph* CreateFunctionGraph(UBlueprint* InBlueprint, FStringView InFunctionName, EFunctionFlags ExtraFunctionFlag, const FStringView Category, bool bIsEditable)
 {
 	FName UniqueFunctionName = FBlueprintEditorUtils::FindUniqueKismetName(InBlueprint, InFunctionName.GetData());
@@ -53,7 +71,7 @@ UEdGraph* CreateIntermediateFunctionGraph(FKismetCompilerContext& InContext, FSt
 	Schema->CreateDefaultNodesForGraph(*FunctionGraph);
 
 	// function entry node
-	UK2Node_FunctionEntry* FunctionEntry = CastChecked<UK2Node_FunctionEntry>(FunctionGraph->Nodes[0]);
+	UK2Node_FunctionEntry* FunctionEntry = Private::FindFunctionEntry(FunctionGraph);
 	check(FunctionEntry);
 	{
 		FunctionEntry->AddExtraFlags(ExtraFunctionFlag);
@@ -127,16 +145,7 @@ bool GenerateViewModelSetter(FKismetCompilerContext& InContext, UEdGraph* InFunc
 {
 	check(InFunctionGraph);
 
-	UK2Node_FunctionEntry* FunctionEntry = nullptr;
-	for (UEdGraphNode* Node : InFunctionGraph->Nodes)
-	{
-		FunctionEntry = Cast<UK2Node_FunctionEntry>(Node);
-		if (FunctionEntry)
-		{
-			break;
-		}
-	}
-
+	UK2Node_FunctionEntry* FunctionEntry = Private::FindFunctionEntry(InFunctionGraph);
 	if (FunctionEntry == nullptr)
 	{
 		return false;
@@ -210,16 +219,7 @@ bool GenerateViewModelFieldNotifySetter(FKismetCompilerContext& InContext, UEdGr
 {
 	check(InFunctionGraph);
 
-	UK2Node_FunctionEntry* FunctionEntry = nullptr;
-	for (UEdGraphNode* Node : InFunctionGraph->Nodes)
-	{
-		FunctionEntry = Cast<UK2Node_FunctionEntry>(Node);
-		if (FunctionEntry)
-		{
-			break;
-		}
-	}
-
+	UK2Node_FunctionEntry* FunctionEntry = Private::FindFunctionEntry(InFunctionGraph);
 	if (FunctionEntry == nullptr)
 	{
 		return false;
@@ -284,16 +284,7 @@ bool GenerateViewModelFielNotifyBroadcast(FKismetCompilerContext& InContext, UEd
 {
 	check(InFunctionGraph);
 
-	UK2Node_FunctionEntry* FunctionEntry = nullptr;
-	for (UEdGraphNode* Node : InFunctionGraph->Nodes)
-	{
-		FunctionEntry = Cast<UK2Node_FunctionEntry>(Node);
-		if (FunctionEntry)
-		{
-			break;
-		}
-	}
-
+	UK2Node_FunctionEntry* FunctionEntry = Private::FindFunctionEntry(InFunctionGraph);
 	if (FunctionEntry == nullptr)
 	{
 		return false;
@@ -395,16 +386,7 @@ bool GenerateSetter(UEdGraph* InFunctionGraph, TArrayView<UE::MVVM::FMVVMConstFi
 {
 	check(InFunctionGraph);
 
-	UK2Node_FunctionEntry* FunctionEntry = nullptr;
-	for (UEdGraphNode* Node : InFunctionGraph->Nodes)
-	{
-		FunctionEntry = Cast<UK2Node_FunctionEntry>(Node);
-		if (FunctionEntry)
-		{
-			break;
-		}
-	}
-
+	UK2Node_FunctionEntry* FunctionEntry = Private::FindFunctionEntry(InFunctionGraph);
 	if (FunctionEntry == nullptr)
 	{
 		return false;
@@ -541,4 +523,72 @@ bool GenerateSetter(UBlueprint* InBlueprint, UEdGraph* InFunctionGraph, TArrayVi
 	return Private::GenerateSetter(InFunctionGraph, InSetterPath, InBlueprint, nullptr);
 }
 
+bool IsFunctionEntryMatchSignature(const UEdGraph* FunctionGraph, const UFunction* FunctionSignature)
+{
+	if (FunctionGraph == nullptr || FunctionSignature == nullptr)
+	{
+		return false;
+	}
+
+	const UK2Node_FunctionEntry* FunctionEntry = Private::FindFunctionEntry(FunctionGraph);
+	if (FunctionEntry == nullptr)
+	{
+		return false;
+	}
+
+	// Generate pins list
+	TArray<UEdGraphPin*> EntryPins = FunctionEntry->Pins;
+	for (int32 Index = EntryPins.Num()-1; Index >= 0; --Index)
+	{
+		const UEdGraphPin* CurPin = EntryPins[Index];
+		if (CurPin->PinType.PinCategory == UEdGraphSchema_K2::PC_Exec
+			|| CurPin->PinType.PinSubCategory == UEdGraphSchema_K2::PSC_Self
+			|| CurPin->Direction == EGPD_Input
+			|| CurPin->ParentPin != nullptr)
+		{
+			EntryPins.RemoveAt(Index, 1);
+		}
+	}
+
+	const UEdGraphSchema_K2* K2Schema = GetDefault<UEdGraphSchema_K2>();
+
+	// Generate signature pins
+	TArray<FUserPinInfo> SignaturePins;
+	{
+		for (TFieldIterator<FProperty> It(FunctionSignature); It; ++It)
+		{
+			FUserPinInfo& PinInfo = SignaturePins.AddDefaulted_GetRef();
+			K2Schema->ConvertPropertyToPinType(*It, PinInfo.PinType);
+			PinInfo.PinName = It->GetFName();
+			PinInfo.DesiredPinDirection = EGPD_Output;
+		}
+	}
+
+	// Do we have the same number of arguments
+	if (SignaturePins.Num() != EntryPins.Num())
+	{
+		return false;
+	}
+
+	// Now check through the event's pins, and check for compatible pins, removing them if we find a match.
+	for (int32 Index = 0; Index < SignaturePins.Num(); ++Index)
+	{
+		const UEdGraphPin* CurrentEventPin = EntryPins[Index];
+		const FUserPinInfo& CurrentPinInfo = SignaturePins[Index];
+
+		bool bMatchFound = false;
+		if (CurrentEventPin->PinName != CurrentPinInfo.PinName)
+		{
+			return false;
+		}
+
+			// Check to make sure pins are of the same type
+		if (!K2Schema->ArePinTypesCompatible(CurrentEventPin->PinType, CurrentPinInfo.PinType))
+		{
+			return false;
+		}
+	}
+	return true;
+}
+
 } //namespace
diff --git a/Engine/Plugins/Runtime/ModelViewViewModel/Source/ModelViewViewModelBlueprint/Private/MVVMFunctionGraphHelper.h b/Engine/Plugins/Runtime/ModelViewViewModel/Source/ModelViewViewModelBlueprint/Private/MVVMFunctionGraphHelper.h
--- a/Engine/Plugins/Runtime/ModelViewViewModel/Source/ModelViewViewModelBlueprint/Private/MVVMFunctionGraphHelper.h
+++ b/Engine/Plugins/Runtime/ModelViewViewModel/Source/ModelViewViewModelBlueprint/Private/MVVMFunctionGraphHelper.h
@@ -49,4 +49,7 @@ namespace UE::MVVM::FunctionGraphHelper
 	/** Generate the nodes needed to set a variable/function value. To use while while compiling the Blueprint. */
 	bool GenerateIntermediateSetter(FKismetCompilerContext& Context, UEdGraph* FunctionGraph, TArrayView<UE::MVVM::FMVVMConstFieldVariant> SetterPath);
 
+	/** @return true when the entry node in the graph matches the function signature. */
+	bool IsFunctionEntryMatchSignature(const UEdGraph* FunctionGraph, const UFunction* FunctionSignature);
+
 } //namespace
diff --git a/Engine/Plugins/Runtime/ModelViewViewModel/Source/ModelViewViewModelBlueprint/Private/MVVMPropertyPath.cpp b/Engine/Plugins/Runtime/ModelViewViewModel/Source/ModelViewViewModelBlueprint/Private/MVVMPropertyPath.cpp
--- a/Engine/Plugins/Runtime/ModelViewViewModel/Source/ModelViewViewModelBlueprint/Private/MVVMPropertyPath.cpp
+++ b/Engine/Plugins/Runtime/ModelViewViewModel/Source/ModelViewViewModelBlueprint/Private/MVVMPropertyPath.cpp
@@ -2,13 +2,21 @@
 
 #include "MVVMPropertyPath.h"
 #include "BlueprintCompilationManager.h"
+#include "Blueprint/WidgetTree.h"
 #include "Engine/Blueprint.h"
 #include "Engine/UserDefinedStruct.h"
 #include "Kismet2/BlueprintEditorUtils.h"
 #include "Kismet2/StructureEditorUtils.h"
+#include "MVVMBlueprintView.h"
+#include "MVVMBlueprintViewModelContext.h"
+#include "MVVMWidgetBlueprintExtension_View.h"
+#include "Types/MVVMFieldVariant.h"
+#include "WidgetBlueprint.h"
 
 #include UE_INLINE_GENERATED_CPP_BY_NAME(MVVMPropertyPath)
 
+#define LOCTEXT_NAMESPACE "MVVMBlueprintFieldPath"
+
 FMVVMBlueprintFieldPath::FMVVMBlueprintFieldPath(const UBlueprint* InContext, UE::MVVM::FMVVMConstFieldVariant InField)
 {
 	if (!InField.IsValid())
@@ -277,4 +285,147 @@ bool FMVVMBlueprintPropertyPath::HasFieldInLocalScope() const
 	}
 
 	return false;
-}
\ No newline at end of file
+}
+
+
+namespace UE::MVVM::Private
+{
+FText GetWidgetDisplayName(const UWidgetBlueprint* WidgetBlueprint, FName WidgetName, bool bUseDisplayName, bool bIncludeMetaData)
+{
+	if (!bUseDisplayName && !bIncludeMetaData)
+	{
+		return FText::FromName(WidgetName);
+	}
+
+	UWidgetTree* WidgetTree = WidgetBlueprint ? WidgetBlueprint->WidgetTree : nullptr;
+	UWidget* FoundWidget = WidgetTree ? WidgetTree->FindWidget(WidgetName) : nullptr;
+	if (bIncludeMetaData)
+	{
+		return FoundWidget ? FoundWidget->GetLabelTextWithMetadata() : FText::FromName(WidgetName);
+	}
+	return FoundWidget ? FoundWidget->GetLabelText() : FText::FromName(WidgetName);
+}
+
+
+FText GetViewModelDisplayName(const UWidgetBlueprint* WidgetBlueprint, FGuid Id, bool bUseDisplayName)
+{
+	UMVVMWidgetBlueprintExtension_View* ExtensionView = UMVVMWidgetBlueprintExtension_View::GetExtension<UMVVMWidgetBlueprintExtension_View>(WidgetBlueprint);
+	UMVVMBlueprintView* BlueprintView = ExtensionView ? ExtensionView->GetBlueprintView() : nullptr;
+	const FMVVMBlueprintViewModelContext* ViewModel = BlueprintView ? BlueprintView->FindViewModel(Id) : nullptr;
+	if (bUseDisplayName)
+	{
+		return ViewModel ? ViewModel->GetDisplayName() : LOCTEXT("None", "<None>");
+	}
+	return ViewModel ? FText::FromName(ViewModel->GetViewModelName()) : LOCTEXT("None", "<None>");
+}
+
+
+FText GetRootName(const UWidgetBlueprint* WidgetBlueprint, const FMVVMBlueprintPropertyPath& PropertyPath, bool bUseDisplayName, bool bIncludeMetaData)
+{
+	if (PropertyPath.IsFromWidget())
+	{
+		if (WidgetBlueprint->GetFName() == PropertyPath.GetWidgetName())
+		{
+			return bUseDisplayName ? LOCTEXT("Self", "Self") : FText::FromString(WidgetBlueprint->GetFriendlyName());
+		}
+		else
+		{
+			return GetWidgetDisplayName(WidgetBlueprint, PropertyPath.GetWidgetName(), bUseDisplayName, bIncludeMetaData);
+		}
+	}
+	else if (PropertyPath.IsFromViewModel())
+	{
+		return GetViewModelDisplayName(WidgetBlueprint, PropertyPath.GetViewModelId(), bUseDisplayName);
+	}
+	return FText::GetEmpty();
+}
+}//namespace
+
+
+FText FMVVMBlueprintPropertyPath::ToText(const UWidgetBlueprint* WidgetBlueprint, bool bUseDisplayName) const
+{
+	if (WidgetBlueprint == nullptr)
+	{
+		return FText::GetEmpty();
+	}
+
+	auto GetDisplayNameForField = [](const UE::MVVM::FMVVMConstFieldVariant& Field) -> FText
+	{
+		if (!Field.IsEmpty())
+		{
+			if (Field.IsProperty())
+			{
+				return Field.GetProperty()->GetDisplayNameText();
+			}
+			else if (Field.IsFunction())
+			{
+				return Field.GetFunction()->GetDisplayNameText();
+			}
+		}
+		return LOCTEXT("None", "<None>");
+	};
+
+	TArray<FText> JoinArgs;
+	JoinArgs.Add(UE::MVVM::Private::GetRootName(WidgetBlueprint, *this, bUseDisplayName, false));
+	TArray<UE::MVVM::FMVVMConstFieldVariant> Fields = GetFields(WidgetBlueprint->SkeletonGeneratedClass);
+	for (const UE::MVVM::FMVVMConstFieldVariant& Field : Fields)
+	{
+		JoinArgs.Add(bUseDisplayName ? GetDisplayNameForField(Field) : FText::FromName(Field.GetName()));
+	}
+	return FText::Join(LOCTEXT("PathDelimiter", "."), JoinArgs);
+}
+
+
+FString FMVVMBlueprintPropertyPath::ToString(const UWidgetBlueprint* WidgetBlueprint, bool bUseDisplayName, bool bIncludeMetaData) const
+{
+	if (WidgetBlueprint == nullptr)
+	{
+		return FString();
+	}
+
+	TStringBuilder<256> Builder;
+	Builder << UE::MVVM::Private::GetRootName(WidgetBlueprint, *this, bUseDisplayName, bIncludeMetaData).ToString();
+
+	TArray<UE::MVVM::FMVVMConstFieldVariant> Fields = GetFields(WidgetBlueprint->SkeletonGeneratedClass);
+	for (const UE::MVVM::FMVVMConstFieldVariant& Field : Fields)
+	{
+		Builder << '.';
+		if (!Field.IsEmpty())
+		{
+			if (Field.IsProperty())
+			{
+				if (bUseDisplayName)
+				{
+					Builder << Field.GetProperty()->GetDisplayNameText().ToString();
+				}
+				else
+				{
+					Builder << Field.GetProperty()->GetFName();
+				}
+			}
+			else if (Field.IsFunction())
+			{
+				if (bUseDisplayName)
+				{
+					Builder << Field.GetFunction()->GetDisplayNameText().ToString();
+				}
+				else
+				{
+					Builder << Field.GetFunction()->GetFName();
+				}
+				if (bIncludeMetaData)
+				{
+					const FString& FunctionKeywords = Field.GetFunction()->GetMetaData(FBlueprintMetadata::MD_FunctionKeywords);
+					if (!FunctionKeywords.IsEmpty())
+					{
+						Builder << TEXT(".");
+						Builder << FunctionKeywords;
+					}
+				}
+			}
+		}
+	}
+	return Builder.ToString();
+}
+
+#undef LOCTEXT_NAMESPACE
diff --git a/Engine/Plugins/Runtime/ModelViewViewModel/Source/ModelViewViewModelBlueprint/Private/MVVMViewBlueprintCompiler.cpp b/Engine/Plugins/Runtime/ModelViewViewModel/Source/ModelViewViewModelBlueprint/Private/MVVMViewBlueprintCompiler.cpp
--- a/Engine/Plugins/Runtime/ModelViewViewModel/Source/ModelViewViewModelBlueprint/Private/MVVMViewBlueprintCompiler.cpp
+++ b/Engine/Plugins/Runtime/ModelViewViewModel/Source/ModelViewViewModelBlueprint/Private/MVVMViewBlueprintCompiler.cpp
@@ -46,6 +46,7 @@ FAutoConsoleCommand CVarTestGenerateSetter(
 #endif
 
 static const FText CouldNotCreateSourceFieldPathFormat = LOCTEXT("CouldNotCreateSourceFieldPath", "Couldn't create the source field path '{0}'. {1}");
+static const FText PropertyPathIsInvalidFormat = LOCTEXT("PropertyPathIsInvalid", "The property path '{0}' is invalid.");
 
 FString PropertyPathToString(const UClass* InSelfContext, const UMVVMBlueprintView* BlueprintView, const FMVVMBlueprintPropertyPath& PropertyPath)
 {
@@ -146,12 +147,46 @@ void FMVVMViewBlueprintCompiler::AddMessageForBinding(FMVVMBlueprintViewBinding&
 	BlueprintView->AddMessageToBinding(Binding.BindingId, NewMessage);
 }
 
+void FMVVMViewBlueprintCompiler::AddMessageForEvent(UMVVMBlueprintViewEvent* Event, const FText& MessageText, UMVVMBlueprintViewEvent::EMessageType MessageType, FName ArgumentName) const
+{
+	const FText EventName = Event->GetDisplayName(true);
+
+	FText FormattedError;
+	if (!ArgumentName.IsNone())
+	{
+		FormattedError = FText::Format(LOCTEXT("EventFormatWithArgument", "Event '{0}': Argument '{1}' - {2}"), EventName, FText::FromName(ArgumentName), MessageText);
+	}
+	else
+	{
+		FormattedError = FText::Format(LOCTEXT("EventFormat", "Event '{0}': {1}"), EventName, MessageText);
+	}
+
+	switch (MessageType)
+	{
+	case UMVVMBlueprintViewEvent::EMessageType::Info:
+		WidgetBlueprintCompilerContext.MessageLog.Note(*FormattedError.ToString());
+		break;
+	case UMVVMBlueprintViewEvent::EMessageType::Warning:
+		WidgetBlueprintCompilerContext.MessageLog.Warning(*FormattedError.ToString());
+		break;
+	case UMVVMBlueprintViewEvent::EMessageType::Error:
+		WidgetBlueprintCompilerContext.MessageLog.Error(*FormattedError.ToString());
+		break;
+	default:
+		break;
+	}
+	UMVVMBlueprintViewEvent::FMessage NewMessage = { FormattedError, MessageType };
+	Event->AddCompilationToBinding(NewMessage);
+}
+
+
 void FMVVMViewBlueprintCompiler::AddErrorForViewModel(const FMVVMBlueprintViewModelContext& ViewModel, const FText& Message) const
 {
 	const FText FormattedError = FText::Format(LOCTEXT("ViewModelFormat", "Viewodel '{0}': {1}"), ViewModel.GetDisplayName(), Message);
 	WidgetBlueprintCompilerContext.MessageLog.Error(*FormattedError.ToString());
 }
 
+
 void FMVVMViewBlueprintCompiler::AddExtension(UWidgetBlueprintGeneratedClass* Class, UMVVMViewClass* ViewExtension)
 {
 	WidgetBlueprintCompilerContext.AddExtension(Class, ViewExtension);
@@ -191,14 +226,9 @@ void FMVVMViewBlueprintCompiler::CleanOldData(UWidgetBlueprintGeneratedClass* Cl
 }
 
 
-void FMVVMViewBlueprintCompiler::CleanTemporaries(UWidgetBlueprintGeneratedClass* ClassToClean)
-{
-}
-
-
 void FMVVMViewBlueprintCompiler::CreateFunctions(UMVVMBlueprintView* BlueprintView)
 {
-	if (!bAreSourcesCreatorValid || !bIsBindingsValid)
+	if (!bAreSourcesCreatorValid || !bAreBindingsValid || !bAreEventsValid)
 	{
 		return;
 	}
@@ -206,7 +236,7 @@ void FMVVMViewBlueprintCompiler::CreateFunctions(UMVVMBlueprintView* BlueprintVi
 	// Build the list of BP destinations
 	CreateBindingDestinationContexts(BlueprintView);
 
-	if (!bAreSourcesCreatorValid || !bIsBindingsValid)
+	if (!bAreSourcesCreatorValid || !bAreBindingsValid || !bAreEventsValid)
 	{
 		return;
 	}
@@ -236,7 +266,7 @@ void FMVVMViewBlueprintCompiler::CreateVariables(const FWidgetBlueprintCompilerC
 		return;
 	}
 
-	if (!bAreSourcesCreatorValid || !bAreSourceContextsValid || !bIsBindingsValid)
+	if (!bAreSourcesCreatorValid || !bAreSourceContextsValid || !bAreBindingsValid || !bAreEventsValid)
 	{
 		return;
 	}
@@ -452,6 +482,58 @@ void FMVVMViewBlueprintCompiler::CreateSourceLists(const FWidgetBlueprintCompile
 	UWidgetBlueprintGeneratedClass* SkeletonClass = Context.GetSkeletonGeneratedClass();
 	const FName DefaultWidgetCategory = Context.GetWidgetBlueprint()->GetFName();
 
+	auto GenerateCompilerSourceContext = [Self = this, BlueprintView, DefaultWidgetCategory, Class = SkeletonClass, &ViewModelGuids, &WidgetSources](const FMVVMBlueprintPropertyPath& PropertyPath) -> TValueOrError<void, FText>
+	{
+		if (PropertyPath.IsFromWidget())
+		{
+			if (PropertyPath.GetWidgetName() == Class->ClassGeneratedBy->GetFName())
+			{
+				// it's the userwidget
+				return MakeValue();
+			}
+
+			// If the widget doesn't have a property, add one automatically.
+			if (!WidgetSources.Contains(PropertyPath.GetWidgetName()))
+			{
+				WidgetSources.Add(PropertyPath.GetWidgetName());
+
+				UWidget** WidgetPtr = Self->WidgetNameToWidgetPointerMap.Find(PropertyPath.GetWidgetName());
+				if (WidgetPtr == nullptr || *WidgetPtr == nullptr)
+				{
+					return MakeError(FText::Format(LOCTEXT("InvalidWidgetFormat", "Could not find the targeted widget: {0}"), FText::FromName(PropertyPath.GetWidgetName())));
+				}
+
+				UWidget* Widget = *WidgetPtr;
+				FCompilerUserWidgetPropertyContext SourceVariable;
+				SourceVariable.Class = Widget->GetClass();
+				SourceVariable.PropertyName = PropertyPath.GetWidgetName();
+				SourceVariable.DisplayName = FText::FromString(Widget->GetDisplayLabel());
+				SourceVariable.CategoryName = TEXT("Widget");
+				SourceVariable.ViewModelId = FGuid();
+				SourceVariable.bPublicGetter = false;
+				Self->CompilerUserWidgetPropertyContexts.Emplace(MoveTemp(SourceVariable));
+			}
+		}
+		else if (PropertyPath.IsFromViewModel())
+		{
+			const FMVVMBlueprintViewModelContext* SourceViewModelContext = BlueprintView->FindViewModel(PropertyPath.GetViewModelId());
+			if (SourceViewModelContext == nullptr)
+			{
+				return MakeError(FText::Format(LOCTEXT("BindingViewModelNotFound", "Could not find viewmodel with GUID {0}."), GetViewModelIdText(PropertyPath)));
+			}
+
+			if (!ViewModelGuids.Contains(SourceViewModelContext->GetViewModelId()))
+			{
+				return MakeError(FText::Format(LOCTEXT("BindingViewModelInvalid", "Viewmodel {0} {1} was invalid."), SourceViewModelContext->GetDisplayName(), GetViewModelIdText(PropertyPath)));
+			}
+		}
+		else
+		{
+			return MakeError(LOCTEXT("SourcePathNotSet", "A source path is required, but not set."));
+		}
+		return MakeValue();
+	};
+
 	// Only find the source first property and destination first property.
 	//The full path will be tested later. We want to build the list of property needed.
 	for (int32 Index = 0; Index < BlueprintView->GetNumBindings(); ++Index)
@@ -470,86 +552,20 @@ void FMVVMViewBlueprintCompiler::CreateSourceLists(const FWidgetBlueprintCompile
 			continue;
 		}
 
-		FMVVMViewBlueprintCompiler* Self = this;
-		auto GenerateCompilerSourceContext = [Self, BlueprintView, DefaultWidgetCategory, Class = SkeletonClass, &Binding, &ViewModelGuids, &WidgetSources](const FMVVMBlueprintPropertyPath& PropertyPath,  FName ArgumentName = FName()) -> bool
-		{
-			if (PropertyPath.IsFromWidget())
-			{
-				if (PropertyPath.GetWidgetName() == Class->ClassGeneratedBy->GetFName())
-				{
-					// it's the userwidget
-					return true;
-				}
-
-				// If the widget doesn't have a property, add one automatically.
-				if (!WidgetSources.Contains(PropertyPath.GetWidgetName()))
-				{
-					WidgetSources.Add(PropertyPath.GetWidgetName());
+		const bool bIsForwardBinding = IsForwardBinding(Binding.BindingType);
+		const bool bIsBackwardBinding = IsBackwardBinding(Binding.BindingType);
 
-					UWidget** WidgetPtr = Self->WidgetNameToWidgetPointerMap.Find(PropertyPath.GetWidgetName());
-					if (WidgetPtr == nullptr || *WidgetPtr == nullptr)
-					{
-						Self->AddMessageForBinding(Binding,
-							BlueprintView,
-							FText::Format(LOCTEXT("InvalidWidgetFormat", "Could not find the targeted widget: {0}"), 
-								FText::FromName(PropertyPath.GetWidgetName())
-							),
-							EBindingMessageType::Error,
-							ArgumentName
-						);
-						return false;
-					}
 
-					UWidget* Widget = *WidgetPtr;
-					FCompilerUserWidgetPropertyContext SourceVariable;
-					SourceVariable.Class = Widget->GetClass();
-					SourceVariable.PropertyName = PropertyPath.GetWidgetName();
-					SourceVariable.DisplayName = FText::FromString(Widget->GetDisplayLabel());
-					SourceVariable.CategoryName = TEXT("Widget");
-					SourceVariable.ViewModelId = FGuid();
-					SourceVariable.bPublicGetter = false;
-					Self->CompilerUserWidgetPropertyContexts.Emplace(MoveTemp(SourceVariable));
-				}
-			}
-			else if (PropertyPath.IsFromViewModel())
-			{
-				const FMVVMBlueprintViewModelContext* SourceViewModelContext = BlueprintView->FindViewModel(PropertyPath.GetViewModelId());
-				if (SourceViewModelContext == nullptr)
-				{
-					Self->AddMessageForBinding(Binding,
-						BlueprintView,
-						FText::Format(LOCTEXT("BindingViewModelNotFound", "Could not find viewmodel with GUID {0}."), GetViewModelIdText(PropertyPath)),
-						EBindingMessageType::Error,
-						ArgumentName
-					);
-					return false;
-				}
-
-				if (!ViewModelGuids.Contains(SourceViewModelContext->GetViewModelId()))
-				{
-					Self->AddMessageForBinding(Binding,
-						BlueprintView,
-						FText::Format(LOCTEXT("BindingViewModelInvalid", "Viewmodel {0} {1} was invalid."), 
-							SourceViewModelContext->GetDisplayName(),
-							GetViewModelIdText(PropertyPath)
-						),
-						EBindingMessageType::Error,
-						ArgumentName
-					);
-					return false;
-				}
-			}
-			else
+		auto RunGenerateCompilerSourceContext = [Self = this, &GenerateCompilerSourceContext, &Binding, BlueprintView](const FMVVMBlueprintPropertyPath& PropertyPath, FName ArgumentName)
+		{
+			TValueOrError<void, FText> SourceContextResult = GenerateCompilerSourceContext(PropertyPath);
+			if (SourceContextResult.HasError())
 			{
-				Self->AddMessageForBinding(Binding, BlueprintView, LOCTEXT("SourcePathNotSet", "A source path is required, but not set."), EBindingMessageType::Error, ArgumentName);
-				return false;
+				Self->AddMessageForBinding(Binding, BlueprintView, SourceContextResult.StealError(), EBindingMessageType::Error, ArgumentName);
+				Self->bAreSourceContextsValid = false;
 			}
-			return true;
 		};
 
-		const bool bIsForwardBinding = IsForwardBinding(Binding.BindingType);
-		const bool bIsBackwardBinding = IsBackwardBinding(Binding.BindingType);
-
 		if (bIsForwardBinding || bIsBackwardBinding)
 		{
 			UMVVMBlueprintViewConversionFunction* ConversionFunction = Binding.Conversion.GetConversionFunction(bIsForwardBinding);
@@ -560,7 +576,7 @@ void FMVVMViewBlueprintCompiler::CreateSourceLists(const FWidgetBlueprintCompile
 				{
 					if (Binding.BindingType == EMVVMBindingMode::TwoWay)
 					{
-						Self->AddMessageForBinding(Binding, BlueprintView, LOCTEXT("TwoWayBindingsWithConversion", "Two-way bindings are not allowed to use conversion functions."), EBindingMessageType::Error);
+						AddMessageForBinding(Binding, BlueprintView, LOCTEXT("TwoWayBindingsWithConversion", "Two-way bindings are not allowed to use conversion functions."), EBindingMessageType::Error);
 						bAreSourceContextsValid = false;
 						continue;
 					}
@@ -570,26 +586,71 @@ void FMVVMViewBlueprintCompiler::CreateSourceLists(const FWidgetBlueprintCompile
 					{
 						if (Pin.UsedPathAsValue())
 						{
-							bAreSourceContextsValid &= GenerateCompilerSourceContext(Pin.GetPath(), Pin.GetName());
+							RunGenerateCompilerSourceContext(Pin.GetPath(), Pin.GetName());
 						}
 					}
 
 					// generate destination source
 					if (bIsForwardBinding)
 					{
-						bAreSourceContextsValid &= GenerateCompilerSourceContext(Binding.DestinationPath);
+						RunGenerateCompilerSourceContext(Binding.DestinationPath, FName());
 					}
 					else
 					{
-						bAreSourceContextsValid &= GenerateCompilerSourceContext(Binding.SourcePath);
+
+						RunGenerateCompilerSourceContext(Binding.SourcePath, FName());
 					}
 				}
 			}
 			else
 			{
 				// if we aren't using a conversion function, just validate the widget and viewmodel paths
-				bAreSourceContextsValid &= GenerateCompilerSourceContext(Binding.DestinationPath);
-				bAreSourceContextsValid &= GenerateCompilerSourceContext(Binding.SourcePath);
+				RunGenerateCompilerSourceContext(Binding.DestinationPath, FName());
+				RunGenerateCompilerSourceContext(Binding.SourcePath, FName());
+			}
+		}
+	}
+
+	// Find the event property.
+	//All inputs must also exist (for the BP to compile).
+	for (UMVVMBlueprintViewEvent* EventPtr : BlueprintView->GetEvents())
+	{
+		if (EventPtr == nullptr)
+		{
+			WidgetBlueprintCompilerContext.MessageLog.Error(*LOCTEXT("EventInvalid", "Internal error: An event is invalid.").ToString());
+			bAreSourceContextsValid = false;
+			continue;
+		}
+		
+		if (!EventPtr->bCompile)
+		{
+			continue;
+		}
+
+		if (!EventPtr->GetEventPath().HasPaths())
+		{
+			AddMessageForEvent(EventPtr, LOCTEXT("EventInvalidEventPath", "The event path is invalid."), UMVVMBlueprintViewEvent::EMessageType::Error, FName());
+			bAreSourceContextsValid = false;
+			continue;
+		}
+
+		auto RunGenerateCompilerSourceContext = [Self = this, &GenerateCompilerSourceContext, EventPtr, BlueprintView](const FMVVMBlueprintPropertyPath& PropertyPath, FName ArgumentName)
+		{
+			TValueOrError<void, FText> SourceContextResult = GenerateCompilerSourceContext(PropertyPath);
+			if (SourceContextResult.HasError())
+			{
+				Self->AddMessageForEvent(EventPtr, SourceContextResult.StealError(), UMVVMBlueprintViewEvent::EMessageType::Error, ArgumentName);
+				Self->bAreSourceContextsValid = false;
+			}
+		};
+
+		RunGenerateCompilerSourceContext(EventPtr->GetEventPath(), FName());
+		RunGenerateCompilerSourceContext(EventPtr->GetDestinationPath(), FName());
+		for (const FMVVMBlueprintPin& Pin : EventPtr->GetPins())
+		{
+			if (Pin.UsedPathAsValue())
+			{
+				RunGenerateCompilerSourceContext(Pin.GetPath(), Pin.GetName());
 			}
 		}
 	}
@@ -652,7 +713,7 @@ void FMVVMViewBlueprintCompiler::CreateFunctionsDeclaration(const FWidgetBluepri
 
 void FMVVMViewBlueprintCompiler::CreateIntermediateGraphFunctions(const FWidgetBlueprintCompilerContext::FCreateVariableContext & Context, UMVVMBlueprintView * BlueprintView)
 {
-	// Add Generated Conversion functions to the Context.List
+	// Add Generated Conversion functions to the blueprint
 	const int32 NumBindings = BlueprintView->GetNumBindings();
 	for (int32 Index = 0; Index < NumBindings; ++Index)
 	{
@@ -660,7 +721,7 @@ void FMVVMViewBlueprintCompiler::CreateIntermediateGraphFunctions(const FWidgetB
 		if (BindingPtr == nullptr)
 		{
 			WidgetBlueprintCompilerContext.MessageLog.Error(*FText::Format(LOCTEXT("InvalidBindingIndex", "Internal error. Invalid binding index given."), Index).ToString());
-			bIsBindingsValid = false;
+			bAreBindingsValid = false;
 			continue;
 		}
 
@@ -696,12 +757,41 @@ void FMVVMViewBlueprintCompiler::CreateIntermediateGraphFunctions(const FWidgetB
 			AddGeneratedConversionFunction(false);
 		}
 	}
+
+	// Add Generated event to the blueprint
+	for (UMVVMBlueprintViewEvent* EventPtr : BlueprintView->GetEvents())
+	{
+		if (EventPtr == nullptr)
+		{
+			WidgetBlueprintCompilerContext.MessageLog.Error(*LOCTEXT("InvalidEvent", "Internal error. Invalid event.").ToString());
+			bAreEventsValid = false;
+			continue;
+		}
+
+		if (!EventPtr->bCompile)
+		{
+			continue;
+		}
+
+		UEdGraph* WrapperGraph = EventPtr->GetOrCreateWrapperGraph();
+		if (WrapperGraph == nullptr)
+		{
+			AddMessageForEvent(EventPtr, LOCTEXT("InvalidEventGraph", "The event could not be generated"), UMVVMBlueprintViewEvent::EMessageType::Warning, FName());
+			continue;
+		}
+
+		bool bAlreadyContained = WidgetBlueprintCompilerContext.Blueprint->FunctionGraphs.Contains(WrapperGraph);
+		if (ensure(!bAlreadyContained))
+		{
+			Context.AddGeneratedFunctionGraph(WrapperGraph);
+		}
+	}
 }
 
 
 void FMVVMViewBlueprintCompiler::CreateBindingDestinationContexts(UMVVMBlueprintView* BlueprintView)
 {
-	if (!bAreSourcesCreatorValid || !bAreSourceContextsValid || !bIsBindingsValid)
+	if (!bAreSourcesCreatorValid || !bAreSourceContextsValid || !bAreBindingsValid || !bAreEventsValid)
 	{
 		return;
 	}
@@ -723,7 +813,7 @@ void FMVVMViewBlueprintCompiler::CreateBindingDestinationContexts(UMVVMBlueprint
 		if (BindingPtr == nullptr)
 		{
 			WidgetBlueprintCompilerContext.MessageLog.Error(*FText::Format(LOCTEXT("InvalidBindingIndex", "Internal error. Invalid binding index given."), Index).ToString());
-			bIsBindingsValid = false;
+			bAreBindingsValid = false;
 			continue;
 		}
 
@@ -743,24 +833,24 @@ void FMVVMViewBlueprintCompiler::CreateBindingDestinationContexts(UMVVMBlueprint
 			TArray<UE::MVVM::FMVVMConstFieldVariant> SkeletalSetterPath = CreateBindingDestinationPath(BlueprintView, NewSkeletonClass, DestinationPath);
 			if (!IsPropertyPathValid(SkeletalSetterPath))
 			{
-				AddMessageForBinding(*BindingPtr, BlueprintView, FText::Format(LOCTEXT("PropertyPathIsInvalid", "The property path '{0}' is invalid."),
+				AddMessageForBinding(*BindingPtr, BlueprintView, FText::Format(Private::PropertyPathIsInvalidFormat,
 					PropertyPathToText(NewSkeletonClass, BlueprintView, DestinationPath)),
 					EBindingMessageType::Error
 				);
-				bIsBindingsValid = false;
+				bAreBindingsValid = false;
 				return;
 			}
 			
 			// Generate the path with property converted to BP function
 			TValueOrError<TArray<FMVVMConstFieldVariant>, FText> SkeletalGeneratedField = FieldPathHelper::GenerateFieldPathList(SkeletalSetterPath, false);
-			if (SkeletalGeneratedField.HasError())
+			if (SkeletalGeneratedField.HasError() || !IsPropertyPathValid(SkeletalGeneratedField.GetValue()))
 			{
 				AddMessageForBinding(*BindingPtr, BlueprintView, FText::Format(Private::CouldNotCreateSourceFieldPathFormat
 					, ::UE::MVVM::FieldPathHelper::ToText(SkeletalSetterPath)
 					, SkeletalGeneratedField.GetError())
 					, EBindingMessageType::Error
 				);
-				bIsBindingsValid = false;
+				bAreBindingsValid = false;
 				return;
 			}
 
@@ -781,7 +871,7 @@ void FMVVMViewBlueprintCompiler::CreateBindingDestinationContexts(UMVVMBlueprint
 						PropertyPathToText(NewSkeletonClass, BlueprintView, DestinationPath)),
 						EBindingMessageType::Error
 					);
-					bIsBindingsValid = false;
+					bAreBindingsValid = false;
 					return;
 				}
 
@@ -793,7 +883,7 @@ void FMVVMViewBlueprintCompiler::CreateBindingDestinationContexts(UMVVMBlueprint
 						PropertyPathToText(NewSkeletonClass, BlueprintView, DestinationPath)),
 						EBindingMessageType::Error
 					);
-					bIsBindingsValid = false;
+					bAreBindingsValid = false;
 					return;
 				}
 
@@ -806,7 +896,7 @@ void FMVVMViewBlueprintCompiler::CreateBindingDestinationContexts(UMVVMBlueprint
 						PropertyPathToText(NewSkeletonClass, BlueprintView, DestinationPath)),
 						EBindingMessageType::Error
 					);
-					bIsBindingsValid = false;
+					bAreBindingsValid = false;
 					return;
 				}
 
@@ -829,7 +919,7 @@ void FMVVMViewBlueprintCompiler::CreateBindingDestinationContexts(UMVVMBlueprint
 
 bool FMVVMViewBlueprintCompiler::PreCompile(UWidgetBlueprintGeneratedClass* Class, UMVVMBlueprintView* BlueprintView)
 {
-	if (!bAreSourcesCreatorValid || !bAreSourceContextsValid || !bIsBindingsValid)
+	if (!bAreSourcesCreatorValid || !bAreSourceContextsValid || !bAreBindingsValid || !bAreEventsValid)
 	{
 		return false;
 	}
@@ -839,20 +929,22 @@ bool FMVVMViewBlueprintCompiler::PreCompile(UWidgetBlueprintGeneratedClass* Clas
 	BindingSourceContexts.Reset(NumBindings*2);
 	SimpleBindingContexts.Reset();
 	ComplexConversionFunctionContexts.Reset();
+	EventDelegateContexts.Reset();
 
 	FPropertyEditorPermissionList::Get().AddPermissionList(Class, FNamePermissionList(), EPropertyPermissionListRules::AllowListAllProperties);
 
 	PreCompileBindingSources(Class, BlueprintView);
 	PreCompileSourceCreators(Class, BlueprintView);
 	PreCompileBindings(Class, BlueprintView);
+	PreCompileEvents(Class, BlueprintView);
 
-	return bAreSourcesCreatorValid && bAreSourceContextsValid && bIsBindingsValid;
+	return bAreSourcesCreatorValid && bAreSourceContextsValid && bAreBindingsValid && bAreEventsValid;
 }
 
 
 bool FMVVMViewBlueprintCompiler::Compile(UWidgetBlueprintGeneratedClass* Class, UMVVMBlueprintView* BlueprintView, UMVVMViewClass* ViewExtension)
 {
-	if (!bAreSourcesCreatorValid || !bAreSourceContextsValid || !bIsBindingsValid)
+	if (!bAreSourcesCreatorValid || !bAreSourceContextsValid || !bAreBindingsValid || !bAreEventsValid)
 	{
 		return false;
 	}
@@ -865,13 +957,14 @@ bool FMVVMViewBlueprintCompiler::Compile(UWidgetBlueprintGeneratedClass* Class,
 	}
 	CompileSourceCreators(CompileResult.GetValue(), Class, BlueprintView, ViewExtension);
 	CompileBindings(CompileResult.GetValue(), Class, BlueprintView, ViewExtension);
+	CompileEvents(CompileResult.GetValue(), Class, BlueprintView, ViewExtension);
 
 	{
 		ViewExtension->bInitializeSourcesOnConstruct = BlueprintView->GetSettings()->bInitializeSourcesOnConstruct;
 		ViewExtension->bInitializeBindingsOnConstruct = ViewExtension->bInitializeSourcesOnConstruct ? BlueprintView->GetSettings()->bInitializeBindingsOnConstruct : false;
 	}
 
-	bool bResult = bAreSourcesCreatorValid && bAreSourceContextsValid && bIsBindingsValid;
+	bool bResult = bAreSourcesCreatorValid && bAreSourceContextsValid && bAreBindingsValid && bAreEventsValid;
 	if (bResult)
 	{
 		ViewExtension->BindingLibrary = MoveTemp(CompileResult.GetValue().Library);
@@ -901,7 +994,7 @@ bool FMVVMViewBlueprintCompiler::PreCompileBindingSources(UWidgetBlueprintGenera
 		if (BindingPtr == nullptr)
 		{
 			WidgetBlueprintCompilerContext.MessageLog.Error(*FText::Format(LOCTEXT("InvalidBindingIndex", "Internal error. Invalid binding index given."), Index).ToString());
-			bIsBindingsValid = false;
+			bAreBindingsValid = false;
 			continue;
 		}
 
@@ -934,9 +1027,7 @@ bool FMVVMViewBlueprintCompiler::PreCompileBindingSources(UWidgetBlueprintGenera
 			{
 				AddMessageForBinding(Binding,
 					BlueprintView,
-					FText::Format(LOCTEXT("PropertyPathIsInvalid", "The property path '{0}' is invalid."), 
-						PropertyPathToText(Class, BlueprintView, Binding.SourcePath)
-					),
+					FText::Format(Private::PropertyPathIsInvalidFormat, PropertyPathToText(Class, BlueprintView, Binding.SourcePath)),
 					EBindingMessageType::Error,
 					ArgumentName
 				);
@@ -1088,20 +1179,20 @@ bool FMVVMViewBlueprintCompiler::PreCompileBindingSources(UWidgetBlueprintGenera
 					}
 					if (!CreateSourceContextForPropertyPath(Binding.SourcePath, true, INDEX_NONE, FName()))
 					{
-						bIsBindingsValid = false;
+						bAreBindingsValid = false;
 						continue;
 					}
 				}
 				else
 				{
 					AddMessageForBinding(Binding, BlueprintView, LOCTEXT("BindingEmptySourcePath", "The binding doesn't have a Source or a Conversion function."), EBindingMessageType::Error);
-					bIsBindingsValid = false;
+					bAreBindingsValid = false;
 					continue;
 				}
 			}
 			else if (ConversionFunctionResult == ECreateSourcesForConversionFunctionResult::Failed)
 			{
-				bIsBindingsValid = false;
+				bAreBindingsValid = false;
 				continue;
 			}
 		}
@@ -1120,26 +1211,26 @@ bool FMVVMViewBlueprintCompiler::PreCompileBindingSources(UWidgetBlueprintGenera
 					}
 					if (!CreateSourceContextForPropertyPath(Binding.DestinationPath, false, INDEX_NONE, FName()))
 					{
-						bIsBindingsValid = false;
+						bAreBindingsValid = false;
 						continue;
 					}
 				}
 				else
 				{
 					AddMessageForBinding(Binding, BlueprintView, LOCTEXT("BindingEmptyDestinationPath", "The binding doesn't have a Destination or a Conversion function."), EBindingMessageType::Error);
-					bIsBindingsValid = false;
+					bAreBindingsValid = false;
 					continue;
 				}
 			}
 			else if (ConversionFunctionResult == ECreateSourcesForConversionFunctionResult::Failed)
 			{
-				bIsBindingsValid = false;
+				bAreBindingsValid = false;
 				continue;
 			}
 		}
 	}
 
-	return bIsBindingsValid;
+	return bAreBindingsValid;
 }
 
 
@@ -1361,7 +1452,7 @@ bool FMVVMViewBlueprintCompiler::CompileSourceCreators(const FCompiledBindingLib
 
 bool FMVVMViewBlueprintCompiler::PreCompileBindings(UWidgetBlueprintGeneratedClass* Class, UMVVMBlueprintView* BlueprintView)
 {
-	if (!bAreSourceContextsValid || !bIsBindingsValid)
+	if (!bAreSourceContextsValid || !bAreBindingsValid)
 	{
 		return false;
 	}
@@ -1379,7 +1470,7 @@ bool FMVVMViewBlueprintCompiler::PreCompileBindings(UWidgetBlueprintGeneratedCla
 					FText::FromName(ViewConversionFunction->GetCompiledFunctionName())),
 					EBindingMessageType::Error
 				);
-				bIsBindingsValid = false;
+				bAreBindingsValid = false;
 				return MakeError();
 			}
 
@@ -1392,14 +1483,14 @@ bool FMVVMViewBlueprintCompiler::PreCompileBindings(UWidgetBlueprintGeneratedCla
 						FText::FromName(FunctionOrWrapperFunction.Get<const UFunction*>()->GetFName())),
 						EBindingMessageType::Error
 					);
-					bIsBindingsValid = false;
+					bAreBindingsValid = false;
 					return MakeError();
 				}
 			}
 			else
 			{
 				AddMessageForBinding(Binding, BlueprintView, LOCTEXT("ConversionFunctionNodeNotAllow", "The conversion function node is not allowed."), EBindingMessageType::Error);
-				bIsBindingsValid = false;
+				bAreBindingsValid = false;
 				return MakeError();
 			}
 		}
@@ -1421,7 +1512,6 @@ bool FMVVMViewBlueprintCompiler::PreCompileBindings(UWidgetBlueprintGeneratedCla
 
 	auto GetSetterPath = [this, Class, BlueprintView](FMVVMBlueprintViewBinding& Binding, int32 BindingIndex, bool bIsForwardBinding) -> TValueOrError<TArray<UE::MVVM::FMVVMConstFieldVariant>, void>
 	{
-
 		const FBindingDestinationContext* DestinationContext = BindingDestinationContexts.FindByPredicate([BindingIndex, bIsForwardBinding](const FBindingDestinationContext& Other)
 			{
 				return Other.BindingIndex == BindingIndex && Other.bIsForwardBinding == bIsForwardBinding;
@@ -1457,10 +1547,7 @@ bool FMVVMViewBlueprintCompiler::PreCompileBindings(UWidgetBlueprintGeneratedCla
 		if (!IsPropertyPathValid(SetterPath))
 		{
 			const FMVVMBlueprintPropertyPath& DestinationPath = bIsForwardBinding ? Binding.DestinationPath : Binding.SourcePath;
-			AddMessageForBinding(Binding, BlueprintView, FText::Format(LOCTEXT("PropertyPathIsInvalid", "The property path '{0}' is invalid."),
-				PropertyPathToText(Class, BlueprintView, DestinationPath)),
-				EBindingMessageType::Error
-			);
+			AddMessageForBinding(Binding, BlueprintView, FText::Format(Private::PropertyPathIsInvalidFormat, PropertyPathToText(Class, BlueprintView, DestinationPath)), EBindingMessageType::Error);
 			return MakeError();
 		}
 		return MakeValue(SetterPath);
@@ -1484,7 +1571,7 @@ bool FMVVMViewBlueprintCompiler::PreCompileBindings(UWidgetBlueprintGeneratedCla
 
 		if (!TestExecutionMode(Binding))
 		{
-			bIsBindingsValid = false;
+			bAreBindingsValid = false;
 			continue;
 		}
 
@@ -1500,7 +1587,7 @@ bool FMVVMViewBlueprintCompiler::PreCompileBindings(UWidgetBlueprintGeneratedCla
 			{
 				AddMessageForBinding(Binding, BlueprintView, FText::Format(LOCTEXT("CouldNotCreateSource", "Could not create source. {0}"),
 					AddFieldResult.GetError()), EBindingMessageType::Error);
-				bIsBindingsValid = false;
+				bAreBindingsValid = false;
 				continue;
 			}
 		}
@@ -1517,7 +1604,7 @@ bool FMVVMViewBlueprintCompiler::PreCompileBindings(UWidgetBlueprintGeneratedCla
 			TValueOrError<TArray<UE::MVVM::FMVVMConstFieldVariant>, void> SetterPathResult = GetSetterPath(Binding, BindingSourceContext.BindingIndex, BindingSourceContext.bIsForwardBinding);
 			if (SetterPathResult.HasError())
 			{
-				bIsBindingsValid = false;
+				bAreBindingsValid = false;
 				continue;
 			}
 			SetterPath = SetterPathResult.StealValue();
@@ -1528,7 +1615,7 @@ bool FMVVMViewBlueprintCompiler::PreCompileBindings(UWidgetBlueprintGeneratedCla
 			TValueOrError<const UFunction*, void> ConversionFunctionResult = GetConversionFunction(Binding, BindingSourceContext.bIsForwardBinding);
 			if (ConversionFunctionResult.HasError())
 			{
-				bIsBindingsValid = false;
+				bAreBindingsValid = false;
 				continue;
 			}
 			ConversionFunction = ConversionFunctionResult.StealValue();
@@ -1541,7 +1628,7 @@ bool FMVVMViewBlueprintCompiler::PreCompileBindings(UWidgetBlueprintGeneratedCla
 				FText::Format(LOCTEXT("CouldNotCreateBinding", "Could not create binding. {0}"), AddBindingResult.GetError()),
 				EBindingMessageType::Error
 			);
-			bIsBindingsValid = false;
+			bAreBindingsValid = false;
 			continue;
 		}
 
@@ -1572,14 +1659,14 @@ bool FMVVMViewBlueprintCompiler::PreCompileBindings(UWidgetBlueprintGeneratedCla
 
 		if (!TestExecutionMode(Binding))
 		{
-			bIsBindingsValid = false;
+			bAreBindingsValid = false;
 			continue;
 		}
 
 		if (!IsOneTimeBinding(Binding.BindingType))
 		{
 			AddMessageForBinding(Binding, BlueprintView, LOCTEXT("FunctionRequiresOneTimeBinding", "There is no source. The binding must be a OneTime binding."), EBindingMessageType::Error);
-			bIsBindingsValid = false;
+			bAreBindingsValid = false;
 			continue;
 		}
 
@@ -1588,7 +1675,7 @@ bool FMVVMViewBlueprintCompiler::PreCompileBindings(UWidgetBlueprintGeneratedCla
 			TValueOrError<TArray<UE::MVVM::FMVVMConstFieldVariant>, void> SetterPathResult = GetSetterPath(Binding, BindingContext.BindingIndex, BindingContext.bIsForwardBinding);
 			if (SetterPathResult.HasError())
 			{
-				bIsBindingsValid = false;
+				bAreBindingsValid = false;
 				continue;
 			}
 			SetterPath = SetterPathResult.StealValue();
@@ -1599,7 +1686,7 @@ bool FMVVMViewBlueprintCompiler::PreCompileBindings(UWidgetBlueprintGeneratedCla
 			TValueOrError<const UFunction*, void> ConversionFunctionResult = GetConversionFunction(Binding, BindingContext.bIsForwardBinding);
 			if (ConversionFunctionResult.HasError())
 			{
-				bIsBindingsValid = false;
+				bAreBindingsValid = false;
 				continue;
 			}
 			ConversionFunction = ConversionFunctionResult.StealValue();
@@ -1614,7 +1701,7 @@ bool FMVVMViewBlueprintCompiler::PreCompileBindings(UWidgetBlueprintGeneratedCla
 				FText::Format(LOCTEXT("CouldNotCreateBinding", "Could not create binding. {0}"), AddBindingResult.GetError()),
 				EBindingMessageType::Error
 			);
-			bIsBindingsValid = false;
+			bAreBindingsValid = false;
 			continue;
 		}
 
@@ -1629,7 +1716,7 @@ bool FMVVMViewBlueprintCompiler::PreCompileBindings(UWidgetBlueprintGeneratedCla
 		CompilerBindings.Emplace(NewBinding);
 	}
 
-	if (bIsBindingsValid)
+	if (bAreBindingsValid)
 	{
 		// Confirm that at least one source is valid for each pin of a complex conversion function.
 		for (const FComplexConversionFunctionContext& ComplexConversionFunctionContext : ComplexConversionFunctionContexts)
@@ -1639,18 +1726,18 @@ bool FMVVMViewBlueprintCompiler::PreCompileBindings(UWidgetBlueprintGeneratedCla
 				FMVVMBlueprintViewBinding* BindingPtr = BlueprintView->GetBindingAt(ComplexConversionFunctionContext.BindingIndex);
 				check(BindingPtr);
 				AddMessageForBinding(*BindingPtr, BlueprintView, LOCTEXT("CouldNotCreateSourceForConversionFunction", "There is no source. The binding must be a OneTime binding."), EBindingMessageType::Error);
-				bIsBindingsValid = false;
+				bAreBindingsValid = false;
 			}
 		}
 	}
 
-	return bIsBindingsValid;
+	return bAreBindingsValid;
 }
 
 
 bool FMVVMViewBlueprintCompiler::CompileBindings(const FCompiledBindingLibraryCompiler::FCompileResult& CompileResult, UWidgetBlueprintGeneratedClass* Class, UMVVMBlueprintView* BlueprintView, UMVVMViewClass* ViewExtension)
 {
-	if (!bIsBindingsValid)
+	if (!bAreBindingsValid)
 	{
 		return false;
 	}
@@ -1686,12 +1773,12 @@ bool FMVVMViewBlueprintCompiler::CompileBindings(const FCompiledBindingLibraryCo
 			else if (CompileBinding.bNeedsValidSource)
 			{
 				WidgetBlueprintCompilerContext.MessageLog.Error(*LOCTEXT("InvalidSourceInternal", "Internal error. The binding doesn't have a valid source.").ToString());
-				bIsBindingsValid = false;
+				bAreBindingsValid = false;
 				return false;
 			}
 		}
 
-		auto AddErroMessage = [this, ViewBinding, BlueprintView](const FText& ErrorMessage)
+		auto AddErrorMessage = [this, ViewBinding, BlueprintView](const FText& ErrorMessage)
 		{
 				if (ViewBinding)
 				{
@@ -1706,16 +1793,16 @@ bool FMVVMViewBlueprintCompiler::CompileBindings(const FCompiledBindingLibraryCo
 		const FMVVMVCompiledFieldId* CompiledFieldId = CompileResult.FieldIds.Find(CompileBinding.FieldIdHandle);
 		if (CompiledFieldId == nullptr && CompileBinding.bFieldIdNeeded)
 		{
-			AddErroMessage(FText::Format(LOCTEXT("FieldIdNotGenerated", "Could not generate field ID for property '{0}'."), FText::FromName(NewBinding.SourcePropertyName)));
-			bIsBindingsValid = false;
+			AddErrorMessage(FText::Format(LOCTEXT("FieldIdNotGenerated", "Could not generate field ID for property '{0}'."), FText::FromName(NewBinding.SourcePropertyName)));
+			bAreBindingsValid = false;
 			continue;
 		}
 
 		const FMVVMVCompiledBinding* CompiledBinding = CompileResult.Bindings.Find(CompileBinding.CompiledBinding.BindingHandle);
 		if (CompiledBinding == nullptr && CompileBinding.Type != ECompilerBindingType::ViewModelDynamic)
 		{
-			AddErroMessage(LOCTEXT("CompiledBindingNotGenerated", "Could not generate compiled binding."));
-			bIsBindingsValid = false;
+			AddErrorMessage(LOCTEXT("CompiledBindingNotGenerated", "Could not generate compiled binding."));
+			bAreBindingsValid = false;
 			continue;
 		}
 
@@ -1730,8 +1817,8 @@ bool FMVVMViewBlueprintCompiler::CompileBindings(const FCompiledBindingLibraryCo
 					ViewModelContext = BlueprintView->FindViewModel(CompilerUserWidgetPropertyContexts[CompileBinding.UserWidgetPropertyContextIndex].ViewModelId);
 					if (ViewModelContext == nullptr)
 					{
-						AddErroMessage(LOCTEXT("CompiledBindingWithInvalidIVewModelId", "Internal error: the viewmodel became invalid."));
-						bIsBindingsValid = false;
+						AddErrorMessage(LOCTEXT("CompiledBindingWithInvalidIVewModelId", "Internal error: the viewmodel became invalid."));
+						bAreBindingsValid = false;
 						continue;
 					}
 				}
@@ -1760,7 +1847,7 @@ bool FMVVMViewBlueprintCompiler::CompileBindings(const FCompiledBindingLibraryCo
 			{
 				FText ErrorMessage = FText::Format(LOCTEXT("CompiledViewModelDynamicBindingNotGenerated", "Was not able to find the source for {0}."), FText::FromName(CompileBinding.DynamicViewModelName));
 				WidgetBlueprintCompilerContext.MessageLog.Error(*ErrorMessage.ToString());
-				bIsBindingsValid = false;
+				bAreBindingsValid = false;
 				continue;
 			}
 		}
@@ -1775,7 +1862,7 @@ bool FMVVMViewBlueprintCompiler::CompileBindings(const FCompiledBindingLibraryCo
 			if (CompileBinding.Type != ECompilerBindingType::PropertyBinding)
 			{
 				WidgetBlueprintCompilerContext.MessageLog.Error(*WrongBindingTypeInternalErrorMessage.ToString());
-				bIsBindingsValid = false;
+				bAreBindingsValid = false;
 				continue;
 			}
 
@@ -1801,7 +1888,7 @@ bool FMVVMViewBlueprintCompiler::CompileBindings(const FCompiledBindingLibraryCo
 			if (CompileBinding.Type != ECompilerBindingType::ViewModelDynamic)
 			{
 				WidgetBlueprintCompilerContext.MessageLog.Error(*WrongBindingTypeInternalErrorMessage.ToString());
-				bIsBindingsValid = false;
+				bAreBindingsValid = false;
 				continue;
 			}
 			NewBinding.ExecutionMode = EMVVMExecutionMode::Immediate;
@@ -1816,7 +1903,190 @@ bool FMVVMViewBlueprintCompiler::CompileBindings(const FCompiledBindingLibraryCo
 		ViewExtension->CompiledBindings.Emplace(MoveTemp(NewBinding));
 	}
 
-	return bIsBindingsValid;
+	return bAreBindingsValid;
+}
+
+
+bool FMVVMViewBlueprintCompiler::PreCompileEvents(UWidgetBlueprintGeneratedClass* Class, UMVVMBlueprintView* BlueprintView)
+{
+	if (!bAreSourceContextsValid || !bAreEventsValid)
+	{
+		return false;
+	}
+
+	for (UMVVMBlueprintViewEvent* EventPtr : BlueprintView->GetEvents())
+	{
+		if (EventPtr == nullptr || !EventPtr->bCompile)
+		{
+			continue;
+		}
+
+		UEdGraph* GeneratedGraph = EventPtr->GetOrCreateWrapperGraph();
+		if (GeneratedGraph == nullptr)
+		{
+			AddMessageForEvent(EventPtr
+				, FText::Format(LOCTEXT("InvalidEventGeneratedGraph", "The generated graph for event {0} is invalid."), EventPtr->GetDisplayName(true))
+				, UMVVMBlueprintViewEvent::EMessageType::Error
+				, FName());
+			bAreEventsValid = false;
+			continue;
+		}
+
+		auto TestPath = [Self = this, BlueprintView, Class](UMVVMBlueprintViewEvent* EventPtr, const FMVVMBlueprintPropertyPath& PropertyPath, FName ArgumentName)
+		{
+			TArray<UE::MVVM::FMVVMConstFieldVariant> GeneratedPath = Self->CreateBindingDestinationPath(BlueprintView, Class, PropertyPath);
+			if (!Self->IsPropertyPathValid(GeneratedPath))
+			{
+				Self->AddMessageForEvent(EventPtr
+					, FText::Format(Private::PropertyPathIsInvalidFormat, PropertyPathToText(Class, BlueprintView, PropertyPath))
+					, UMVVMBlueprintViewEvent::EMessageType::Error
+					, ArgumentName);
+				Self->bAreEventsValid = false;
+				GeneratedPath.Empty();
+			}
+			return GeneratedPath;
+		};
+
+		// Does it resolve and are the field allowed
+		TArray<UE::MVVM::FMVVMConstFieldVariant> EventPath = TestPath(EventPtr, EventPtr->GetEventPath(), FName());
+		if (EventPath.Num() == 0)
+		{
+			continue;
+		}
+		if (TestPath(EventPtr, EventPtr->GetDestinationPath(), FName()).Num() == 0)
+		{
+			continue;
+		}
+
+		for (const FMVVMBlueprintPin& Pin : EventPtr->GetPins())
+		{
+			if (Pin.UsedPathAsValue())
+			{
+				if (TestPath(EventPtr, Pin.GetPath(), Pin.GetName()).Num() == 0)
+				{
+					continue;
+				}
+			}
+		}
+
+		//Does EventPath resolves to a MulticastDelegateProperty
+		const FMulticastDelegateProperty* DelegateProperty = EventPath.Last().IsProperty() ? CastField<const FMulticastDelegateProperty>(EventPath.Last().GetProperty()) : nullptr;
+		if (DelegateProperty == nullptr)
+		{
+			AddMessageForEvent(EventPtr
+				, FText::Format(LOCTEXT("EventPathIsNotMulticastDelegate", "The event {0} is not a multicast delegate."), PropertyPathToText(Class, BlueprintView, EventPtr->GetEventPath()))
+				, UMVVMBlueprintViewEvent::EMessageType::Error
+				, FName());
+			bAreEventsValid = false;
+			continue;
+		}
+
+		// Does it still have the same signature
+		if (!UE::MVVM::FunctionGraphHelper::IsFunctionEntryMatchSignature(GeneratedGraph, DelegateProperty->SignatureFunction))
+		{
+			AddMessageForEvent(EventPtr
+				, FText::Format(LOCTEXT("EventPathFunctionSignatureError", "The event {0} doesn't match the function signature."), DelegateProperty->GetDisplayNameText())
+				, UMVVMBlueprintViewEvent::EMessageType::Error
+				, FName());
+			bAreEventsValid = false;
+			continue;
+		}
+
+		// Generate the FieldPath to get the delegate property at runtime
+		TValueOrError<TArray<FMVVMConstFieldVariant>, FText> GeneratedEventFields = FieldPathHelper::GenerateFieldPathList(EventPath, true);
+		if (GeneratedEventFields.HasError())
+		{
+			AddMessageForEvent(EventPtr
+				, FText::Format(Private::CouldNotCreateSourceFieldPathFormat, ::UE::MVVM::FieldPathHelper::ToText(EventPath), GeneratedEventFields.GetError())
+				, UMVVMBlueprintViewEvent::EMessageType::Error
+				, FName());
+			bAreEventsValid = false;
+			continue;
+		}
+
+		TValueOrError<FCompiledBindingLibraryCompiler::FFieldPathHandle, FText> FieldPathResult = BindingLibraryCompiler.AddFieldPath(GeneratedEventFields.GetValue(), true);
+		if (FieldPathResult.HasError())
+		{
+			AddMessageForEvent(EventPtr
+				, FText::Format(Private::CouldNotCreateSourceFieldPathFormat, ::UE::MVVM::FieldPathHelper::ToText(EventPath), FieldPathResult.GetError())
+				, UMVVMBlueprintViewEvent::EMessageType::Error
+				, FName());
+			bAreEventsValid = false;
+			continue;
+		}
+
+		FName SourceName;
+		if (EventPtr->GetEventPath().IsFromViewModel())
+		{
+			const FMVVMBlueprintViewModelContext* SourceViewModelContext = BlueprintView->FindViewModel(EventPtr->GetEventPath().GetViewModelId());
+			check(SourceViewModelContext);
+			FName ViewModelName = SourceViewModelContext->GetViewModelName();
+			const int32 DestinationVariableContextIndex = CompilerUserWidgetPropertyContexts.IndexOfByPredicate([ViewModelName](const FCompilerUserWidgetPropertyContext& Other) { return Other.PropertyName == ViewModelName; });
+			check(DestinationVariableContextIndex != INDEX_NONE);
+			SourceName = CompilerUserWidgetPropertyContexts[DestinationVariableContextIndex].PropertyName;
+		}
+		else if (EventPtr->GetEventPath().IsFromWidget())
+		{
+			FName WidgetName = EventPtr->GetEventPath().GetWidgetName();
+			checkf(!WidgetName.IsNone(), TEXT("The destination should have been checked and set bAreSourceContextsValid."));
+			const bool bSourceIsUserWidget = WidgetName == Class->ClassGeneratedBy->GetFName();
+			if (!bSourceIsUserWidget)
+			{
+				const int32 VariableContextIndex = CompilerUserWidgetPropertyContexts.IndexOfByPredicate([WidgetName](const FCompilerUserWidgetPropertyContext& Other) { return Other.PropertyName == WidgetName; });
+				if (ensureAlwaysMsgf(VariableContextIndex != INDEX_NONE, TEXT("Could not find source context for destination '%s'"), *WidgetName.ToString()))
+				{
+					SourceName = CompilerUserWidgetPropertyContexts[VariableContextIndex].PropertyName;
+				}
+			}
+		}
+
+		// No need to add the generated function to the field compiler.
+		//They are in the BP generated code.
+
+		FEventDelegateContext& NewContext = EventDelegateContexts.AddDefaulted_GetRef();
+		NewContext.BlueprintEvent = EventPtr;
+		NewContext.DelegateFieldPathHandle = FieldPathResult.StealValue();
+		NewContext.GeneratedGraphName = EventPtr->GetWrapperGraphName();
+		NewContext.SourceName = SourceName;
+	}
+
+	return bAreEventsValid;
+}
+
+
+bool FMVVMViewBlueprintCompiler::CompileEvents(const FCompiledBindingLibraryCompiler::FCompileResult& CompileResult, UWidgetBlueprintGeneratedClass* Class, UMVVMBlueprintView* BlueprintView, UMVVMViewClass* ViewExtension)
+{
+	if (!bAreSourceContextsValid || !bAreEventsValid)
+	{
+		return false;
+	}
+
+	for (FEventDelegateContext& Context : EventDelegateContexts)
+	{
+		check(Context.BlueprintEvent);
+
+		const FMVVMVCompiledFieldPath* CompiledFieldPath = CompileResult.FieldPaths.Find(Context.DelegateFieldPathHandle);
+		if (CompiledFieldPath == nullptr)
+		{
+			AddMessageForEvent(Context.BlueprintEvent, LOCTEXT("CompiledEventFieldPathNotGenerated", "Could not generate the event path."), UMVVMBlueprintViewEvent::EMessageType::Error, FName());
+			bAreEventsValid = false;
+			continue;
+		}
+
+		if (Context.GeneratedGraphName.IsNone() || Class->FindFunctionByName(Context.GeneratedGraphName) == nullptr)
+		{
+			AddMessageForEvent(Context.BlueprintEvent, LOCTEXT("CompiledEventFieldPathNotGenerated", "Could not generate the event path."), UMVVMBlueprintViewEvent::EMessageType::Error, FName());
+			bAreEventsValid = false;
+			continue;
+		}
+
+		FMVVMViewClass_CompiledEvent& NewBinding = ViewExtension->CompiledEvents.AddDefaulted_GetRef();
+		NewBinding.FieldPath = *CompiledFieldPath;
+		NewBinding.FunctionName = Context.GeneratedGraphName;
+		NewBinding.SourceName = Context.SourceName;
+	}
+
+	return bAreEventsValid;
 }
 
 
@@ -2170,13 +2440,27 @@ bool FMVVMViewBlueprintCompiler::IsPropertyPathValid(TArrayView<const FMVVMConst
 		{
 			return false;
 		}
-		if (Field.IsProperty() && Field.GetProperty() == nullptr)
+		if (Field.IsProperty())
 		{
-			return false;
+			if (Field.GetProperty() == nullptr)
+			{
+				return false;
+			}
+			if (!GetDefault<UMVVMDeveloperProjectSettings>()->IsPropertyAllowed(Field.GetProperty()))
+			{
+				return false;
+			}
 		}
-		if (Field.IsFunction() && Field.GetFunction() == nullptr)
+		if (Field.IsFunction())
 		{
-			return false;
+			if (Field.GetFunction() == nullptr)
+			{
+				return false;
+			}
+			if (!GetDefault<UMVVMDeveloperProjectSettings>()->IsFunctionAllowed(Field.GetFunction()))
+			{
+				return false;
+			}
 		}
 	}
 	return true;
@@ -2184,8 +2468,6 @@ bool FMVVMViewBlueprintCompiler::IsPropertyPathValid(TArrayView<const FMVVMConst
 
 bool FMVVMViewBlueprintCompiler::CanBeSetInNative(TArrayView<const FMVVMConstFieldVariant> PropertyPath)
 {
-	check(IsPropertyPathValid(PropertyPath));
-	
 	for (int32 Index = PropertyPath.Num() - 1; Index >= 0; --Index)
 	{
 		const FMVVMConstFieldVariant& Variant = PropertyPath[Index];
diff --git a/Engine/Plugins/Runtime/ModelViewViewModel/Source/ModelViewViewModelBlueprint/Private/MVVMViewBlueprintCompiler.h b/Engine/Plugins/Runtime/ModelViewViewModel/Source/ModelViewViewModelBlueprint/Private/MVVMViewBlueprintCompiler.h
--- a/Engine/Plugins/Runtime/ModelViewViewModel/Source/ModelViewViewModelBlueprint/Private/MVVMViewBlueprintCompiler.h
+++ b/Engine/Plugins/Runtime/ModelViewViewModel/Source/ModelViewViewModelBlueprint/Private/MVVMViewBlueprintCompiler.h
@@ -3,6 +3,7 @@
 #pragma once
 
 #include "Bindings/MVVMCompiledBindingLibraryCompiler.h"
+#include "MVVMBlueprintViewEvent.h"
 #include "MVVMBlueprintViewModelContext.h"
 #include "Types/MVVMFieldVariant.h"
 #include "WidgetBlueprintCompiler.h"
@@ -47,7 +48,6 @@ public:
 
 	void AddExtension(UWidgetBlueprintGeneratedClass* Class, UMVVMViewClass* ViewExtension);
 	void CleanOldData(UWidgetBlueprintGeneratedClass* ClassToClean, UObject* OldCDO);
-	void CleanTemporaries(UWidgetBlueprintGeneratedClass* ClassToClean);
 
 	/** Generate function that are hidden from the user (not on the Skeleton class). */
 	void CreateFunctions(UMVVMBlueprintView* BlueprintView);
@@ -77,10 +77,14 @@ private:
 
 	bool PreCompileBindings(UWidgetBlueprintGeneratedClass* Class, UMVVMBlueprintView* BlueprintView);
 	bool CompileBindings(const FCompiledBindingLibraryCompiler::FCompileResult& CompileResult, UWidgetBlueprintGeneratedClass* Class, UMVVMBlueprintView* BlueprintView, UMVVMViewClass* ViewExtension);
+	
+	bool PreCompileEvents(UWidgetBlueprintGeneratedClass* Class, UMVVMBlueprintView* BlueprintView);
+	bool CompileEvents(const FCompiledBindingLibraryCompiler::FCompileResult& CompileResult, UWidgetBlueprintGeneratedClass* Class, UMVVMBlueprintView* BlueprintView, UMVVMViewClass* ViewExtension);
 
 	const FCompilerSourceCreatorContext* FindViewModelSource(FGuid Id) const;
 
 	void AddMessageForBinding(FMVVMBlueprintViewBinding& Binding, UMVVMBlueprintView* BlueprintView, const FText& MessageText, EBindingMessageType MessageType, FName ArgumentName = FName()) const;
+	void AddMessageForEvent(UMVVMBlueprintViewEvent* Event, const FText& MessageText, UMVVMBlueprintViewEvent::EMessageType MessageType, FName ArgumentName) const;
 	void AddErrorForViewModel(const FMVVMBlueprintViewModelContext& ViewModel, const FText& Message) const;
 
 	TValueOrError<FBindingSourceContext, FText> CreateBindingSourceContext(const UMVVMBlueprintView* BlueprintView, const UWidgetBlueprintGeneratedClass* Class, const FMVVMBlueprintPropertyPath& PropertyPath, bool bIsOneTimeBinding);
@@ -149,7 +153,6 @@ private:
 		FCompiledBindingLibraryCompiler::FFieldPathHandle SourceRead;
 		FCompiledBindingLibraryCompiler::FFieldPathHandle DestinationWrite;
 		FCompiledBindingLibraryCompiler::FFieldPathHandle ConversionFunction;
-		FCompiledBindingLibraryCompiler::FFieldPathHandle ExecutionFunction;
 
 		bool bIsConversionFunctionComplex = false;
 	};
@@ -256,12 +259,26 @@ private:
 	};
 	TArray<FComplexConversionFunctionContext> ComplexConversionFunctionContexts;
 
+	/**
+	 * Context for multicast delegate
+	 */
+	struct FEventDelegateContext
+	{
+		UMVVMBlueprintViewEvent* BlueprintEvent = nullptr;
+
+		FCompiledBindingLibraryCompiler::FFieldPathHandle DelegateFieldPathHandle;
+		FName GeneratedGraphName;
+		FName SourceName;
+	};
+	TArray<FEventDelegateContext> EventDelegateContexts;
+
 	TMap<FName, UWidget*> WidgetNameToWidgetPointerMap;
 	FWidgetBlueprintCompilerContext& WidgetBlueprintCompilerContext;
 	FCompiledBindingLibraryCompiler BindingLibraryCompiler;
 	bool bAreSourcesCreatorValid = true;
 	bool bAreSourceContextsValid = true;
-	bool bIsBindingsValid = true;
+	bool bAreBindingsValid = true;
+	bool bAreEventsValid = true;
 };
 
 } //namespace
diff --git a/Engine/Plugins/Runtime/ModelViewViewModel/Source/ModelViewViewModelBlueprint/Private/MVVMWidgetBlueprintExtension_View.cpp b/Engine/Plugins/Runtime/ModelViewViewModel/Source/ModelViewViewModelBlueprint/Private/MVVMWidgetBlueprintExtension_View.cpp
--- a/Engine/Plugins/Runtime/ModelViewViewModel/Source/ModelViewViewModelBlueprint/Private/MVVMWidgetBlueprintExtension_View.cpp
+++ b/Engine/Plugins/Runtime/ModelViewViewModel/Source/ModelViewViewModelBlueprint/Private/MVVMWidgetBlueprintExtension_View.cpp
@@ -134,11 +134,21 @@ void UMVVMWidgetBlueprintExtension_View::HandleFinishCompilingClass(UWidgetBluep
 			bCompiled = CurrentCompilerContext->Compile(Class, BlueprintView, ViewExtension);
 		}
 
-		CurrentCompilerContext->CleanTemporaries(Class);
-
-		if (bCompiled && UE::MVVM::Private::GAllowViewClass)
+		if (bCompiled)
 		{
 			check(ViewExtension);
+			if (!UE::MVVM::Private::GAllowViewClass)
+			{
+				ViewExtension->SetFlags(RF_Transient);
+
+				// If we are not allowed to add the view class, add the transient flags on added conversion graph and event graph.
+				for (TFieldIterator<UFunction> FunctionIter(Class, EFieldIteratorFlags::ExcludeSuper); FunctionIter; ++FunctionIter)
+				{
+					UFunction* Function = *FunctionIter;
+					Function->SetFlags(RF_Transient);
+				}
+			}
+
 			// Does it have any bindings
 			if (const_cast<const UMVVMViewClass*>(ViewExtension)->GetCompiledBindings().Num() > 0)
 			{
@@ -151,22 +161,10 @@ void UMVVMWidgetBlueprintExtension_View::HandleFinishCompilingClass(UWidgetBluep
 				CurrentCompilerContext->AddExtension(Class, ViewExtension);
 			}
 		}
-		else if (bCompiled)
-		{
-			// If we are not allowed to add the view class, add the transient flags on added conversion graph.
-			for (TFieldIterator<UFunction> FunctionIter(Class, EFieldIteratorFlags::ExcludeSuper); FunctionIter; ++FunctionIter)
-			{
-				UFunction* Function = *FunctionIter;
-				Function->SetFlags(RF_Transient);
-			}
-		}
-	}
-	else
-	{
-		CurrentCompilerContext->CleanTemporaries(Class);
 	}
 }
 
+
 UMVVMWidgetBlueprintExtension_View::FSearchData UMVVMWidgetBlueprintExtension_View::HandleGatherSearchData(const UBlueprint* OwningBlueprint) const
 {
 	UMVVMWidgetBlueprintExtension_View::FSearchData SearchData;
@@ -199,6 +197,7 @@ UMVVMWidgetBlueprintExtension_View::FSearchData UMVVMWidgetBlueprintExtension_Vi
 	return SearchData;
 }
 
+
 void UMVVMWidgetBlueprintExtension_View::SetFilterSettings(FMVVMViewBindingFilterSettings InFilterSettings)
 {
 	FilterSettings = InFilterSettings;
diff --git a/Engine/Plugins/Runtime/ModelViewViewModel/Source/ModelViewViewModelBlueprint/Public/Bindings/MVVMConversionFunctionHelper.h b/Engine/Plugins/Runtime/ModelViewViewModel/Source/ModelViewViewModelBlueprint/Public/Bindings/MVVMConversionFunctionHelper.h
--- a/Engine/Plugins/Runtime/ModelViewViewModel/Source/ModelViewViewModelBlueprint/Public/Bindings/MVVMConversionFunctionHelper.h
+++ b/Engine/Plugins/Runtime/ModelViewViewModel/Source/ModelViewViewModelBlueprint/Public/Bindings/MVVMConversionFunctionHelper.h
@@ -3,6 +3,8 @@
 #pragma once
 
 #include "Containers/Map.h"
+#include "Templates/Function.h"
+#include "Templates/SubclassOf.h"
 
 struct FMVVMBlueprintPropertyPath;
 struct FMVVMBlueprintViewBinding;
@@ -40,7 +42,10 @@ namespace UE::MVVM::ConversionFunctionHelper
 	MODELVIEWVIEWMODELBLUEPRINT_API FName CreateWrapperName(const FMVVMBlueprintViewBinding& Binding, bool bSourceToDestination);
 
 	/** */
-	MODELVIEWVIEWMODELBLUEPRINT_API TPair<UEdGraph*, UK2Node*> CreateGraph(UBlueprint* WidgetBlueprint, FName GraphName, const UFunction* Function, bool bTransient);
+	MODELVIEWVIEWMODELBLUEPRINT_API TPair<UEdGraph*, UK2Node*> CreateGraph(UBlueprint* WidgetBlueprint, FName GraphName, const UFunction* Signature, const UFunction* FunctionToWrap, bool bTransient);
+	
+	/** */
+	MODELVIEWVIEWMODELBLUEPRINT_API TPair<UEdGraph*, UK2Node*> CreateGraph(UBlueprint* WidgetBlueprint, FName GraphName, const UFunction* Signature, const TSubclassOf<UK2Node> Node, bool bTransient, TFunctionRef<void(UK2Node*)> InitNodeCallback);
 
 	/** Find the conversion function node from the given graph. */
 	MODELVIEWVIEWMODELBLUEPRINT_API UK2Node* GetWrapperNode(UEdGraph* Graph);
diff --git a/Engine/Plugins/Runtime/ModelViewViewModel/Source/ModelViewViewModelBlueprint/Public/MVVMBlueprintView.h b/Engine/Plugins/Runtime/ModelViewViewModel/Source/ModelViewViewModelBlueprint/Public/MVVMBlueprintView.h
--- a/Engine/Plugins/Runtime/ModelViewViewModel/Source/ModelViewViewModelBlueprint/Public/MVVMBlueprintView.h
+++ b/Engine/Plugins/Runtime/ModelViewViewModel/Source/ModelViewViewModelBlueprint/Public/MVVMBlueprintView.h
@@ -8,6 +8,7 @@
 #include "MVVMBlueprintView.generated.h"
 
 class UMVVMWidgetBlueprintExtension_View;
+class UMVVMBlueprintViewEvent;
 
 class UWidget;
 class UWidgetBlueprint;
@@ -116,6 +117,19 @@ public:
 		return Bindings;
 	}
 
+	UMVVMBlueprintViewEvent* AddDefaultEvent();
+	void RemoveEvent(UMVVMBlueprintViewEvent* Event);
+
+	TArrayView<TObjectPtr<UMVVMBlueprintViewEvent>> GetEvents()
+	{
+		return Events;
+	}
+
+	const TArrayView<const TObjectPtr<UMVVMBlueprintViewEvent>> GetEvents() const
+	{
+		return Events;
+	}
+
 	TArray<FText> GetBindingMessages(FGuid Id, UE::MVVM::EBindingMessageType InMessageType) const;
 	bool HasBindingMessage(FGuid Id, UE::MVVM::EBindingMessageType InMessageType) const;
 	void AddMessageToBinding(FGuid Id, UE::MVVM::FBindingMessage MessageToAdd);
@@ -152,6 +166,9 @@ private:
 
 	UPROPERTY(EditAnywhere, Category = "Viewmodel")
 	TArray<FMVVMBlueprintViewBinding> Bindings;
+	
+	UPROPERTY(Instanced, EditAnywhere, Category = "Viewmodel")
+	TArray<TObjectPtr<UMVVMBlueprintViewEvent>> Events;
 
 	UPROPERTY(EditAnywhere, Category = "Viewmodel")
 	TArray<FMVVMBlueprintViewModelContext> AvailableViewModels;
diff --git a/Engine/Plugins/Runtime/ModelViewViewModel/Source/ModelViewViewModelBlueprint/Public/MVVMBlueprintViewEvent.h b/Engine/Plugins/Runtime/ModelViewViewModel/Source/ModelViewViewModelBlueprint/Public/MVVMBlueprintViewEvent.h
new file mode 100644
--- /dev/null
+++ b/Engine/Plugins/Runtime/ModelViewViewModel/Source/ModelViewViewModelBlueprint/Public/MVVMBlueprintViewEvent.h
@@ -0,0 +1,150 @@
+// Copyright Epic Games, Inc. All Rights Reserved.
+
+#pragma once
+
+#include "MVVMPropertyPath.h"
+#include "MVVMBlueprintPin.h"
+
+#include "MVVMBlueprintViewEvent.generated.h"
+
+class UEdGraph;
+class UK2Node;
+class UWidgetBlueprint;
+
+/**
+ *
+ */
+UCLASS(Within = MVVMBlueprintView)
+class MODELVIEWVIEWMODELBLUEPRINT_API UMVVMBlueprintViewEvent : public UObject
+{
+	GENERATED_BODY()
+
+public:
+	enum class EMessageType : uint8
+	{
+		Info,
+		Warning,
+		Error
+	};
+
+	struct FMessage
+	{
+		FText MessageText;
+		EMessageType MessageType;
+	};
+
+	/** @return true if the property path contains a multicast delegate property. */
+	static bool Supports(const UWidgetBlueprint* WidgetBlueprint, const FMVVMBlueprintPropertyPath& PropertyPath);
+
+public:
+	/** Whether the event is enabled or disabled by default. The instance may enable the event at runtime. */
+	UPROPERTY(EditAnywhere, Category = "Viewmodel")
+	bool bEnabled = true;
+
+	/** The event is visible in the editor, but is not compiled and cannot be used at runtime. */
+	UPROPERTY(EditAnywhere, Category = "Viewmodel")
+	bool bCompile = true;
+
+public:
+	const FMVVMBlueprintPropertyPath& GetEventPath() const
+	{
+		return EventPath;
+	}
+	void SetEventPath(FMVVMBlueprintPropertyPath EventPath);
+
+	const FMVVMBlueprintPropertyPath& GetDestinationPath() const
+	{
+		return DestinationPath;
+	}
+	void SetDestinationPath(FMVVMBlueprintPropertyPath DestinationPath);
+
+public:
+	UEdGraph* GetWrapperGraph() const
+	{
+		return CachedWrapperGraph;
+	}
+
+	FName GetWrapperGraphName() const
+	{
+		return GraphName;
+	}
+
+	void RemoveWrapperGraph();
+
+	UK2Node* GetWrapperNode() const
+	{
+		return CachedWrapperNode;
+	}
+
+	UEdGraph* GetOrCreateWrapperGraph();
+
+public:
+	TArrayView<const FMVVMBlueprintPin> GetPins() const
+	{
+		return SavedPins;
+	}
+
+	/** Generates SavedPins from the wrapper graph, if it exists. */
+	void SavePinValues();
+
+	UEdGraphPin* GetOrCreateGraphPin(FName PinName);
+
+	FMVVMBlueprintPropertyPath GetPinPath(FName PinName) const;
+	void SetPinPath(FName PinName, const FMVVMBlueprintPropertyPath& Path);
+
+public:
+	TArray<FText> GetCompilationMessages(EMessageType InMessageType) const;
+	bool HasCompilationMessage(EMessageType InMessageType) const;
+	void AddCompilationToBinding(FMessage MessageToAdd);
+	void ResetCompilationMessages();
+
+public:
+	/** 
+	 * Get a string that identifies this event. 
+	 */
+	FText GetDisplayName(bool bUseDisplayName) const;
+
+	/**
+	 * Get a string that identifies this event and is specifically formatted for search. 
+	 * This includes the display name and variable name of all fields and widgets, as well as all function keywords.
+	 * For use in the UI, use GetDisplayNameString()
+	 */
+	FString GetSearchableString() const;
+
+public:
+	virtual void PostEditChangeChainProperty(FPropertyChangedChainEvent& PropertyChainEvent) override;
+
+private:
+	static const UFunction* GetEventSignature(const UWidgetBlueprint* WidgetBlueprint, const FMVVMBlueprintPropertyPath& PropertyPath);
+	const UFunction* GetEventSignature() const;
+	UWidgetBlueprint* GetWidgetBlueprintInternal() const;
+	UEdGraph* CreateWrapperGraphInternal();
+	UEdGraph* CreateWrapperGraphInternal(const UFunction* Signature, const UFunction* Function);
+	UEdGraph* CreateWrapperGraphInternal(const UFunction* Signature, const FProperty* Property);
+	void LoadPinValuesInternal();
+
+private:
+	UPROPERTY(VisibleAnywhere, Category = "Viewmodel")
+	FMVVMBlueprintPropertyPath EventPath;
+
+	UPROPERTY(VisibleAnywhere, Category = "Viewmodel")
+	FMVVMBlueprintPropertyPath DestinationPath;
+
+	/**
+	 * The pin that are modified and we saved data.
+	 * The data may not be modified. We used the default value of the K2Node in that case.
+	 */
+	UPROPERTY(VisibleAnywhere, Category = "Viewmodel")
+	TArray<FMVVMBlueprintPin> SavedPins;
+
+	UPROPERTY(VisibleAnywhere, Category = "Viewmodel")
+	FName GraphName;
+
+	TArray<FMessage> Messages;
+
+	UPROPERTY(Transient, DuplicateTransient)
+	mutable TObjectPtr<UEdGraph> CachedWrapperGraph;
+
+	UPROPERTY(Transient, DuplicateTransient)
+	mutable TObjectPtr<UK2Node> CachedWrapperNode;
+};
diff --git a/Engine/Plugins/Runtime/ModelViewViewModel/Source/ModelViewViewModelBlueprint/Public/MVVMPropertyPath.h b/Engine/Plugins/Runtime/ModelViewViewModel/Source/ModelViewViewModelBlueprint/Public/MVVMPropertyPath.h
--- a/Engine/Plugins/Runtime/ModelViewViewModel/Source/ModelViewViewModelBlueprint/Public/MVVMPropertyPath.h
+++ b/Engine/Plugins/Runtime/ModelViewViewModel/Source/ModelViewViewModelBlueprint/Public/MVVMPropertyPath.h
@@ -9,6 +9,7 @@
 #include "MVVMPropertyPath.generated.h"
 
 class UBlueprint;
+class UWidgetBlueprint;
 
 /**
  * A single item in a Property Path
@@ -216,6 +217,10 @@ public:
 			}
 		}
 	}
+
+public:
+	FText ToText(const UWidgetBlueprint* Blueprint, bool bUseDisplayName) const;
+	FString ToString(const UWidgetBlueprint* Blueprint, bool bUseDisplayName, bool bIncludeMetaData) const;
 };
 
 template<>
