MVVM: Display ViewEvent in the BindingListView. When the destination path ends with a multicate delagate, transform the binding into an event and viseversa.
Also refactor how filtering works. The element should always exist and only the display should filter them.

diff --git a/Engine/Plugins/Runtime/ModelViewViewModel/Source/ModelViewViewModelEditor/Private/MVVMEditorSubsystem.cpp b/Engine/Plugins/Runtime/ModelViewViewModel/Source/ModelViewViewModelEditor/Private/MVVMEditorSubsystem.cpp
--- a/Engine/Plugins/Runtime/ModelViewViewModel/Source/ModelViewViewModelEditor/Private/MVVMEditorSubsystem.cpp
+++ b/Engine/Plugins/Runtime/ModelViewViewModel/Source/ModelViewViewModelEditor/Private/MVVMEditorSubsystem.cpp
@@ -18,6 +18,7 @@
 #include "K2Node_VariableGet.h"
 #include "MVVMBlueprintView.h"
 #include "MVVMBlueprintViewConversionFunction.h"
+#include "MVVMBlueprintViewEvent.h"
 #include "MVVMDeveloperProjectSettings.h"
 #include "MVVMSubsystem.h"
 #include "MVVMWidgetBlueprintExtension_View.h"
@@ -58,6 +59,32 @@ namespace UE::MVVM::Private
 		FPropertyChangedChainEvent ChainEvent(EditChain, ChangeEvent);
 		BlueprintView->PostEditChangeChainProperty(ChainEvent);
 	}
+	
+	void OnEventPreEditChange(UMVVMBlueprintViewEvent* Event, FName PropertyName)
+	{
+		FProperty* ChangedProperty = UMVVMBlueprintViewEvent::StaticClass()->FindPropertyByName(PropertyName);
+		check(ChangedProperty != nullptr);
+
+		FEditPropertyChain EditChain;
+		EditChain.AddTail(ChangedProperty);
+		EditChain.SetActivePropertyNode(ChangedProperty);
+
+		Event->PreEditChange(EditChain);
+	}
+
+	void OnEventPostEditChange(UMVVMBlueprintViewEvent* Event, FName PropertyName)
+	{
+		FProperty* ChangedProperty = UMVVMBlueprintViewEvent::StaticClass()->FindPropertyByName(PropertyName);
+		check(ChangedProperty != nullptr);
+
+		FEditPropertyChain EditChain;
+		EditChain.AddTail(ChangedProperty);
+		EditChain.SetActivePropertyNode(ChangedProperty);
+
+		FPropertyChangedEvent ChangeEvent(ChangedProperty, EPropertyChangeType::ValueSet);
+		FPropertyChangedChainEvent ChainEvent(EditChain, ChangeEvent);
+		Event->PostEditChangeChainProperty(ChainEvent);
+	}
 
 	UK2Node_FunctionResult* FindFunctionResult(UEdGraph* Graph)
 	{
@@ -236,6 +263,20 @@ void UMVVMEditorSubsystem::RemoveBinding(UWidgetBlueprint* WidgetBlueprint, cons
 	}
 }
 
+UMVVMBlueprintViewEvent* UMVVMEditorSubsystem::AddEvent(UWidgetBlueprint* WidgetBlueprint)
+{
+	UMVVMBlueprintView* View = RequestView(WidgetBlueprint);
+	return View->AddDefaultEvent();
+}
+
+void UMVVMEditorSubsystem::RemoveEvent(UWidgetBlueprint* WidgetBlueprint, UMVVMBlueprintViewEvent* Event)
+{
+	if (UMVVMBlueprintView* View = GetView(WidgetBlueprint))
+	{
+		View->RemoveEvent(Event);
+	}
+}
+
 UFunction* UMVVMEditorSubsystem::GetConversionFunction(const UWidgetBlueprint* WidgetBlueprint, const FMVVMBlueprintViewBinding& Binding, bool bSourceToDestination) const
 {
 	if (UMVVMBlueprintViewConversionFunction* ConversionFunction = Binding.Conversion.GetConversionFunction(bSourceToDestination))
@@ -318,7 +359,7 @@ void UMVVMEditorSubsystem::SetDestinationToSourceConversionFunction(UWidgetBluep
 	}
 }
 
-void UMVVMEditorSubsystem::SetDestinationPathForBinding(UWidgetBlueprint* WidgetBlueprint, FMVVMBlueprintViewBinding& Binding, FMVVMBlueprintPropertyPath Field)
+void UMVVMEditorSubsystem::SetDestinationPathForBinding(UWidgetBlueprint* WidgetBlueprint, FMVVMBlueprintViewBinding& Binding, FMVVMBlueprintPropertyPath PropertyPath)
 {
 	if (UMVVMBlueprintView* View = GetView(WidgetBlueprint))
 	{
@@ -326,7 +367,17 @@ void UMVVMEditorSubsystem::SetDestinationPathForBinding(UWidgetBlueprint* Widget
 
 		UE::MVVM::Private::OnBindingPreEditChange(View, GET_MEMBER_NAME_CHECKED(FMVVMBlueprintViewBinding, DestinationPath));
 
-		Binding.DestinationPath = Field;
+		bool bSupports = UMVVMBlueprintViewEvent::Supports(WidgetBlueprint, PropertyPath);
+		if (bSupports)
+		{
+			UMVVMBlueprintViewEvent* Event = AddEvent(WidgetBlueprint);
+			Event->SetEventPath(PropertyPath);
+			View->RemoveBinding(&Binding);
+		}
+		else
+		{
+			Binding.DestinationPath = PropertyPath;
+		}
 
 		UE::MVVM::Private::OnBindingPostEditChange(View, GET_MEMBER_NAME_CHECKED(FMVVMBlueprintViewBinding, DestinationPath));
 	}
@@ -432,6 +483,102 @@ void UMVVMEditorSubsystem::SetCompileForBinding(UWidgetBlueprint* WidgetBlueprin
 	}
 }
 
+void UMVVMEditorSubsystem::SetEventPath(UMVVMBlueprintViewEvent* Event, FMVVMBlueprintPropertyPath PropertyPath)
+{
+	UMVVMBlueprintView* View = Event ? Event->GetOuterUMVVMBlueprintView() : nullptr;
+	if (View)
+	{
+		UWidgetBlueprint* WidgetBlueprint = Event->GetOuterUMVVMBlueprintView()->GetOuterUMVVMWidgetBlueprintExtension_View()->GetWidgetBlueprint();
+
+		FScopedTransaction Transaction(LOCTEXT("SetEventPath", "Set Event Path"));
+
+		FName EventPath = "EventPath";
+		UE::MVVM::Private::OnEventPreEditChange(Event, EventPath);
+
+		bool bSupports = UMVVMBlueprintViewEvent::Supports(WidgetBlueprint, PropertyPath);
+
+		if (bSupports)
+		{
+			Event->SetEventPath(PropertyPath);
+		}
+		else
+		{
+			FMVVMBlueprintViewBinding& Binding = AddBinding(WidgetBlueprint);
+			SetDestinationPathForBinding(WidgetBlueprint, Binding, PropertyPath);
+			View->RemoveEvent(Event);
+		}
+
+		UE::MVVM::Private::OnEventPostEditChange(Event, EventPath);
+	}
+}
+
+void UMVVMEditorSubsystem::SetEventDestinationPath(UMVVMBlueprintViewEvent* Event, FMVVMBlueprintPropertyPath PropertyPath)
+{
+	const UMVVMBlueprintView* View = Event ? Event->GetOuterUMVVMBlueprintView() : nullptr;
+	if (View)
+	{
+		FScopedTransaction Transaction(LOCTEXT("SetEventDestinationPath", "Set Destination Path"));
+
+		FName DestinationPath = "DestinationPath";
+		UE::MVVM::Private::OnEventPreEditChange(Event, DestinationPath);
+
+		Event->SetDestinationPath(PropertyPath);
+
+		UE::MVVM::Private::OnEventPostEditChange(Event, DestinationPath);
+	}
+}
+
+void UMVVMEditorSubsystem::SetEventArgumentPath(UMVVMBlueprintViewEvent* Event, FName ArgumentName, const FMVVMBlueprintPropertyPath& Path) const
+{
+	const UMVVMBlueprintView* View = Event ? Event->GetOuterUMVVMBlueprintView() : nullptr;
+	if (View)
+	{
+		FScopedTransaction Transaction(LOCTEXT("SetEventPath", "Set Event Path"));
+
+		UE::MVVM::Private::OnEventPreEditChange(Event, "SavedPins");
+	
+		Event->SetPinPath(ArgumentName, Path);
+
+		UE::MVVM::Private::OnEventPostEditChange(Event, "SavedPins");
+	}
+}
+
+void UMVVMEditorSubsystem::SetEnabledForEvent(UMVVMBlueprintViewEvent* Event, bool bEnabled)
+{
+	if (Event->bEnabled != bEnabled)
+	{
+		const UMVVMBlueprintView* View = Event ? Event->GetOuterUMVVMBlueprintView() : nullptr;
+		if (View)
+		{
+			FScopedTransaction Transaction(LOCTEXT("SetBindingEnabled", "Set Binding Enabled"));
+
+			UE::MVVM::Private::OnEventPreEditChange(Event, GET_MEMBER_NAME_CHECKED(UMVVMBlueprintViewEvent, bEnabled));
+
+			Event->bEnabled = bEnabled;
+
+			UE::MVVM::Private::OnEventPostEditChange(Event, GET_MEMBER_NAME_CHECKED(UMVVMBlueprintViewEvent, bEnabled));
+		}
+	}
+}
+
+void UMVVMEditorSubsystem::SetCompileForEvent(UMVVMBlueprintViewEvent* Event, bool bCompile)
+{
+	if (Event->bCompile != bCompile)
+	{
+		const UMVVMBlueprintView* View = Event ? Event->GetOuterUMVVMBlueprintView() : nullptr;
+		if (View)
+		{
+			FScopedTransaction Transaction(LOCTEXT("SetBindingCompiled", "Set Binding Compiled"));
+
+			UE::MVVM::Private::OnEventPreEditChange(Event, GET_MEMBER_NAME_CHECKED(UMVVMBlueprintViewEvent, bCompile));
+
+			Event->bCompile = bCompile;
+
+			UE::MVVM::Private::OnEventPostEditChange(Event, GET_MEMBER_NAME_CHECKED(UMVVMBlueprintViewEvent, bCompile));
+		}
+	}
+}
+
 bool UMVVMEditorSubsystem::IsValidConversionFunction(const UWidgetBlueprint* WidgetBlueprint, const UFunction* Function, const FMVVMBlueprintPropertyPath& Source, const FMVVMBlueprintPropertyPath& Destination) const
 {
 	TValueOrError<const FProperty*, FText> ReturnResult = UE::MVVM::BindingHelper::TryGetReturnTypeForConversionFunction(Function);
diff --git a/Engine/Plugins/Runtime/ModelViewViewModel/Source/ModelViewViewModelEditor/Private/Widgets/SMVVMEventParameter.cpp b/Engine/Plugins/Runtime/ModelViewViewModel/Source/ModelViewViewModelEditor/Private/Widgets/SMVVMEventParameter.cpp
new file mode 100644
--- /dev/null
+++ b/Engine/Plugins/Runtime/ModelViewViewModel/Source/ModelViewViewModelEditor/Private/Widgets/SMVVMEventParameter.cpp
@@ -0,0 +1,198 @@
+// Copyright Epic Games, Inc. All Rights Reserved.
+
+#include "SMVVMEventParameter.h" 
+#include "EdGraphSchema_K2.h"
+#include "Editor.h"
+#include "MVVMBlueprintView.h"
+#include "MVVMBlueprintViewEvent.h"
+#include "MVVMEditorSubsystem.h"
+#include "NodeFactory.h"
+#include "SGraphPin.h"
+#include "Styling/MVVMEditorStyle.h"
+#include "WidgetBlueprint.h"
+#include "Widgets/Images/SImage.h"
+#include "Widgets/SBoxPanel.h"
+#include "Widgets/SMVVMFieldSelector.h"
+#include "Widgets/Input/SCheckBox.h"
+#include "Widgets/Layout/SBox.h"
+#include "Widgets/Layout/SWidgetSwitcher.h"
+#include "Widgets/SOverlay.h"
+#include "Widgets/Text/STextBlock.h"
+
+#define LOCTEXT_NAMESPACE "MVVMEventParameter"
+
+namespace UE::MVVM
+{
+
+void SEventParameter::Construct(const FArguments& InArgs, UWidgetBlueprint* InWidgetBlueprint)
+{
+	WidgetBlueprint = InWidgetBlueprint;
+	check(InWidgetBlueprint != nullptr);
+
+	ViewEvent = InArgs._Event;
+	check(InArgs._Event);
+
+	ParameterName = InArgs._ParameterName;
+	check(!ParameterName.IsNone());
+
+	bAllowDefault = InArgs._AllowDefault;
+
+	UEdGraphPin* Pin = InArgs._Event->GetOrCreateGraphPin(ParameterName);
+	check(Pin);
+
+	bool bIsBooleanPin = Pin->PinType.PinCategory == UEdGraphSchema_K2::PC_Boolean;
+	TSharedRef<SWidget> ValueWidget = SNullWidget::NullWidget;
+	// create a new pin widget so that we can get the default value widget out of it
+	if (TSharedPtr<SGraphPin> PinWidget = FNodeFactory::CreatePinWidget(Pin))
+	{
+		GraphPin = PinWidget;
+		ValueWidget = PinWidget->GetDefaultValueWidget();
+	}
+
+	if (ValueWidget == SNullWidget::NullWidget)
+	{
+		ValueWidget = SNew(STextBlock)
+			.Text(LOCTEXT("DefaultValue", "Default Value"))
+			.TextStyle(FAppStyle::Get(), "HintText");
+	}
+	// booleans are represented by a checkbox which doesn't expand to the minsize we have, so don't put a border around them
+	else if (!bIsBooleanPin)
+	{
+		ValueWidget = SNew(SBorder)
+			.Padding(0.0f)
+			.BorderImage(FMVVMEditorStyle::Get().GetBrush("FunctionParameter.Border"))
+			[
+				ValueWidget
+			];
+	}
+
+	FMVVMBlueprintPropertyPath Path = OnGetSelectedField();
+	bDefaultValueVisible = Path.IsEmpty();
+
+	TSharedPtr<SHorizontalBox> HBox;
+
+	ChildSlot
+	[
+		SNew(SBox)
+		.MinDesiredWidth(bIsBooleanPin ? FOptionalSize() : 100.0f)
+		[
+			SAssignNew(HBox, SHorizontalBox)
+			+ SHorizontalBox::Slot()
+			.Padding(0.0f, 2.0f)
+			.VAlign(VAlign_Center)
+			[
+				SNew(SWidgetSwitcher)
+				.WidgetIndex(this, &SEventParameter::GetCurrentWidgetIndex)
+				+ SWidgetSwitcher::Slot()
+				[
+					ValueWidget
+				]
+				+ SWidgetSwitcher::Slot()
+				[
+					SNew(SFieldSelector, WidgetBlueprint.Get())
+					.OnGetPropertyPath(this, &SEventParameter::OnGetSelectedField)
+					.OnFieldSelectionChanged(this, &SEventParameter::HandleFieldSelectionChanged)
+					.OnGetSelectionContext(this, &SEventParameter::GetSelectedSelectionContext)
+				]
+			]
+		]
+	];
+
+	if (InArgs._AllowDefault)
+	{
+		HBox->AddSlot()
+			.HAlign(HAlign_Right)
+			.VAlign(VAlign_Center)
+			.Padding(8.0f, 0.0f, 0.0f, 0.0f)
+			.AutoWidth()
+			[
+				SNew(SCheckBox)
+				.ToolTipText(LOCTEXT("BindArgument", "Bind this argument to a property."))
+				.Style(FAppStyle::Get(), "ToggleButtonCheckbox")
+				.IsChecked(this, &SEventParameter::OnGetIsBindArgumentChecked)
+				.OnCheckStateChanged(this, &SEventParameter::OnBindArgumentChecked)
+				.Padding(FMargin(4))
+				[
+					SNew(SImage)
+					.DesiredSizeOverride(FVector2D(16.0f, 16.0f))
+					.Image_Lambda([this]()
+						{
+							ECheckBoxState CheckState = OnGetIsBindArgumentChecked();
+							return (CheckState == ECheckBoxState::Checked) ? FAppStyle::GetBrush("Icons.Link") : FAppStyle::GetBrush("Icons.Unlink");
+						})
+				]
+			];		
+	}
+}
+
+int32 SEventParameter::GetCurrentWidgetIndex() const
+{
+	return bDefaultValueVisible && bAllowDefault ? 0 : 1;
+}
+
+ECheckBoxState SEventParameter::OnGetIsBindArgumentChecked() const 
+{
+	return bDefaultValueVisible ? ECheckBoxState::Unchecked : ECheckBoxState::Checked;
+}
+
+void SEventParameter::OnBindArgumentChecked(ECheckBoxState Checked)
+{
+	bDefaultValueVisible = (Checked != ECheckBoxState::Checked);
+
+	if (bDefaultValueVisible)
+	{
+		PreviousSelectedField = OnGetSelectedField();
+		SetSelectedField(FMVVMBlueprintPropertyPath());
+	}
+	else
+	{
+		SetSelectedField(PreviousSelectedField);
+	}
+}
+
+FMVVMBlueprintPropertyPath SEventParameter::OnGetSelectedField() const
+{
+	if (const UMVVMBlueprintViewEvent* EventPtr = ViewEvent.Get())
+	{
+		return EventPtr->GetPinPath(ParameterName);
+	}
+	return FMVVMBlueprintPropertyPath();
+}
+
+void SEventParameter::SetSelectedField(const FMVVMBlueprintPropertyPath& Path)
+{
+	if (UMVVMBlueprintViewEvent* EventPtr = ViewEvent.Get())
+	{
+		const UMVVMEditorSubsystem* Subsystem = GEditor->GetEditorSubsystem<UMVVMEditorSubsystem>();
+		Subsystem->SetEventArgumentPath(EventPtr, ParameterName, Path);
+	}
+}
+
+void SEventParameter::HandleFieldSelectionChanged(FMVVMBlueprintPropertyPath SelectedField, const UFunction* Function)
+{
+	SetSelectedField(SelectedField);
+}
+
+FFieldSelectionContext SEventParameter::GetSelectedSelectionContext() const
+{
+	FFieldSelectionContext Result;
+
+	const UMVVMBlueprintViewEvent* EventPtr = ViewEvent.Get();
+	if (EventPtr == nullptr)
+	{
+		return Result;
+	}
+	
+	Result.BindingMode = EMVVMBindingMode::OneTimeToDestination;
+	//Result.AssignableTo = ConversionFunction->FindPropertyByName(ParameterName);
+	Result.bAllowWidgets = true;
+	Result.bAllowViewModels = true;
+	Result.bAllowConversionFunctions = false;
+	Result.bReadable = true;
+	Result.bWritable = false;
+
+	return Result;
+}
+
+}//namespace
+#undef LOCTEXT_NAMESPACE
diff --git a/Engine/Plugins/Runtime/ModelViewViewModel/Source/ModelViewViewModelEditor/Private/Widgets/SMVVMEventParameter.h b/Engine/Plugins/Runtime/ModelViewViewModel/Source/ModelViewViewModelEditor/Private/Widgets/SMVVMEventParameter.h
new file mode 100644
--- /dev/null
+++ b/Engine/Plugins/Runtime/ModelViewViewModel/Source/ModelViewViewModelEditor/Private/Widgets/SMVVMEventParameter.h
@@ -0,0 +1,58 @@
+// Copyright Epic Games, Inc. All Rights Reserved.
+
+#pragma once
+
+#include "MVVMPropertyPath.h"
+#include "Widgets/SCompoundWidget.h"
+#include "Widgets/SMVVMFieldSelectorMenu.h"
+
+class SGraphPin;
+class UK2Node_CallFunction;
+class UWidgetBlueprint;
+
+class UMVVMBlueprintViewEvent;
+
+enum class ECheckBoxState : uint8;
+
+namespace UE::MVVM
+{
+
+class SFieldSelector;
+
+class SEventParameter : public SCompoundWidget
+{
+public:
+	SLATE_BEGIN_ARGS(SEventParameter) {}
+		SLATE_ARGUMENT_DEFAULT(UMVVMBlueprintViewEvent*, Event) =  nullptr;
+		SLATE_ARGUMENT(FName, ParameterName)
+		SLATE_ARGUMENT_DEFAULT(bool, AllowDefault) = true;
+	SLATE_END_ARGS()
+
+	void Construct(const FArguments& InArgs, UWidgetBlueprint* WidgetBlueprint);
+
+private:
+	ECheckBoxState OnGetIsBindArgumentChecked() const;
+	void OnBindArgumentChecked(ECheckBoxState Checked);
+
+	FMVVMBlueprintPropertyPath OnGetSelectedField() const;
+	void SetSelectedField(const FMVVMBlueprintPropertyPath& Path);
+
+	void HandleFieldSelectionChanged(FMVVMBlueprintPropertyPath SelectedField, const UFunction* Function);
+	FFieldSelectionContext GetSelectedSelectionContext() const;
+
+	int32 GetCurrentWidgetIndex() const;
+
+private:
+	TWeakObjectPtr<UWidgetBlueprint> WidgetBlueprint;
+	TWeakObjectPtr<UMVVMBlueprintViewEvent> ViewEvent;
+	FName ParameterName;
+	/** This reference is just to keep the default value widget alive. */
+	TSharedPtr<SGraphPin> GraphPin;
+
+	FMVVMBlueprintPropertyPath PreviousSelectedField;
+
+	bool bAllowDefault = true;
+	bool bDefaultValueVisible = true;
+};
+
+}
diff --git a/Engine/Plugins/Runtime/ModelViewViewModel/Source/ModelViewViewModelEditor/Private/Widgets/SMVVMViewBindingListView.cpp b/Engine/Plugins/Runtime/ModelViewViewModel/Source/ModelViewViewModelEditor/Private/Widgets/SMVVMViewBindingListView.cpp
--- a/Engine/Plugins/Runtime/ModelViewViewModel/Source/ModelViewViewModelEditor/Private/Widgets/SMVVMViewBindingListView.cpp
+++ b/Engine/Plugins/Runtime/ModelViewViewModel/Source/ModelViewViewModelEditor/Private/Widgets/SMVVMViewBindingListView.cpp
@@ -11,17 +11,23 @@
 #include "Framework/MultiBox/MultiBoxBuilder.h"
 #include "Framework/Views/TableViewMetadata.h"
 #include "Misc/MessageDialog.h"
+
 #include "MVVMBlueprintView.h"
+#include "MVVMBlueprintViewBinding.h"
 #include "MVVMBlueprintViewConversionFunction.h"
+#include "MVVMBlueprintViewEvent.h"
 #include "MVVMEditorSubsystem.h"
 #include "MVVMDeveloperProjectSettings.h"
 #include "MVVMWidgetBlueprintExtension_View.h"
+
 #include "ScopedTransaction.h"
 #include "SEnumCombo.h"
 #include "SSimpleButton.h" 
 #include "Styling/MVVMEditorStyle.h"
 #include "Styling/StyleColors.h"
 #include "ViewModelFieldDragDropOp.h"
+
+#include "Widgets/SMVVMEventParameter.h"
 #include "Widgets/SMVVMFieldEntry.h"
 #include "Widgets/SMVVMFieldSelector.h"
 #include "Widgets/SMVVMFunctionParameter.h"
@@ -40,7 +46,9 @@
 namespace UE::MVVM
 {
 
-// a wrapper around either a widget row or a binding row
+/**
+ * A wrapper around either a widget row or a binding row
+ */
 struct FBindingEntry
 {
 	enum class ERowType
@@ -48,7 +56,9 @@ struct FBindingEntry
 		None,
 		Group,
 		Binding,
-		Parameter
+		BindingParameter,
+		Event,
+		EventParameter,
 	};
 
 	FMVVMBlueprintViewBinding* GetBinding(UMVVMBlueprintView* View) const
@@ -66,6 +76,45 @@ struct FBindingEntry
 		return RowType;
 	}
 
+	FName GetGroupName() const
+	{
+		return Name;
+	}
+
+	FGuid GetGroupAsViewModel() const
+	{
+		return BindingId;
+	}
+
+	bool IsGroupWidget() const
+	{
+		return bGroupIsWidget;
+	}
+
+	void SetGroup(FName WidgetName)
+	{
+		check(RowType == ERowType::None);
+		RowType = ERowType::Group;
+		Name = WidgetName;
+		bGroupIsWidget = true;
+	}
+
+	void SetGroup(FName ViewModelName, FGuid ViewModelId)
+	{
+		if (!ViewModelId.IsValid())
+		{
+			SetGroup(ViewModelName);
+		}
+		else
+		{
+			check(RowType == ERowType::None);
+			RowType = ERowType::Group;
+			Name = ViewModelName;
+			BindingId = ViewModelId;
+			bGroupIsWidget = false;
+		}
+	}
+
 	FGuid GetBindingId() const
 	{
 		return BindingId;
@@ -78,44 +127,75 @@ struct FBindingEntry
 		BindingId = Id;
 	}
 
-	void SetGroupName(FName GroupName)
+	FName GetBindingParameterName() const
 	{
-		check(RowType == ERowType::None);
-		RowType = ERowType::Group;
-		Name = GroupName;
+		return Name;
 	}
 
-	void SetParameterName(FGuid Id, FName ParameterName)
+	void SetBindingParameterName(FGuid Id, FName ParameterName)
 	{
 		check(RowType == ERowType::None);
-		RowType = ERowType::Parameter;
+		RowType = ERowType::BindingParameter;
 		BindingId = Id;
 		Name = ParameterName;
 	}
 
-	FName GetName() const
+	UMVVMBlueprintViewEvent* GetEvent() const
+	{
+		return Event.Get();
+	}
+
+	void SetEvent(UMVVMBlueprintViewEvent* InEvent)
+	{
+		check(RowType == ERowType::None);
+		RowType = ERowType::Event;
+		Event = InEvent;
+	}
+
+	FName GetEventParameterName() const
 	{
 		return Name;
 	}
 
-	TConstArrayView<TSharedPtr<FBindingEntry>> GetChildren() const
+	void SetEventParameterName(UMVVMBlueprintViewEvent* InEvent, FName ParameterName)
+	{
+		check(RowType == ERowType::None);
+		RowType = ERowType::EventParameter;
+		Event = InEvent;
+		Name = ParameterName;
+	}
+
+	TConstArrayView<TSharedPtr<FBindingEntry>> GetAllChildren() const
+	{
+		return AllChildren;
+	}
+
+	TConstArrayView<TSharedPtr<FBindingEntry>> GetFilteredChildren() const
 	{
-		return Children;
+		return bUseFilteredChildren ? FilteredChildren : AllChildren;
 	}
 
-	void AddChild(TSharedRef<FBindingEntry> Child)
+	void AddChild(TSharedPtr<FBindingEntry> Child)
 	{
-		Children.Add(Child);
+		AllChildren.Add(Child);
 	}
 
-	void RemoveChildAt(int32 IndexToRemove)
+	void AddFilteredChild(TSharedPtr<FBindingEntry> Child)
 	{
-		Children.RemoveAt(IndexToRemove);
+		FilteredChildren.Add(Child);
+		bUseFilteredChildren = true;
 	}
 
 	void ResetChildren()
 	{
-		Children.Reset();
+		AllChildren.Reset();
+		FilteredChildren.Reset();
+		bUseFilteredChildren = false;
+	}
+
+	void SetUseFilteredChildList()
+	{
+		bUseFilteredChildren = true;
 	}
 
 	bool operator==(const FBindingEntry& Other) const
@@ -129,18 +209,26 @@ struct FBindingEntry
 	{
 		FString RowToString;
 		FString FunctionKeywords;
-		
-		FMVVMBlueprintViewBinding* BindingInRow; // Initialized and used only when RowType is Binding.
 
 		switch (RowType)
 		{
 		case UE::MVVM::FBindingEntry::ERowType::Group:
-		case UE::MVVM::FBindingEntry::ERowType::Parameter:
+		case UE::MVVM::FBindingEntry::ERowType::BindingParameter:
+		case UE::MVVM::FBindingEntry::ERowType::EventParameter:
 			RowToString = Name.ToString();
 			break;
 		case UE::MVVM::FBindingEntry::ERowType::Binding:
-			BindingInRow = GetBinding(View);
-			RowToString.Append(BindingInRow->GetSearchableString(WidgetBP));
+			{
+				FMVVMBlueprintViewBinding* BindingInRow = GetBinding(View);
+				check(BindingInRow);
+				RowToString.Append(BindingInRow->GetSearchableString(WidgetBP));
+			}
+			break;
+		case UE::MVVM::FBindingEntry::ERowType::Event:
+			if (UMVVMBlueprintViewEvent* EventInRow = GetEvent())
+			{
+				RowToString.Append(EventInRow->GetSearchableString());
+			}
 			break;
 		default:
 			break;
@@ -154,7 +242,11 @@ private:
 	ERowType RowType = ERowType::None;
 	FName Name;
 	FGuid BindingId;
-	TArray<TSharedPtr<FBindingEntry>> Children;
+	TWeakObjectPtr<UMVVMBlueprintViewEvent> Event;
+	TArray<TSharedPtr<FBindingEntry>> AllChildren;
+	TArray<TSharedPtr<FBindingEntry>> FilteredChildren;
+	bool bGroupIsWidget = false;
+	bool bUseFilteredChildren = false;
 };
 
 namespace Private
@@ -186,7 +278,7 @@ namespace Private
 	{
 		TreeView->SetItemExpansion(Entry, true);
 
-		for (const TSharedPtr<FBindingEntry>& Child : Entry->GetChildren())
+		for (const TSharedPtr<FBindingEntry>& Child : Entry->GetFilteredChildren())
 		{
 			ExpandAll(TreeView, Child);
 		}
@@ -200,7 +292,24 @@ namespace Private
 			{
 				return Entry;
 			}
-			TSharedPtr<FBindingEntry> Result = FindBinding(BindingId, Entry->GetChildren());
+			TSharedPtr<FBindingEntry> Result = FindBinding(BindingId, Entry->GetAllChildren());
+			if (Result)
+			{
+				return Result;
+			}
+		}
+		return TSharedPtr<FBindingEntry>();
+	}
+	
+	TSharedPtr<FBindingEntry> FindEvent(UMVVMBlueprintViewEvent* Event, TConstArrayView<TSharedPtr<FBindingEntry>> Entries)
+	{
+		for (const TSharedPtr<FBindingEntry>& Entry : Entries)
+		{
+			if (Entry->GetRowType() == FBindingEntry::ERowType::Event && Entry->GetEvent() == Event)
+			{
+				return Entry;
+			}
+			TSharedPtr<FBindingEntry> Result = FindEvent(Event, Entry->GetAllChildren());
 			if (Result)
 			{
 				return Result;
@@ -209,77 +318,93 @@ namespace Private
 		return TSharedPtr<FBindingEntry>();
 	}
 
-	void FilterBindingsList(FString FilterString, TArray<TSharedPtr<FBindingEntry>>& RootGroups, UMVVMBlueprintView* BlueprintView, UMVVMWidgetBlueprintExtension_View* MVVMExtensionPtr)
+	void FilterEntryList(FString FilterString, const TArray<TSharedPtr<FBindingEntry>>& RootGroups, TArray<TSharedPtr<FBindingEntry>>& FilteredRootGroups, UMVVMBlueprintView* BlueprintView, UMVVMWidgetBlueprintExtension_View* MVVMExtensionPtr)
 	{
 		if (!FilterString.TrimStartAndEnd().IsEmpty())
 		{
 			TArray<FString> SearchKeywords;
 			FilterString.ParseIntoArray(SearchKeywords, TEXT(" "));
-			FString EntryString;
 
-			auto IsAllKeywordsInString = [](FString EntryString, TArray<FString>& SearchKeywords) -> bool {
-				for (const FString& Keyword : SearchKeywords)
+			struct FIsAllKeywordsInString
+			{
+				bool operator()(FString EntryString, TArray<FString>& SearchKeywords)
+				{
+					for (const FString& Keyword : SearchKeywords)
+					{
+						if (!EntryString.Contains(Keyword))
+						{
+							return false;
+						}
+					}
+					return true;
+				}
+			} IsAllKeywordsInString;
+
+			struct FAddFilteredEntry
+			{
+				FAddFilteredEntry(FIsAllKeywordsInString& InIsAllKeywordsInString, TArray<FString>& InSearchKeywords, UMVVMBlueprintView* InBlueprintView, UWidgetBlueprint* InWidgetBlueprint)
+					: IsAllKeywordsInString(InIsAllKeywordsInString)
+					, SearchKeywords(InSearchKeywords)
+					, WidgetBlueprint(InWidgetBlueprint)
+					, BlueprintView(InBlueprintView)
+				{}
+				void operator()(TSharedPtr<FBindingEntry> ParentEntry)
 				{
-					if (!EntryString.Contains(Keyword))
+					for (TSharedPtr<FBindingEntry> Entry : ParentEntry->GetAllChildren())
 					{
-						return false;
+						FString EntryString = Entry->GetSearchNameString(BlueprintView, WidgetBlueprint);
+						if (IsAllKeywordsInString(EntryString, SearchKeywords))
+						{
+							// If the filter text is found in the group name, we keep the entire group.
+							ParentEntry->AddFilteredChild(Entry);
+						}
+						else
+						{
+							(*this)(Entry);
+						}
 					}
+					ParentEntry->SetUseFilteredChildList();
 				}
-				return true;
+				FIsAllKeywordsInString& IsAllKeywordsInString;
+				TArray<FString>& SearchKeywords;
+				UWidgetBlueprint* WidgetBlueprint = nullptr;
+				UMVVMBlueprintView* BlueprintView = nullptr;
 			};
+			FAddFilteredEntry AddFilteredEntry = FAddFilteredEntry(IsAllKeywordsInString, SearchKeywords, BlueprintView, MVVMExtensionPtr->GetWidgetBlueprint());
 
-			for (int32 GroupEntryIndex = RootGroups.Num() - 1; GroupEntryIndex >= 0; GroupEntryIndex--)
+			for (const TSharedPtr<FBindingEntry>& GroupEntry : RootGroups)
 			{
-				TSharedPtr<FBindingEntry> GroupEntry = RootGroups[GroupEntryIndex];
-				EntryString = GroupEntry->GetSearchNameString(BlueprintView, MVVMExtensionPtr->GetWidgetBlueprint());
+				FString EntryString = GroupEntry->GetSearchNameString(BlueprintView, MVVMExtensionPtr->GetWidgetBlueprint());
 
 				// If the filter text is found in the group name, we keep the entire group.
 				if (IsAllKeywordsInString(EntryString, SearchKeywords))
 				{
-					continue;
+					FilteredRootGroups.Add(GroupEntry);
 				}
-				for (int32 BindingEntryIndex = GroupEntry->GetChildren().Num() - 1; BindingEntryIndex >= 0; BindingEntryIndex--)
+				else
 				{
-					TSharedPtr<FBindingEntry> BindingEntry = GroupEntry->GetChildren()[BindingEntryIndex];
-					EntryString = BindingEntry->GetSearchNameString(BlueprintView, MVVMExtensionPtr->GetWidgetBlueprint());
-					
-					// If the filter text is found in the binding string, we keep the entire binding.
-					if (IsAllKeywordsInString(EntryString, SearchKeywords))
-					{
-						continue;
-					}
-
-					bool IsParameterSearched = false;
-					for (TSharedPtr<FBindingEntry> ParameterEntry : BindingEntry->GetChildren())
-					{
-						EntryString = ParameterEntry->GetSearchNameString(BlueprintView, MVVMExtensionPtr->GetWidgetBlueprint());
-
-						// If the filter text is found in any of the parameter names, it is sufficient to keep the owner binding.
-						if (IsAllKeywordsInString(EntryString, SearchKeywords))
-						{
-							IsParameterSearched = true;
-							break;
-						}
-
-					}
-					if (!IsParameterSearched)
+					AddFilteredEntry(GroupEntry);
+					if (GroupEntry->GetFilteredChildren().Num() > 0)
 					{
-						GroupEntry->RemoveChildAt(BindingEntryIndex);
+						FilteredRootGroups.Add(GroupEntry);
 					}
 				}
-				if (GroupEntry->GetChildren().Num() == 0)
-				{
-					RootGroups.RemoveAt(GroupEntryIndex);
-				}
 			}
 		}
+		else
+		{
+			FilteredRootGroups = RootGroups;
+		}
 	}
-}
+} // namespace
 
-class SWidgetRow : public STableRow<TSharedPtr<FBindingEntry>>
+/**
+ * 
+ */
+class SGroupRow : public STableRow<TSharedPtr<FBindingEntry>>
 {
 public:
-	SLATE_BEGIN_ARGS(SWidgetRow) {}
+	SLATE_BEGIN_ARGS(SGroupRow) {}
 	SLATE_END_ARGS()
 
 	void Construct(const FArguments& Args, const TSharedRef<STableViewBase>& OwnerTableView, const TSharedPtr<FBindingEntry>& InEntry, UWidgetBlueprint* InWidgetBlueprint)
@@ -306,8 +431,8 @@ public:
 						.ShowClear(false)
 						.AutoRefresh(true)
 						.ViewModels(false)
-						.SelectedBindingSource(this, &SWidgetRow::GetSelectedWidget)
-						.OnSelectionChanged(this, &SWidgetRow::SetSelectedWidget)
+						.SelectedBindingSource(this, &SGroupRow::GetSelectedWidget)
+						.OnSelectionChanged(this, &SGroupRow::SetSelectedWidget)
 					]
 				]
 				+ SHorizontalBox::Slot()
@@ -317,10 +442,9 @@ public:
 				[
 					SNew(SSimpleButton)
 					.Icon(FAppStyle::Get().GetBrush("Icons.Plus"))
-					.IsEnabled_Lambda([this]() { return !Entry->GetName().IsNone(); })
-					.OnClicked(this, &SWidgetRow::AddBinding)
-
-				]
+					.IsEnabled_Lambda([this]() { return !Entry->GetGroupName().IsNone(); })
+					.OnClicked(this, &SGroupRow::AddBinding)
+				] 
 			],
 			OwnerTableView
 		);
@@ -335,7 +459,7 @@ public:
 			[
 				/* Restore the border that we're meant to have that reacts to selection/hover/etc. */
 				SNew(SBorder)
-				.BorderImage(this, &SWidgetRow::GetBorderImage)
+				.BorderImage(this, &SGroupRow::GetBorderImage)
 				.Padding(0)
 				[
 					ChildContent.ToSharedRef()
@@ -345,10 +469,11 @@ public:
 	}
 
 private:
-
 	FBindingSource GetSelectedWidget() const
 	{
-		return FBindingSource::CreateForWidget(WidgetBlueprintWeak.Get(), Entry->GetName());
+		return Entry->IsGroupWidget()
+			? FBindingSource::CreateForWidget(WidgetBlueprintWeak.Get(), Entry->GetGroupName())
+			: FBindingSource::CreateForViewModel(WidgetBlueprintWeak.Get(), Entry->GetGroupAsViewModel());
 	}
 
 	void SetSelectedWidget(FBindingSource Source)
@@ -358,14 +483,28 @@ private:
 			UMVVMEditorSubsystem* EditorSubsystem = GEditor->GetEditorSubsystem<UMVVMEditorSubsystem>();
 			UMVVMBlueprintView* View = EditorSubsystem->GetView(WidgetBlueprint);
 
-			for (const TSharedPtr<FBindingEntry>& ChildEntry : Entry->GetChildren())
+			for (const TSharedPtr<FBindingEntry>& ChildEntry : Entry->GetAllChildren())
 			{
-				if (FMVVMBlueprintViewBinding* Binding = ChildEntry->GetBinding(View))
+				if (ChildEntry->GetRowType() == FBindingEntry::ERowType::Binding)
+				{
+					if (FMVVMBlueprintViewBinding* Binding = ChildEntry->GetBinding(View))
+					{
+						FMVVMBlueprintPropertyPath CurrentPath = Binding->DestinationPath;
+						CurrentPath.SetWidgetName(Source.Name);
+
+						EditorSubsystem->SetDestinationPathForBinding(WidgetBlueprint, *Binding, CurrentPath);
+					}
+				}
+
+				if (ChildEntry->GetRowType() == FBindingEntry::ERowType::Event)
 				{
-					FMVVMBlueprintPropertyPath CurrentPath = Binding->DestinationPath;
-					CurrentPath.SetWidgetName(Source.Name);
+					if (UMVVMBlueprintViewEvent* Event = ChildEntry->GetEvent())
+					{
+						FMVVMBlueprintPropertyPath CurrentPath = Event->GetEventPath();
+						CurrentPath.SetWidgetName(Source.Name);
 
-					EditorSubsystem->SetDestinationPathForBinding(WidgetBlueprint, *Binding, CurrentPath);
+						EditorSubsystem->SetEventPath(Event, CurrentPath);
+					}
 				}
 			}
 		}
@@ -378,7 +517,15 @@ private:
 			UMVVMEditorSubsystem* EditorSubsystem = GEditor->GetEditorSubsystem<UMVVMEditorSubsystem>();
 			FMVVMBlueprintViewBinding& Binding = EditorSubsystem->AddBinding(WidgetBlueprint);
 			FMVVMBlueprintPropertyPath Path;
-			Path.SetWidgetName(Entry->GetName());
+			if (Entry->IsGroupWidget())
+			{
+				Path.SetWidgetName(Entry->GetGroupName());
+			}
+			else
+			{
+				Path.SetViewModelId(Entry->GetGroupAsViewModel());
+			}
+			
 			EditorSubsystem->SetDestinationPathForBinding(WidgetBlueprint, Binding, Path);
 		}
 
@@ -390,6 +537,9 @@ private:
 	TWeakObjectPtr<UWidgetBlueprint> WidgetBlueprintWeak;
 };
 
+/**
+ * 
+ */
 class SBindingRow : public STableRow<TSharedPtr<FBindingEntry>>
 {
 public:
@@ -578,16 +728,6 @@ public:
 		return nullptr;
 	}
 
-	TArray<FMVVMBlueprintViewBinding*> GetThisViewBindingAsArray() const
-	{
-		TArray<FMVVMBlueprintViewBinding*> Result;
-		if (FMVVMBlueprintViewBinding* ViewBinding = GetThisViewBinding())
-		{
-			Result.Add(ViewBinding);
-		}
-		return Result;
-	}
-
 private:
 
 	FSlateColor GetErrorBorderColor() const
@@ -1178,128 +1318,593 @@ private:
 	IConsoleVariable* CVarDefaultExecutionMode = nullptr;
 };
 
-class SFunctionParameterRow : public STableRow<TSharedPtr<FBindingEntry>>
+/**
+*
+*/
+class SEventRow : public STableRow<TSharedPtr<FBindingEntry>>
 {
-	SLATE_BEGIN_ARGS(SFunctionParameterRow) {}
+public:
+	SLATE_BEGIN_ARGS(SEventRow) {}
 	SLATE_END_ARGS()
 
 	void Construct(const FArguments& Args, const TSharedRef<STableViewBase>& OwnerTableView, const TSharedPtr<FBindingEntry>& InEntry, UWidgetBlueprint* InWidgetBlueprint)
 	{
 		Entry = InEntry;
-		check(Entry->GetRowType() == FBindingEntry::ERowType::Parameter);
-
-		WidgetBlueprint = InWidgetBlueprint;
-
-		UMVVMEditorSubsystem* EditorSubsystem = GEditor->GetEditorSubsystem<UMVVMEditorSubsystem>();
-		UMVVMBlueprintView* View = EditorSubsystem->GetView(InWidgetBlueprint);
-		FMVVMBlueprintViewBinding* Binding = Entry->GetBinding(View);
-		const bool bSourceToDestination = UE::MVVM::IsForwardBinding(Binding->BindingType);
-
-		FSlateColor PrimaryColor, SecondaryColor;
-		const FSlateBrush* PrimaryBrush = nullptr;
-		const FSlateBrush* SecondaryBrush = nullptr;
-		FText DisplayName, ToolTip;
-
-		bool bSimpleConversionFunction = false;
-		
-		if (UEdGraphPin* Pin = EditorSubsystem->GetConversionFunctionArgumentPin(InWidgetBlueprint, *Binding, Entry->GetName(), bSourceToDestination))
-		{
-			PrimaryBrush = FBlueprintEditor::GetVarIconAndColorFromPinType(Pin->PinType, PrimaryColor, SecondaryBrush, SecondaryColor);
-			DisplayName = Pin->GetDisplayName();
-			ToolTip = FText::FromString(Pin->PinToolTip);
-		}
-		else if (const UFunction* Function = EditorSubsystem->GetConversionFunction(InWidgetBlueprint, *Binding, bSourceToDestination))
-		{
-			// no wrapper graph, this is a simple conversion function of the form: int32 Convert(float x)
-			if (const FProperty* Argument = BindingHelper::GetFirstArgumentProperty(Function))
-			{
-				bSimpleConversionFunction = true;
-
-				PrimaryBrush = FBlueprintEditor::GetVarIconAndColorFromProperty(Argument, PrimaryColor, SecondaryBrush, SecondaryColor);
-				DisplayName = Argument->GetDisplayNameText();
-				ToolTip = Argument->GetToolTipText();
-			}
-		}
+		ViewEvent = Entry->GetEvent();
+		WidgetBlueprintWeak = InWidgetBlueprint;
 
 		STableRow<TSharedPtr<FBindingEntry>>::Construct(
 			STableRow<TSharedPtr<FBindingEntry>>::FArguments()
 			.ShowWires(true)
-			.Style(FMVVMEditorStyle::Get(), "BindingView.ParameterRow")
+			.Style(FMVVMEditorStyle::Get(), "BindingView.BindingRow")
 			[
-				SNew(SBox)
-				.HeightOverride(30)
-				.ToolTipText(ToolTip)
+				SNew(SBorder)
+				.BorderImage(FAppStyle::Get().GetBrush("PlainBorder"))
+				.Padding(0.0f)
+				.BorderBackgroundColor(this, &SEventRow::GetErrorBorderColor)
 				[
-					SNew(SHorizontalBox)
-
-					+ SHorizontalBox::Slot()
-					.HAlign(HAlign_Left)
-					.VAlign(VAlign_Center)
-					.Padding(0, 0, 8, 0)
-					.AutoWidth()
+					SNew(SBox)
+					.HeightOverride(30)
 					[
 						SNew(SHorizontalBox)
 						+ SHorizontalBox::Slot()
+						.Padding(2.0f, 0.0f)
 						.VAlign(VAlign_Center)
-						.Padding(0, 0, 8, 0)
+						.HAlign(HAlign_Left)
 						.AutoWidth()
 						[
-							SNew(SImage)
-							.Image(PrimaryBrush)
-							.ColorAndOpacity(PrimaryColor)
-							.DesiredSizeOverride(FVector2D(16, 16))
+							SNew(SCheckBox)
+							.IsChecked(this, &SEventRow::IsEventCompiled)
+							.OnCheckStateChanged(this, &SEventRow::OnIsEventCompileChanged)
 						]
+
 						+ SHorizontalBox::Slot()
+						.Padding(2.0f, 0.0f)
+						.VAlign(VAlign_Center)
+						.HAlign(HAlign_Left)
+						.AutoWidth()
 						[
-							SNew(STextBlock)
-							.Text(DisplayName)
+							SNew(SSimpleButton)
+							.Icon(FAppStyle::Get().GetBrush("Icons.Error"))
+							.Visibility(this, &SEventRow::GetErrorButtonVisibility)
+							.ToolTipText(this, &SEventRow::GetErrorButtonToolTip)
+							.OnClicked(this, &SEventRow::OnErrorButtonClicked)
+						]
+
+						+ SHorizontalBox::Slot()
+						.Padding(4.0f, 0.0f)
+						.VAlign(VAlign_Center)
+						.HAlign(HAlign_Left)
+						.AutoWidth()
+						[
+							SNew(SBox)
+							.MinDesiredWidth(150.0f)
+							[
+								SNew(SFieldSelector, InWidgetBlueprint)
+								.OnGetPropertyPath(this, &SEventRow::GetSelectedPropertyPath, true)
+								.OnGetConversionFunction(this, &SEventRow::GetSelectedConversionFunction, true)
+								.OnFieldSelectionChanged(this, &SEventRow::HandleFieldSelectionChanged, true)
+								.OnGetSelectionContext(this, &SEventRow::GetSelectedSelectionContext, true)
+								.OnDrop(this, &SEventRow::HandleFieldSelectorDrop, true)
+								.OnDragEnter(this, &SEventRow::HandleFieldSelectorDragEnter, true)
+								.ShowContext(false)
+							]
+						]
+
+						+ SHorizontalBox::Slot()
+						.Padding(4.0f, 0.0f)
+						.VAlign(VAlign_Fill)
+						.HAlign(HAlign_Left)
+						.AutoWidth()
+						[
+							SNew(SBox)
+							.HAlign(HAlign_Center)
+							.VAlign(VAlign_Center)
+							.WidthOverride(16.0f)
+							.HeightOverride(16.0f)
+							[
+								SNew(SImage)
+								.Image(FMVVMEditorStyle::Get().GetBrush("BindingMode.OneWay"))
+							]
+						]
+
+						+ SHorizontalBox::Slot()
+						.Padding(4.0f, 0.0f, 2.0f, 0.0f)
+						.VAlign(VAlign_Center)
+						.HAlign(HAlign_Left)
+						.AutoWidth()
+						[
+							SNew(SBox)
+							.MinDesiredWidth(150.0f)
+							[
+								SNew(SFieldSelector, InWidgetBlueprint)
+								.OnGetPropertyPath(this, &SEventRow::GetSelectedPropertyPath, false)
+								.OnGetConversionFunction(this, &SEventRow::GetSelectedConversionFunction, false)
+								.OnFieldSelectionChanged(this, &SEventRow::HandleFieldSelectionChanged, false)
+								.OnGetSelectionContext(this, &SEventRow::GetSelectedSelectionContext, false)
+								.OnDrop(this, &SEventRow::HandleFieldSelectorDrop, false)
+								.OnDragEnter(this, &SEventRow::HandleFieldSelectorDragEnter, false)
+							]
 						]
-					]
-					+ SHorizontalBox::Slot()
-					.HAlign(HAlign_Left)
-					.VAlign(VAlign_Center)
-					.AutoWidth()
-					[
-						SNew(SFunctionParameter, InWidgetBlueprint)
-						.BindingId(Binding->BindingId)
-						.ParameterName(Entry->GetName())
-						.SourceToDestination(bSourceToDestination)
-						.AllowDefault(!bSimpleConversionFunction)
 					]
 				]
-			], 
-			OwnerTableView
-		);
+			],
+			OwnerTableView);
+
+		TSharedPtr<SWidget> ChildContent = ChildSlot.DetachWidget();
+		ChildSlot
+		[
+			/* Add a single pixel top and bottom border for this widget. */
+			SNew(SBorder)
+			.BorderImage(FAppStyle::Get().GetBrush("Brushes.Recessed"))
+			.Padding(0.0f, 2.0f, 0.0f, 1.0f)
+			[
+			/* Restore the border that we're meant to have that reacts to selection/hover/etc. */
+				SNew(SBorder)
+				.BorderImage(this, &SEventRow::GetBorderImage)
+				.Padding(0.0f)
+				[
+					ChildContent.ToSharedRef()
+				]
+			]
+		];
 	}
 
 private:
-
-	EMVVMBindingMode OnGetBindingMode() const
+	FSlateColor GetErrorBorderColor() const
 	{
-		if (UWidgetBlueprint* WidgetBlueprintPtr = WidgetBlueprint.Get())
+		if (const UMVVMBlueprintViewEvent* Event = ViewEvent.Get())
 		{
-			UMVVMEditorSubsystem* EditorSubsystem = GEditor->GetEditorSubsystem<UMVVMEditorSubsystem>();
-			UMVVMBlueprintView* View = EditorSubsystem->GetView(WidgetBlueprintPtr);
-
-			if (FMVVMBlueprintViewBinding* Binding = Entry->GetBinding(View))
+			if (Event->HasCompilationMessage(UMVVMBlueprintViewEvent::EMessageType::Error))
+			{
+				return FStyleColors::Error;
+			}
+			else if (Event->HasCompilationMessage(UMVVMBlueprintViewEvent::EMessageType::Warning))
 			{
-				return Binding->BindingType;
+				return FStyleColors::Warning;
 			}
 		}
-		return EMVVMBindingMode::OneWayToDestination;
+		return FStyleColors::Transparent;
 	}
 
-private:
-	TWeakObjectPtr<UWidgetBlueprint> WidgetBlueprint;
-	TSharedPtr<FBindingEntry> Entry;
-};
+	EVisibility GetErrorButtonVisibility() const
+	{
+		if (const UMVVMBlueprintViewEvent* Event = ViewEvent.Get())
+		{
+			bool bHasBindingError = Event->HasCompilationMessage(UMVVMBlueprintViewEvent::EMessageType::Error);
+			bool bHasBindingWarning = Event->HasCompilationMessage(UMVVMBlueprintViewEvent::EMessageType::Warning);
+			return bHasBindingError || bHasBindingWarning ? EVisibility::Visible : EVisibility::Collapsed;
+		}
+		return EVisibility::Collapsed;
+	}
 
-void SBindingsList::Construct(const FArguments& InArgs, TSharedPtr<SBindingsPanel> Owner, UMVVMWidgetBlueprintExtension_View* InMVVMExtension)
-{
-	BindingPanel = Owner;
-	MVVMExtension = InMVVMExtension;
-	check(InMVVMExtension);
-	check(InMVVMExtension->GetBlueprintView());
+	FText GetErrorButtonToolTip() const
+	{
+		if (const UMVVMBlueprintViewEvent* Event = ViewEvent.Get())
+		{
+			TArray<FText> BindingErrorList = Event->GetCompilationMessages(UMVVMBlueprintViewEvent::EMessageType::Error);
+			TArray<FText> BindingWarningList = Event->GetCompilationMessages(UMVVMBlueprintViewEvent::EMessageType::Warning);
+			BindingErrorList.Append(BindingWarningList);
+
+			static const FText NewLineText = FText::FromString(TEXT("\n"));
+			FText HintText = LOCTEXT("ErrorButtonText", "Errors and Warnings: (Click to show in a separate window)");
+			FText ErrorsText = FText::Join(NewLineText, BindingErrorList);
+			return FText::Join(NewLineText, HintText, ErrorsText);
+		}
+		return FText();
+	}
+
+	FReply OnErrorButtonClicked()
+	{
+		ErrorDialog.Reset();
+		ErrorItems.Reset();
+
+		if (const UMVVMBlueprintViewEvent* Event = ViewEvent.Get())
+		{
+			for (const FText& ErrorText : Event->GetCompilationMessages(UMVVMBlueprintViewEvent::EMessageType::Error))
+			{
+				ErrorItems.Add(MakeShared<FText>(ErrorText));
+			}
+
+			for (const FText& WarningText : Event->GetCompilationMessages(UMVVMBlueprintViewEvent::EMessageType::Warning))
+			{
+				ErrorItems.Add(MakeShared<FText>(WarningText));
+			}
+
+			const FText BindingDisplayName = Event->GetDisplayName(true);
+			ErrorDialog = SNew(SCustomDialog)
+				.Title(FText::Format(LOCTEXT("Compilation Errors and Warnings", "Compilation Errors and Warnings for {0}"), BindingDisplayName))
+				.Buttons({
+					SCustomDialog::FButton(LOCTEXT("OK", "OK"))
+					})
+				.Content()
+				[
+					SNew(SListView<TSharedPtr<FText>>)
+					.ListItemsSource(&ErrorItems)
+					.OnGenerateRow(this, &SEventRow::OnGenerateErrorRow)
+				];
+
+			ErrorDialog->Show();
+		}
+
+		return FReply::Handled();
+	}
+
+	TSharedRef<ITableRow> OnGenerateErrorRow(TSharedPtr<FText> Text, const TSharedRef<STableViewBase>& TableView) const
+	{
+		return SNew(STableRow<TSharedPtr<FText>>, TableView)
+			.Content()
+			[
+				SNew(SEditableTextBox)
+				.BackgroundColor(FStyleColors::Background)
+				.IsReadOnly(true)
+				.Text(*Text.Get())
+			];
+	}
+
+	ECheckBoxState IsEventEnabled() const
+	{
+		if (const UMVVMBlueprintViewEvent* Event = ViewEvent.Get())
+		{
+			return Event->bEnabled ? ECheckBoxState::Checked : ECheckBoxState::Unchecked;
+		}
+		return ECheckBoxState::Undetermined;
+	}
+
+	void OnIsEventEnableChanged(ECheckBoxState NewState)
+	{
+		if (NewState == ECheckBoxState::Undetermined)
+		{
+			return;
+		}
+
+		if (UMVVMBlueprintViewEvent* Event = ViewEvent.Get())
+		{
+			UMVVMEditorSubsystem* Subsystem = GEditor->GetEditorSubsystem<UMVVMEditorSubsystem>();
+			Subsystem->SetEnabledForEvent(Event, NewState == ECheckBoxState::Checked);
+		}
+	}
+
+	ECheckBoxState IsEventCompiled() const
+	{
+		if (const UMVVMBlueprintViewEvent* Event = ViewEvent.Get())
+		{
+			return Event->bCompile ? ECheckBoxState::Checked : ECheckBoxState::Unchecked;
+		}
+		return ECheckBoxState::Undetermined;
+	}
+
+	void OnIsEventCompileChanged(ECheckBoxState NewState)
+	{
+		if (NewState == ECheckBoxState::Undetermined)
+		{
+			return;
+		}
+
+		if (UMVVMBlueprintViewEvent* Event = ViewEvent.Get())
+		{
+			UMVVMEditorSubsystem* Subsystem = GEditor->GetEditorSubsystem<UMVVMEditorSubsystem>();
+			Subsystem->SetCompileForEvent(Event, NewState == ECheckBoxState::Checked);
+		}
+	}
+
+	FMVVMBlueprintPropertyPath GetSelectedPropertyPath(bool bEvent) const
+	{
+		if (const UMVVMBlueprintViewEvent* Event = ViewEvent.Get())
+		{
+			return bEvent ? Event->GetEventPath() : Event->GetDestinationPath();
+		}
+		return FMVVMBlueprintPropertyPath();
+	}
+
+	const UFunction* GetSelectedConversionFunction(bool bSourceToDest) const
+	{
+		if (const UMVVMBlueprintViewEvent* Event = ViewEvent.Get())
+		{
+		}
+
+		return nullptr;
+	}
+
+	void HandleFieldSelectionChanged(FMVVMBlueprintPropertyPath SelectedField, const UFunction* Function, bool bEvent)
+	{
+		if (UMVVMBlueprintViewEvent* Event = ViewEvent.Get())
+		{
+			UMVVMEditorSubsystem* Subsystem = GEditor->GetEditorSubsystem<UMVVMEditorSubsystem>();
+			if (bEvent)
+			{
+				if (Event->GetEventPath() != SelectedField)
+				{
+					Subsystem->SetEventPath(Event, SelectedField);
+				}
+			}
+			else
+			{
+				if (Event->GetDestinationPath() != SelectedField)
+				{
+					Subsystem->SetEventDestinationPath(Event, SelectedField);
+				}
+			}
+		}
+	}
+
+	FFieldSelectionContext GetSelectedSelectionContext(bool bEvent) const
+	{
+		FFieldSelectionContext Result;
+		const UWidgetBlueprint* WidgetBlueprintPtr = WidgetBlueprintWeak.Get();
+		if (WidgetBlueprintPtr == nullptr)
+		{
+			return Result;
+		}
+
+		if (UMVVMBlueprintViewEvent* Event = ViewEvent.Get())
+		{
+			Result.BindingMode = EMVVMBindingMode::OneTimeToDestination;
+			if (bEvent && !Event->GetEventPath().GetWidgetName().IsNone())
+			{
+				Result.FixedBindingSource = FBindingSource::CreateForWidget(WidgetBlueprintPtr, Event->GetEventPath().GetWidgetName());
+			}
+
+			Result.bAllowWidgets = true;
+			Result.bAllowViewModels = !bEvent;
+			Result.bAllowConversionFunctions = false;
+			Result.bReadable = bEvent;
+			Result.bWritable = !bEvent;
+		}
+		return Result;
+	}
+
+	FReply HandleFieldSelectorDrop(const FGeometry& MyGeometry, const FDragDropEvent& DragDropEvent, bool bEvent)
+	{
+		TSharedPtr<FDecoratedDragDropOp> DragDropOp = DragDropEvent.GetOperationAs<FDecoratedDragDropOp>();
+		if (!DragDropOp.IsValid())
+		{
+			return FReply::Unhandled();
+		}
+
+		// Accept all drag-drop operations that are widget properties, but only accept view model fields when we are dropping into the Source box.
+		if (!DragDropOp->IsOfType<FWidgetPropertyDragDropOp>() && (!DragDropOp->IsOfType<FViewModelFieldDragDropOp>() || !bEvent))
+		{
+			return FReply::Unhandled();
+		}
+
+		UMVVMBlueprintViewEvent* Event = ViewEvent.Get();
+		UWidgetBlueprint* WidgetBlueprintPtr = WidgetBlueprintWeak.Get();
+		if (Event == nullptr || WidgetBlueprintPtr == nullptr)
+		{
+			return FReply::Unhandled();
+		}
+
+		TSharedPtr<FViewModelFieldDragDropOp> ViewModelFieldDragDropOp = DragDropEvent.GetOperationAs<FViewModelFieldDragDropOp>();
+		TSharedPtr<FWidgetPropertyDragDropOp> WidgetPropertyDragDropOp = DragDropEvent.GetOperationAs<FWidgetPropertyDragDropOp>();
+		bool bIsViewModelProperty = ViewModelFieldDragDropOp != nullptr;
+
+		UWidgetBlueprint* DragDropWidgetBP = bIsViewModelProperty ? ViewModelFieldDragDropOp->WidgetBP.Get() : WidgetPropertyDragDropOp->WidgetBP.Get();
+		if (WidgetBlueprintPtr != DragDropWidgetBP)
+		{
+			return FReply::Unhandled();
+		}
+
+		TArray<FFieldVariant> FieldPath = bIsViewModelProperty ? ViewModelFieldDragDropOp->DraggedField : WidgetPropertyDragDropOp->DraggedPropertyPath;
+		FMVVMBlueprintPropertyPath PropertyPath;
+		for (const FFieldVariant& Field : FieldPath)
+		{
+			PropertyPath.AppendPropertyPath(WidgetBlueprintPtr, FMVVMConstFieldVariant(Field));
+		}
+
+		UMVVMEditorSubsystem* Subsystem = GEditor->GetEditorSubsystem<UMVVMEditorSubsystem>();
+		if (bIsViewModelProperty)
+		{
+			if (ViewModelFieldDragDropOp->ViewModelId.IsValid())
+			{
+				PropertyPath.SetViewModelId(ViewModelFieldDragDropOp->ViewModelId);
+			}
+			else
+			{
+				return FReply::Unhandled();
+			}
+		}
+		else
+		{
+			if (UWidget* OwnerWidgetPtr = WidgetPropertyDragDropOp->OwnerWidget.Get())
+			{
+				// Search for the widget by its name in the widget tree
+				// If the widget is not found, we know it is the root preview widget so we use the blueprint name.
+				if (WidgetBlueprintPtr->WidgetTree->FindWidget(OwnerWidgetPtr->GetFName()))
+				{
+					PropertyPath.SetWidgetName(OwnerWidgetPtr->GetFName());
+				}
+				else
+				{
+					PropertyPath.SetWidgetName(WidgetBlueprintPtr->GetFName());
+				}
+			}
+		}
+
+		if (bEvent)
+		{
+			Subsystem->SetEventPath(Event, PropertyPath);
+		}
+		else
+		{
+			Subsystem->SetEventDestinationPath(Event, PropertyPath);
+		}
+		return FReply::Handled();
+	}
+
+	void HandleFieldSelectorDragEnter(const FGeometry& MyGeometry, const FDragDropEvent& DragDropEvent, bool bEvent)
+	{
+		TSharedPtr<FDecoratedDragDropOp> DragDropOp = DragDropEvent.GetOperationAs<FDecoratedDragDropOp>();
+		if (DragDropOp.IsValid())
+		{
+			// Accept all drag-drop operations that are widget properties, but only accept view model fields when we are dropping into the Source box.
+			if (DragDropOp->IsOfType<FWidgetPropertyDragDropOp>() || (DragDropOp->IsOfType<FViewModelFieldDragDropOp>() && !bEvent))
+			{
+				if (UWidgetBlueprint* WidgetBlueprintPtr = WidgetBlueprintWeak.Get())
+				{
+					TSharedPtr<FViewModelFieldDragDropOp> ViewModelFieldDragDropOp = DragDropEvent.GetOperationAs<FViewModelFieldDragDropOp>();
+					TSharedPtr<FWidgetPropertyDragDropOp> WidgetPropertyDragDropOp = DragDropEvent.GetOperationAs<FWidgetPropertyDragDropOp>();
+					bool IsViewModelProperty = false;
+
+					if (ViewModelFieldDragDropOp)
+					{
+						IsViewModelProperty = true;
+					}
+
+					UWidgetBlueprint* DragDropWidgetBP = IsViewModelProperty ? ViewModelFieldDragDropOp->WidgetBP.Get() : WidgetPropertyDragDropOp->WidgetBP.Get();
+
+					if (DragDropWidgetBP && DragDropWidgetBP == WidgetBlueprintPtr)
+					{
+						DragDropOp->CurrentIconBrush = FAppStyle::GetBrush(TEXT("Graph.ConnectorFeedback.OK"));
+						return;
+					}
+				}
+			}
+			else
+			{
+				DragDropOp->CurrentIconBrush = FAppStyle::GetBrush(TEXT("Graph.ConnectorFeedback.Error"));
+			}
+		}
+	}
+
+private:
+	TSharedPtr<FBindingEntry> Entry;
+	TWeakObjectPtr<UWidgetBlueprint> WidgetBlueprintWeak;
+	TWeakObjectPtr<UMVVMBlueprintViewEvent> ViewEvent;
+	TSharedPtr<SCustomDialog> ErrorDialog;
+	TArray<TSharedPtr<FText>> ErrorItems;
+};
+
+/**
+ * 
+ */
+class SFunctionParameterRow : public STableRow<TSharedPtr<FBindingEntry>>
+{
+	SLATE_BEGIN_ARGS(SFunctionParameterRow) {}
+	SLATE_END_ARGS()
+
+	void Construct(const FArguments& Args, const TSharedRef<STableViewBase>& OwnerTableView, const TSharedPtr<FBindingEntry>& InEntry, UWidgetBlueprint* InWidgetBlueprint)
+	{
+		Entry = InEntry;
+		check(Entry->GetRowType() == FBindingEntry::ERowType::BindingParameter || Entry->GetRowType() == FBindingEntry::ERowType::EventParameter);
+
+		WidgetBlueprint = InWidgetBlueprint;
+
+		UMVVMEditorSubsystem* EditorSubsystem = GEditor->GetEditorSubsystem<UMVVMEditorSubsystem>();
+
+		FSlateColor PrimaryColor, SecondaryColor;
+		const FSlateBrush* PrimaryBrush = nullptr;
+		const FSlateBrush* SecondaryBrush = nullptr;
+		FText DisplayName, ToolTip;
+
+		TSharedPtr<SWidget> ContentWidget;
+
+		if (Entry->GetRowType() == FBindingEntry::ERowType::BindingParameter)
+		{
+			bool bSimpleConversionFunction = false;
+			UMVVMBlueprintView* View = EditorSubsystem->GetView(InWidgetBlueprint);
+			FMVVMBlueprintViewBinding* Binding = Entry->GetBinding(View);
+			const bool bSourceToDestination = UE::MVVM::IsForwardBinding(Binding->BindingType);
+
+			if (UEdGraphPin* Pin = EditorSubsystem->GetConversionFunctionArgumentPin(InWidgetBlueprint, *Binding, Entry->GetBindingParameterName(), bSourceToDestination))
+			{
+				PrimaryBrush = FBlueprintEditor::GetVarIconAndColorFromPinType(Pin->PinType, PrimaryColor, SecondaryBrush, SecondaryColor);
+				DisplayName = Pin->GetDisplayName();
+				ToolTip = FText::FromString(Pin->PinToolTip);
+			}
+			else if (const UFunction* Function = EditorSubsystem->GetConversionFunction(InWidgetBlueprint, *Binding, bSourceToDestination))
+			{
+				// no wrapper graph, this is a simple conversion function of the form: int32 Convert(float x)
+				if (const FProperty* Argument = BindingHelper::GetFirstArgumentProperty(Function))
+				{
+					bSimpleConversionFunction = true;
+
+					PrimaryBrush = FBlueprintEditor::GetVarIconAndColorFromProperty(Argument, PrimaryColor, SecondaryBrush, SecondaryColor);
+					DisplayName = Argument->GetDisplayNameText();
+					ToolTip = Argument->GetToolTipText();
+				}
+			}
+
+			ContentWidget = SNew(SFunctionParameter, InWidgetBlueprint)
+				.BindingId(Binding->BindingId)
+				.ParameterName(Entry->GetBindingParameterName())
+				.SourceToDestination(bSourceToDestination)
+				.AllowDefault(!bSimpleConversionFunction);
+		}
+		else if (Entry->GetRowType() == FBindingEntry::ERowType::EventParameter)
+		{
+			UMVVMBlueprintViewEvent* ViewEvent = Entry->GetEvent();
+			if (UEdGraphPin* Pin = ViewEvent ? ViewEvent->GetOrCreateGraphPin(Entry->GetEventParameterName()) : nullptr)
+			{
+				PrimaryBrush = FBlueprintEditor::GetVarIconAndColorFromPinType(Pin->PinType, PrimaryColor, SecondaryBrush, SecondaryColor);
+				DisplayName = Pin->GetDisplayName();
+				ToolTip = FText::FromString(Pin->PinToolTip);
+			}
+
+			ContentWidget = SNew(SEventParameter, InWidgetBlueprint)
+				.Event(Entry->GetEvent())
+				.ParameterName(Entry->GetEventParameterName())
+				.AllowDefault(true);
+		}
+
+		STableRow<TSharedPtr<FBindingEntry>>::Construct(
+			STableRow<TSharedPtr<FBindingEntry>>::FArguments()
+			.ShowWires(true)
+			.Style(FMVVMEditorStyle::Get(), "BindingView.ParameterRow")
+			[
+				SNew(SBox)
+				.HeightOverride(30)
+				.ToolTipText(ToolTip)
+				[
+					SNew(SHorizontalBox)
+
+					+ SHorizontalBox::Slot()
+					.HAlign(HAlign_Left)
+					.VAlign(VAlign_Center)
+					.Padding(0, 0, 8, 0)
+					.AutoWidth()
+					[
+						SNew(SHorizontalBox)
+						+ SHorizontalBox::Slot()
+						.VAlign(VAlign_Center)
+						.Padding(0, 0, 8, 0)
+						.AutoWidth()
+						[
+							SNew(SImage)
+							.Image(PrimaryBrush)
+							.ColorAndOpacity(PrimaryColor)
+							.DesiredSizeOverride(FVector2D(16, 16))
+						]
+						+ SHorizontalBox::Slot()
+						[
+							SNew(STextBlock)
+							.Text(DisplayName)
+						]
+					]
+					+ SHorizontalBox::Slot()
+					.HAlign(HAlign_Left)
+					.VAlign(VAlign_Center)
+					.AutoWidth()
+					[
+						ContentWidget.ToSharedRef()
+					]
+				]
+			], 
+			OwnerTableView
+		);
+	}
+
+private:
+	TWeakObjectPtr<UWidgetBlueprint> WidgetBlueprint;
+	TSharedPtr<FBindingEntry> Entry;
+};
+
+void SBindingsList::Construct(const FArguments& InArgs, TSharedPtr<SBindingsPanel> Owner, UMVVMWidgetBlueprintExtension_View* InMVVMExtension)
+{
+	BindingPanel = Owner;
+	MVVMExtension = InMVVMExtension;
+	check(InMVVMExtension);
+	check(InMVVMExtension->GetBlueprintView());
 
 	MVVMExtension->OnBlueprintViewChangedDelegate().AddSP(this, &SBindingsList::Refresh);
 	MVVMExtension->GetBlueprintView()->OnBindingsUpdated.AddSP(this, &SBindingsList::Refresh);
@@ -1309,7 +1914,7 @@ void SBindingsList::Construct(const FArguments& InArgs, TSharedPtr<SBindingsPane
 	ChildSlot
 	[
 		SAssignNew(TreeView, STreeView<TSharedPtr<FBindingEntry>>)
-		.TreeItemsSource(&RootGroups)
+		.TreeItemsSource(&FilteredRootGroups)
 		.SelectionMode(ESelectionMode::Single)
 		.OnGenerateRow(this, &SBindingsList::GenerateEntryRow)
 		.OnGetChildren(this, &SBindingsList::GetChildrenOfEntry)
@@ -1333,8 +1938,7 @@ SBindingsList::~SBindingsList()
 
 void SBindingsList::GetChildrenOfEntry(TSharedPtr<FBindingEntry> Entry, TArray<TSharedPtr<FBindingEntry>>& OutChildren) const
 {
-	TConstArrayView<TSharedPtr<FBindingEntry>> Children = Entry->GetChildren();
-	OutChildren.Append(Children.GetData(), Children.Num());
+	OutChildren.Append(Entry->GetFilteredChildren());
 }
 
 void SBindingsList::Refresh()
@@ -1344,27 +1948,31 @@ void SBindingsList::Refresh()
 		TSharedPtr<FBindingEntry> Group;
 		TArray<TSharedPtr<FBindingEntry>> Children;
 	};
+
 	TArray<FPreviousGroup> PreviousRootGroups;
-	for (TSharedPtr<FBindingEntry> PreviousEntry : RootGroups)
+	for (const TSharedPtr<FBindingEntry>& PreviousEntry : AllRootGroups)
 	{
 		ensure(PreviousEntry->GetRowType() == FBindingEntry::ERowType::Group);
 		FPreviousGroup& NewItem = PreviousRootGroups.AddDefaulted_GetRef();
 		NewItem.Group = PreviousEntry;
 
-		// Add bindings
-		for (TSharedPtr<FBindingEntry> PreviousChildA : PreviousEntry->GetChildren())
+		struct FRecursiveAdd
 		{
-			NewItem.Children.Add(PreviousChildA);
-			// Add function arguments
-			for (TSharedPtr<FBindingEntry> PreviousChildB : PreviousChildA->GetChildren())
+			void operator()(FPreviousGroup& NewItem, const TSharedPtr<FBindingEntry>& Entry)
 			{
-				NewItem.Children.Add(PreviousChildB);
+				for (TSharedPtr<FBindingEntry> PreviousChildEntry : Entry->GetAllChildren())
+				{
+					NewItem.Children.Add(PreviousChildEntry);
+					(*this)(NewItem, PreviousChildEntry);
+				}
+				Entry->ResetChildren();
 			}
-			PreviousChildA->ResetChildren();
-		}
-		PreviousEntry->ResetChildren();
+		};
+		FRecursiveAdd{}(NewItem, PreviousEntry);
 	}
-	RootGroups.Reset();
+
+	AllRootGroups.Reset();
+	FilteredRootGroups.Reset();
 
 	TArray<TSharedPtr<FBindingEntry>> NewEntries;
 
@@ -1377,11 +1985,36 @@ void SBindingsList::Refresh()
 	// then add all bindings that reference that widget as its children
 	if (BlueprintView)
 	{
-		TArrayView<FMVVMBlueprintViewBinding> Bindings = BlueprintView->GetBindings();
-		for (int32 BindingIndex = 0; BindingIndex < Bindings.Num(); ++BindingIndex)
+		auto FindPreviousGroupEntry = [&PreviousRootGroups, Self = this](FName GroupName)
+		{
+			return PreviousRootGroups.FindByPredicate([GroupName](const FPreviousGroup& Other) { return Other.Group->GetGroupName() == GroupName; });
+		};
+		auto FindGroupEntry = [&NewEntries, Self = this](FPreviousGroup* PreviousGroupEntry, FName GroupName, FGuid ViewModelId)
 		{
-			const FMVVMBlueprintViewBinding& Binding = Bindings[BindingIndex];
+			TSharedPtr<FBindingEntry> GroupEntry;
+			if (PreviousGroupEntry)
+			{
+				GroupEntry = PreviousGroupEntry->Group;
+			}
+			else if (TSharedPtr<FBindingEntry>* FoundGroup = NewEntries.FindByPredicate([GroupName](const TSharedPtr<FBindingEntry>& Other)
+				{ return Other->GetGroupName() == GroupName && Other->GetRowType() == FBindingEntry::ERowType::Group; }))
+			{
+				GroupEntry = *FoundGroup;
+			}
+
+			if (!GroupEntry.IsValid())
+			{
+				GroupEntry = MakeShared<FBindingEntry>();
+				GroupEntry->SetGroup(GroupName, ViewModelId);
+
+				NewEntries.Add(GroupEntry);
+			}
+			Self->AllRootGroups.AddUnique(GroupEntry);
+			return GroupEntry;
+		};
 
+		for (const FMVVMBlueprintViewBinding& Binding : BlueprintView->GetBindings())
+		{
 			// Make sure the graph for the bindings is generated
 			if (Binding.Conversion.SourceToDestinationConversion)
 			{
@@ -1393,6 +2026,7 @@ void SBindingsList::Refresh()
 			}
 			
 			FName GroupName;
+			FGuid GroupViewModelId;
 			if (Binding.DestinationPath.IsFromWidget())
 			{
 				GroupName = Binding.DestinationPath.GetWidgetName();
@@ -1402,32 +2036,13 @@ void SBindingsList::Refresh()
 				if (const FMVVMBlueprintViewModelContext* ViewModelContext = BlueprintView->FindViewModel(Binding.DestinationPath.GetViewModelId()))
 				{
 					GroupName = ViewModelContext->GetViewModelName();
+					GroupViewModelId = ViewModelContext->GetViewModelId();
 				}
 			}
-			FPreviousGroup* PreviousGroupEntry = PreviousRootGroups.FindByPredicate([GroupName](const FPreviousGroup& Other) { return Other.Group->GetName() == GroupName; });
 
 			// Find the group entry
-			TSharedPtr<FBindingEntry> GroupEntry;
-			{
-				if (PreviousGroupEntry)
-				{
-					GroupEntry = PreviousGroupEntry->Group;
-				}
-				else if (TSharedPtr<FBindingEntry>* FoundGroup = NewEntries.FindByPredicate([GroupName](const TSharedPtr<FBindingEntry>& Other)
-					{ return Other->GetName() == GroupName && Other->GetRowType() == FBindingEntry::ERowType::Group; }))
-				{
-					GroupEntry = *FoundGroup;
-				}
-
-				if (!GroupEntry.IsValid())
-				{
-					GroupEntry = MakeShared<FBindingEntry>();
-					GroupEntry->SetGroupName(GroupName);
-
-					NewEntries.Add(GroupEntry);
-				}
-				RootGroups.AddUnique(GroupEntry);
-			}
+			FPreviousGroup* PreviousGroupEntry = FindPreviousGroupEntry(GroupName);
+			TSharedPtr<FBindingEntry> GroupEntry = FindGroupEntry(PreviousGroupEntry, GroupName, GroupViewModelId);
 
 			// Create/Find the child entry
 			TSharedPtr<FBindingEntry> BindingEntry;
@@ -1449,7 +2064,7 @@ void SBindingsList::Refresh()
 
 					NewEntries.Add(BindingEntry);
 				}
-				GroupEntry->AddChild(BindingEntry.ToSharedRef());
+				GroupEntry->AddChild(BindingEntry);
 			}
 
 			// Create/Find entries for conversion function parameters
@@ -1466,7 +2081,7 @@ void SBindingsList::Refresh()
 							if (PreviousGroupEntry)
 							{
 								if (TSharedPtr<FBindingEntry>* FoundParameter = PreviousGroupEntry->Children.FindByPredicate([BindingId, ArgumentName = Argument->GetFName()](const TSharedPtr<FBindingEntry>& Other)
-									{ return Other->GetBindingId() == BindingId && Other->GetRowType() == FBindingEntry::ERowType::Parameter && Other->GetName() == ArgumentName; }))
+									{ return Other->GetBindingId() == BindingId && Other->GetRowType() == FBindingEntry::ERowType::BindingParameter && Other->GetBindingParameterName() == ArgumentName; }))
 								{
 									ArgumentEntry = *FoundParameter;
 								}
@@ -1475,17 +2090,89 @@ void SBindingsList::Refresh()
 							if (!ArgumentEntry.IsValid())
 							{
 								ArgumentEntry = MakeShared<FBindingEntry>();
-								ArgumentEntry->SetParameterName(Binding.BindingId, Argument->GetFName());
+								ArgumentEntry->SetBindingParameterName(Binding.BindingId, Argument->GetFName());
 
 								NewEntries.Add(ArgumentEntry);
 							}
-							BindingEntry->AddChild(ArgumentEntry.ToSharedRef());
+							BindingEntry->AddChild(ArgumentEntry);
 						}
 					}
 				}
 			}
 		}
-		Private::FilterBindingsList(FilterText.ToString(), RootGroups, BlueprintView, MVVMExtensionPtr);
+
+		for (UMVVMBlueprintViewEvent* Event : BlueprintView->GetEvents())
+		{
+			Event->GetOrCreateWrapperGraph();
+
+			FName GroupName;
+			FGuid GroupViewModelId;
+			if (Event->GetEventPath().IsFromWidget())
+			{
+				GroupName = Event->GetEventPath().GetWidgetName();
+			}
+			else if (Event->GetEventPath().IsFromViewModel())
+			{
+				if (const FMVVMBlueprintViewModelContext* ViewModelContext = BlueprintView->FindViewModel(Event->GetEventPath().GetViewModelId()))
+				{
+					GroupName = ViewModelContext->GetViewModelName();
+					GroupViewModelId = ViewModelContext->GetViewModelId();
+				}
+			}
+
+			// Find the group entry
+			FPreviousGroup* PreviousGroupEntry = FindPreviousGroupEntry(GroupName);
+			TSharedPtr<FBindingEntry> GroupEntry = FindGroupEntry(PreviousGroupEntry, GroupName, GroupViewModelId);
+
+			// Create/Find the child entry
+			TSharedPtr<FBindingEntry> EventEntry;
+			{
+				if (PreviousGroupEntry)
+				{
+					if (TSharedPtr<FBindingEntry>* FoundBinding = PreviousGroupEntry->Children.FindByPredicate([Event](const TSharedPtr<FBindingEntry>& Other)
+						{
+							return Other->GetRowType() == FBindingEntry::ERowType::Event && Other->GetEvent() == Event;
+						}))
+					{
+						EventEntry = *FoundBinding;
+					}
+				}
+
+				if (!EventEntry.IsValid())
+				{
+					EventEntry = MakeShared<FBindingEntry>();
+					EventEntry->SetEvent(Event);
+
+					NewEntries.Add(EventEntry);
+				}
+				GroupEntry->AddChild(EventEntry);
+			}
+
+			// Create/Find entries for function parameters
+			for (const FMVVMBlueprintPin& Pin : Event->GetPins())
+			{
+				TSharedPtr<FBindingEntry> ArgumentEntry;
+				if (PreviousGroupEntry)
+				{
+					if (TSharedPtr<FBindingEntry>* FoundParameter = PreviousGroupEntry->Children.FindByPredicate([Event, ArgumentName = Pin.GetName()](const TSharedPtr<FBindingEntry>& Other)
+						{ return Other->GetRowType() == FBindingEntry::ERowType::EventParameter && Other->GetEvent() == Event && Other->GetEventParameterName() == ArgumentName; }))
+					{
+						ArgumentEntry = *FoundParameter;
+					}
+				}
+
+				if (!ArgumentEntry.IsValid())
+				{
+					ArgumentEntry = MakeShared<FBindingEntry>();
+					ArgumentEntry->SetEventParameterName(Event, Pin.GetName());
+
+					NewEntries.Add(ArgumentEntry);
+				}
+				EventEntry->AddChild(ArgumentEntry);
+			}
+		}
+
+		Private::FilterEntryList(FilterText.ToString(), AllRootGroups, FilteredRootGroups, BlueprintView, MVVMExtensionPtr);
 	}
 
 	if (TreeView.IsValid())
@@ -1508,7 +2195,7 @@ TSharedRef<ITableRow> SBindingsList::GenerateEntryRow(TSharedPtr<FBindingEntry>
 		{
 			case FBindingEntry::ERowType::Group:
 			{
-				Row = SNew(SWidgetRow, OwnerTable, Entry, MVVMExtensionPtr->GetWidgetBlueprint());
+				Row = SNew(SGroupRow, OwnerTable, Entry, MVVMExtensionPtr->GetWidgetBlueprint());
 				break;
 			}
 			case FBindingEntry::ERowType::Binding:
@@ -1516,11 +2203,17 @@ TSharedRef<ITableRow> SBindingsList::GenerateEntryRow(TSharedPtr<FBindingEntry>
 				Row = SNew(SBindingRow, OwnerTable, Entry, MVVMExtensionPtr->GetWidgetBlueprint());
 				break;
 			}
-			case FBindingEntry::ERowType::Parameter:
+			case FBindingEntry::ERowType::BindingParameter:
+			case FBindingEntry::ERowType::EventParameter:
 			{
 				Row = SNew(SFunctionParameterRow, OwnerTable, Entry, MVVMExtensionPtr->GetWidgetBlueprint());
 				break;
 			}
+			case FBindingEntry::ERowType::Event:
+			{
+				Row = SNew(SEventRow, OwnerTable, Entry, MVVMExtensionPtr->GetWidgetBlueprint());
+				break;
+			}
 		}
 
 		return Row.ToSharedRef();
@@ -1543,17 +2236,28 @@ void SBindingsList::ClearFilterText()
 
 namespace Private
 {
-	void GatherChildBindings(UMVVMBlueprintView* BlueprintView, const TConstArrayView<TSharedPtr<FBindingEntry>>& Entries, TArray<const FMVVMBlueprintViewBinding*>& OutBindings)
+	void GatherAllChildBindings(UMVVMBlueprintView* BlueprintView, const TConstArrayView<TSharedPtr<FBindingEntry>>& Entries, TArray<const FMVVMBlueprintViewBinding*>& OutBindings, TArray<UMVVMBlueprintViewEvent*> & OutEvents)
 	{
 		for (const TSharedPtr<FBindingEntry>& Entry : Entries)
 		{
-			const FMVVMBlueprintViewBinding* Binding = Entry->GetBinding(BlueprintView);
-			if (Binding != nullptr)
+			if (Entry->GetRowType() == FBindingEntry::ERowType::Binding)
 			{
-				OutBindings.AddUnique(Binding);
+				const FMVVMBlueprintViewBinding* Binding = Entry->GetBinding(BlueprintView);
+				if (Binding != nullptr)
+				{
+					OutBindings.AddUnique(Binding);
+				}
 			}
 
-			GatherChildBindings(BlueprintView, Entry->GetChildren(), OutBindings);
+			if (Entry->GetRowType() == FBindingEntry::ERowType::Event)
+			{
+				UMVVMBlueprintViewEvent* Event = Entry->GetEvent();
+				if (Event)
+				{
+					OutEvents.AddUnique(Event);
+				}
+			}
+			GatherAllChildBindings(BlueprintView, Entry->GetAllChildren(), OutBindings, OutEvents);
 		}
 	}
 }
@@ -1573,9 +2277,10 @@ void SBindingsList::OnDeleteSelected()
 			const UWidgetBlueprint* WidgetBlueprint = MVVMExtensionPtr->GetWidgetBlueprint();
 
 			TArray<const FMVVMBlueprintViewBinding*> BindingsToRemove;
-			Private::GatherChildBindings(BlueprintView, Selection, BindingsToRemove);
+			TArray<UMVVMBlueprintViewEvent*> EventsToRemove;
+			Private::GatherAllChildBindings(BlueprintView, Selection, BindingsToRemove, EventsToRemove);
 
-			if (BindingsToRemove.Num() == 0)
+			if (BindingsToRemove.Num() == 0 && EventsToRemove.Num() == 0)
 			{
 				return;
 			}
@@ -1585,6 +2290,10 @@ void SBindingsList::OnDeleteSelected()
 			{
 				BindingDisplayNames.Add(FText::FromString(Binding->GetDisplayNameString(WidgetBlueprint)));
 			}
+			for (const UMVVMBlueprintViewEvent* Event : EventsToRemove)
+			{
+				BindingDisplayNames.Add(Event->GetDisplayName(true));
+			}
 
 			const FText Message = FText::Format(BindingDisplayNames.Num() == 1 ?
 				LOCTEXT("ConfirmDeleteSingle", "Are you sure that you want to delete this binding?\n\n{1}") :
@@ -1609,6 +2318,10 @@ void SBindingsList::OnDeleteSelected()
 				{
 					BlueprintView->RemoveBinding(Binding);
 				}
+				for (UMVVMBlueprintViewEvent* Event : EventsToRemove)
+				{
+					BlueprintView->RemoveEvent(Event);
+				}
 			}
 		}
 	}
@@ -1635,7 +2348,16 @@ TSharedPtr<SWidget> SBindingsList::OnSourceConstructContextMenu()
 
 void SBindingsList::RequestNavigateToBinding(FGuid BindingId)
 {
-	TSharedPtr<FBindingEntry> Entry = Private::FindBinding(BindingId, RootGroups);
+	TSharedPtr<FBindingEntry> Entry = Private::FindBinding(BindingId, FilteredRootGroups);
+	if (Entry && TreeView)
+	{
+		TreeView->RequestNavigateToItem(Entry);
+	}
+}
+
+void SBindingsList::RequestNavigateToEvent(UMVVMBlueprintViewEvent* Event)
+{
+	TSharedPtr<FBindingEntry> Entry = Private::FindEvent(Event, FilteredRootGroups);
 	if (Entry && TreeView)
 	{
 		TreeView->RequestNavigateToItem(Entry);
diff --git a/Engine/Plugins/Runtime/ModelViewViewModel/Source/ModelViewViewModelEditor/Private/Widgets/SMVVMViewBindingListView.h b/Engine/Plugins/Runtime/ModelViewViewModel/Source/ModelViewViewModelEditor/Private/Widgets/SMVVMViewBindingListView.h
--- a/Engine/Plugins/Runtime/ModelViewViewModel/Source/ModelViewViewModelEditor/Private/Widgets/SMVVMViewBindingListView.h
+++ b/Engine/Plugins/Runtime/ModelViewViewModel/Source/ModelViewViewModelEditor/Private/Widgets/SMVVMViewBindingListView.h
@@ -11,6 +11,7 @@ class ITableRow;
 template <typename ItemType> class STreeView;
 class STableViewBase;
 class UMVVMWidgetBlueprintExtension_View;
+class UMVVMBlueprintViewEvent;
 
 namespace UE::MVVM
 {
@@ -36,6 +37,7 @@ public:
 	TSharedPtr<SWidget> OnSourceConstructContextMenu();
 
 	void RequestNavigateToBinding(FGuid BindingId);
+	void RequestNavigateToEvent(UMVVMBlueprintViewEvent* Event);
 
 private:
 	TSharedRef<ITableRow> GenerateEntryRow(TSharedPtr<FBindingEntry> Entry, const TSharedRef<STableViewBase>& OwnerTable) const;
@@ -48,7 +50,8 @@ private:
 private:
 	TWeakPtr<SBindingsPanel> BindingPanel;
 	TSharedPtr<STreeView<TSharedPtr<FBindingEntry>>> TreeView;
-	TArray<TSharedPtr<FBindingEntry>> RootGroups;
+	TArray<TSharedPtr<FBindingEntry>> AllRootGroups;
+	TArray<TSharedPtr<FBindingEntry>> FilteredRootGroups;
 	TWeakObjectPtr<UMVVMWidgetBlueprintExtension_View> MVVMExtension;
 	mutable bool bSelectionChangedGuard = false;
 	FText FilterText;
diff --git a/Engine/Plugins/Runtime/ModelViewViewModel/Source/ModelViewViewModelEditor/Private/Widgets/SMVVMViewBindingPanel.cpp b/Engine/Plugins/Runtime/ModelViewViewModel/Source/ModelViewViewModelEditor/Private/Widgets/SMVVMViewBindingPanel.cpp
--- a/Engine/Plugins/Runtime/ModelViewViewModel/Source/ModelViewViewModelEditor/Private/Widgets/SMVVMViewBindingPanel.cpp
+++ b/Engine/Plugins/Runtime/ModelViewViewModel/Source/ModelViewViewModelEditor/Private/Widgets/SMVVMViewBindingPanel.cpp
@@ -272,15 +272,8 @@ void SBindingsPanel::AddDefaultBinding()
 
 bool SBindingsPanel::CanAddBinding() const
 {
-	if (UMVVMWidgetBlueprintExtension_View* MVVMExtensionPtr = MVVMExtension.Get())
-	{
-		if (UMVVMBlueprintView* BlueprintView = MVVMExtensionPtr->GetBlueprintView())
-		{
-			return BlueprintView->GetViewModels().Num() > 0;
-		}
-	}
-
-	return false;
+	UMVVMWidgetBlueprintExtension_View* MVVMExtensionPtr = MVVMExtension.Get();
+	return MVVMExtensionPtr && MVVMExtensionPtr->GetBlueprintView() != nullptr;
 }
 
 FText SBindingsPanel::GetAddBindingText() const
@@ -341,11 +334,12 @@ TSharedRef<SWidget> SBindingsPanel::HandleAddDefaultBindingContextMenu()
 	FMenuBuilder MenuBuilder(bShouldCloseWindowAfterMenuSelection, nullptr);
 
 	MenuBuilder.AddMenuEntry(
-		LOCTEXT("AddSelectedWidget", "Add Selected Widget(s)"),
+		LOCTEXT("AddSelectedWidget", "Add Selected Widget(s) binding"),
 		FText::GetEmpty(),
 		FSlateIcon(),
 		FUIAction(
-			FExecuteAction::CreateSP(this, &SBindingsPanel::HandleAddDefaultBindingButtonClick, EAddBindingMode::Selected)
+			FExecuteAction::CreateSP(this, &SBindingsPanel::HandleAddDefaultBindingButtonClick, EAddBindingMode::Selected),
+			FCanExecuteAction::CreateSP(this, &SBindingsPanel::CanAddBinding)
 		));
 
 	MenuBuilder.AddMenuEntry(
@@ -353,7 +347,26 @@ TSharedRef<SWidget> SBindingsPanel::HandleAddDefaultBindingContextMenu()
 		FText::GetEmpty(),
 		FSlateIcon(),
 		FUIAction(
-			FExecuteAction::CreateSP(this, &SBindingsPanel::HandleAddDefaultBindingButtonClick, EAddBindingMode::Empty)
+			FExecuteAction::CreateSP(this, &SBindingsPanel::HandleAddDefaultBindingButtonClick, EAddBindingMode::Empty),
+			FCanExecuteAction::CreateSP(this, &SBindingsPanel::CanAddBinding)
+		));
+
+	MenuBuilder.AddMenuEntry(
+		LOCTEXT("AddSelectedWidgetEvent", "Add Selected Widget(s) event"),
+		FText::GetEmpty(),
+		FSlateIcon(),
+		FUIAction(
+			FExecuteAction::CreateSP(this, &SBindingsPanel::HandleAddDefaultEventButtonClick, EAddBindingMode::Selected),
+			FCanExecuteAction::CreateSP(this, &SBindingsPanel::CanAddBinding)
+		));
+
+	MenuBuilder.AddMenuEntry(
+		LOCTEXT("AddEmptyWidgetEvent", "Add Empty event"),
+		FText::GetEmpty(),
+		FSlateIcon(),
+		FUIAction(
+			FExecuteAction::CreateSP(this, &SBindingsPanel::HandleAddDefaultEventButtonClick, EAddBindingMode::Empty),
+			FCanExecuteAction::CreateSP(this, &SBindingsPanel::CanAddBinding)
 		));
 
 	return MenuBuilder.MakeWidget();
@@ -369,6 +382,41 @@ void SBindingsPanel::HandleAddDefaultBindingButtonClick(EAddBindingMode NewMode)
 	AddDefaultBinding();
 }
 
+void SBindingsPanel::HandleAddDefaultEventButtonClick(EAddBindingMode NewMode)
+{
+	if (UMVVMWidgetBlueprintExtension_View* MVVMExtensionPtr = MVVMExtension.Get())
+	{
+		UMVVMEditorSubsystem* EditorSubsystem = GEditor->GetEditorSubsystem<UMVVMEditorSubsystem>();
+		if (TSharedPtr<FWidgetBlueprintEditor> BlueprintEditor = WeakBlueprintEditor.Pin())
+		{
+			UMVVMBlueprintViewEvent* NewEvent = nullptr;
+			if (AddBindingMode == EAddBindingMode::Selected)
+			{
+				for (const FWidgetReference& WidgetReference : BlueprintEditor->GetSelectedWidgets())
+				{
+					if (WidgetReference.IsValid() && WidgetReference.GetTemplate())
+					{
+						NewEvent = EditorSubsystem->AddEvent(MVVMExtensionPtr->GetWidgetBlueprint());
+						FMVVMBlueprintPropertyPath Path;
+						Path.SetWidgetName(WidgetReference.GetTemplate()->GetFName());
+						EditorSubsystem->SetEventPath(NewEvent, Path);
+					}
+				}
+			}
+
+			if (!NewEvent)
+			{
+				NewEvent = EditorSubsystem->AddEvent(MVVMExtensionPtr->GetWidgetBlueprint());
+			}
+
+			if (NewEvent && BindingsList)
+			{
+				BindingsList->RequestNavigateToEvent(NewEvent);
+			}
+		}
+	}
+}
+
 TSharedRef<SWidget> SBindingsPanel::CreateDrawerDockButton()
 {
 	if (bIsDrawerTab)
@@ -734,7 +782,7 @@ EVisibility SBindingsPanel::GetVisibility(bool bVisibleWithBindings) const
 	if (UMVVMWidgetBlueprintExtension_View* MVVMExtensionPtr = MVVMExtension.Get())
 	{
 		if (MVVMExtensionPtr->GetBlueprintView() != nullptr &&
-			MVVMExtensionPtr->GetBlueprintView()->GetNumBindings() > 0)
+			(MVVMExtensionPtr->GetBlueprintView()->GetNumBindings() > 0 || MVVMExtensionPtr->GetBlueprintView()->GetEvents().Num() > 0))
 		{
 			return bVisibleWithBindings ? EVisibility::Visible : EVisibility::Collapsed;
 		}
diff --git a/Engine/Plugins/Runtime/ModelViewViewModel/Source/ModelViewViewModelEditor/Private/Widgets/SMVVMViewBindingPanel.h b/Engine/Plugins/Runtime/ModelViewViewModel/Source/ModelViewViewModelEditor/Private/Widgets/SMVVMViewBindingPanel.h
--- a/Engine/Plugins/Runtime/ModelViewViewModel/Source/ModelViewViewModelEditor/Private/Widgets/SMVVMViewBindingPanel.h
+++ b/Engine/Plugins/Runtime/ModelViewViewModel/Source/ModelViewViewModelEditor/Private/Widgets/SMVVMViewBindingPanel.h
@@ -63,6 +63,7 @@ private:
 	FText GetAddBindingToolTip() const;
 	TSharedRef<SWidget> HandleAddDefaultBindingContextMenu();
 	void HandleAddDefaultBindingButtonClick(EAddBindingMode NewMode);
+	void HandleAddDefaultEventButtonClick(EAddBindingMode NewMode);
 
 	TSharedRef<SWidget> GenerateSettingsMenu();
 
diff --git a/Engine/Plugins/Runtime/ModelViewViewModel/Source/ModelViewViewModelEditor/Public/MVVMEditorSubsystem.h b/Engine/Plugins/Runtime/ModelViewViewModel/Source/ModelViewViewModelEditor/Public/MVVMEditorSubsystem.h
--- a/Engine/Plugins/Runtime/ModelViewViewModel/Source/ModelViewViewModelEditor/Public/MVVMEditorSubsystem.h
+++ b/Engine/Plugins/Runtime/ModelViewViewModel/Source/ModelViewViewModelEditor/Public/MVVMEditorSubsystem.h
@@ -20,6 +20,7 @@ template <typename T> class TSubclassOf;
 
 class UEdGraph;
 class UK2Node_CallFunction;
+class UMVVMBlueprintViewEvent;
 class UWidgetBlueprint;
 
 /** */
@@ -53,6 +54,12 @@ public:
 	UFUNCTION(BlueprintCallable, Category = "Viewmodel")
 	void RemoveBinding(UWidgetBlueprint* WidgetBlueprint, const FMVVMBlueprintViewBinding& Binding);
 
+	UFUNCTION(BlueprintCallable, Category = "Viewmodel")
+	UMVVMBlueprintViewEvent* AddEvent(UWidgetBlueprint* WidgetBlueprint);
+
+	UFUNCTION(BlueprintCallable, Category = "Viewmodel")
+	void RemoveEvent(UWidgetBlueprint* WidgetBlueprint, UMVVMBlueprintViewEvent* Event);
+
 	UFUNCTION(BlueprintCallable, Category = "Viewmodel")
 	TArray<FMVVMAvailableBinding> GetChildViewModels(TSubclassOf<UObject> Class, TSubclassOf<UObject> Accessor);
 
@@ -66,6 +73,12 @@ public:
 	void SetEnabledForBinding(UWidgetBlueprint* WidgetBlueprint, FMVVMBlueprintViewBinding& Binding, bool bEnabled);
 	void SetCompileForBinding(UWidgetBlueprint* WidgetBlueprint, FMVVMBlueprintViewBinding& Binding, bool bCompile);
 
+	void SetEventPath(UMVVMBlueprintViewEvent* Event, FMVVMBlueprintPropertyPath PropertyPath);
+	void SetEventDestinationPath(UMVVMBlueprintViewEvent* Event, FMVVMBlueprintPropertyPath PropertyPath);
+	void SetEventArgumentPath(UMVVMBlueprintViewEvent* Event, FName ArgumentName, const FMVVMBlueprintPropertyPath& PropertyPath) const;
+	void SetEnabledForEvent(UMVVMBlueprintViewEvent* Event, bool bEnabled);
+	void SetCompileForEvent(UMVVMBlueprintViewEvent* Event, bool bCompile);
+
 	UFUNCTION(BlueprintCallable, Category = "Viewmodel")
 	bool IsValidConversionFunction(const UWidgetBlueprint* WidgeteBlueprint, const UFunction* Function, const FMVVMBlueprintPropertyPath& Source, const FMVVMBlueprintPropertyPath& Destination) const;
 
