Add setting to disable Binding Event

diff --git a/Engine/Plugins/Runtime/ModelViewViewModel/Source/ModelViewViewModelBlueprint/Private/MVVMBlueprintViewEvent.cpp b/Engine/Plugins/Runtime/ModelViewViewModel/Source/ModelViewViewModelBlueprint/Private/MVVMBlueprintViewEvent.cpp
--- a/Engine/Plugins/Runtime/ModelViewViewModel/Source/ModelViewViewModelBlueprint/Private/MVVMBlueprintViewEvent.cpp
+++ b/Engine/Plugins/Runtime/ModelViewViewModel/Source/ModelViewViewModelBlueprint/Private/MVVMBlueprintViewEvent.cpp
@@ -7,6 +7,7 @@
 #include "Bindings/MVVMBindingHelper.h"
 #include "Engine/Blueprint.h"
 #include "Kismet2/BlueprintEditorUtils.h"
+#include "MVVMDeveloperProjectSettings.h"
 #include "MVVMWidgetBlueprintExtension_View.h"
 #include "WidgetBlueprint.h"
 
@@ -148,7 +149,7 @@ UWidgetBlueprint* UMVVMBlueprintViewEvent::GetWidgetBlueprintInternal() const
 
 bool UMVVMBlueprintViewEvent::Supports(const UWidgetBlueprint* WidgetBlueprint, const FMVVMBlueprintPropertyPath& PropertyPath)
 {
-	return GetEventSignature(WidgetBlueprint, PropertyPath) != nullptr;
+	return GetDefault<UMVVMDeveloperProjectSettings>()->bAllowBindingEvent && GetEventSignature(WidgetBlueprint, PropertyPath) != nullptr;
 }
 
 const UFunction* UMVVMBlueprintViewEvent::GetEventSignature() const
diff --git a/Engine/Plugins/Runtime/ModelViewViewModel/Source/ModelViewViewModelBlueprint/Private/MVVMViewBlueprintCompiler.cpp b/Engine/Plugins/Runtime/ModelViewViewModel/Source/ModelViewViewModelBlueprint/Private/MVVMViewBlueprintCompiler.cpp
--- a/Engine/Plugins/Runtime/ModelViewViewModel/Source/ModelViewViewModelBlueprint/Private/MVVMViewBlueprintCompiler.cpp
+++ b/Engine/Plugins/Runtime/ModelViewViewModel/Source/ModelViewViewModelBlueprint/Private/MVVMViewBlueprintCompiler.cpp
@@ -1914,6 +1914,11 @@ bool FMVVMViewBlueprintCompiler::PreCompileEvents(UWidgetBlueprintGeneratedClass
 		return false;
 	}
 
+	if (!GetDefault<UMVVMDeveloperProjectSettings>()->bAllowBindingEvent && BlueprintView->GetEvents().Num() > 0)
+	{
+		WidgetBlueprintCompilerContext.MessageLog.Warning(*LOCTEXT("EventsAreNotAllowed", "Binding events are not allowed in your project settings.").ToString());
+	}
+
 	for (UMVVMBlueprintViewEvent* EventPtr : BlueprintView->GetEvents())
 	{
 		if (EventPtr == nullptr || !EventPtr->bCompile)
diff --git a/Engine/Plugins/Runtime/ModelViewViewModel/Source/ModelViewViewModelBlueprint/Public/MVVMDeveloperProjectSettings.h b/Engine/Plugins/Runtime/ModelViewViewModel/Source/ModelViewViewModelBlueprint/Public/MVVMDeveloperProjectSettings.h
--- a/Engine/Plugins/Runtime/ModelViewViewModel/Source/ModelViewViewModelBlueprint/Public/MVVMDeveloperProjectSettings.h
+++ b/Engine/Plugins/Runtime/ModelViewViewModel/Source/ModelViewViewModelBlueprint/Public/MVVMDeveloperProjectSettings.h
@@ -128,6 +128,10 @@ public:
 	UPROPERTY(EditAnywhere, config, Category = "Viewmodel")
 	bool bAllowConversionFunctionGeneratedGraphInEditor = true;
 
+	/** When binding to a multicast delegate property, allow to create an event. */
+	UPROPERTY(EditAnywhere, config, Category = "Viewmodel")
+	bool bAllowBindingEvent = true;
+
 	/** Permission list for filtering which execution mode is allowed. */
 	UPROPERTY(EditAnywhere, config, Category = "Viewmodel")
 	EMVVMDeveloperConversionFunctionFilterType ConversionFunctionFilter = EMVVMDeveloperConversionFunctionFilterType::BlueprintActionRegistry;
diff --git a/Engine/Plugins/Runtime/ModelViewViewModel/Source/ModelViewViewModelEditor/Private/MVVMEditorSubsystem.cpp b/Engine/Plugins/Runtime/ModelViewViewModel/Source/ModelViewViewModelEditor/Private/MVVMEditorSubsystem.cpp
--- a/Engine/Plugins/Runtime/ModelViewViewModel/Source/ModelViewViewModelEditor/Private/MVVMEditorSubsystem.cpp
+++ b/Engine/Plugins/Runtime/ModelViewViewModel/Source/ModelViewViewModelEditor/Private/MVVMEditorSubsystem.cpp
@@ -265,8 +265,12 @@ void UMVVMEditorSubsystem::RemoveBinding(UWidgetBlueprint* WidgetBlueprint, cons
 
 UMVVMBlueprintViewEvent* UMVVMEditorSubsystem::AddEvent(UWidgetBlueprint* WidgetBlueprint)
 {
-	UMVVMBlueprintView* View = RequestView(WidgetBlueprint);
-	return View->AddDefaultEvent();
+	if (GetDefault<UMVVMDeveloperProjectSettings>()->bAllowBindingEvent)
+	{
+		UMVVMBlueprintView* View = RequestView(WidgetBlueprint);
+		return View->AddDefaultEvent();
+	}
+	return nullptr;
 }
 
 void UMVVMEditorSubsystem::RemoveEvent(UWidgetBlueprint* WidgetBlueprint, UMVVMBlueprintViewEvent* Event)
@@ -368,9 +372,14 @@ void UMVVMEditorSubsystem::SetDestinationPathForBinding(UWidgetBlueprint* Widget
 		UE::MVVM::Private::OnBindingPreEditChange(View, GET_MEMBER_NAME_CHECKED(FMVVMBlueprintViewBinding, DestinationPath));
 
 		bool bSupports = UMVVMBlueprintViewEvent::Supports(WidgetBlueprint, PropertyPath);
+		UMVVMBlueprintViewEvent* Event = nullptr;
 		if (bSupports)
 		{
-			UMVVMBlueprintViewEvent* Event = AddEvent(WidgetBlueprint);
+			Event = AddEvent(WidgetBlueprint);
+		}
+
+		if (Event)
+		{
 			Event->SetEventPath(PropertyPath);
 			View->RemoveBinding(&Binding);
 		}
