MVVM: Reuse the same library when generating the view class. It helps the cook make smaller diff packages when the cook is deterministic.

diff --git a/Engine/Plugins/Runtime/ModelViewViewModel/Source/ModelViewViewModel/Private/Bindings/MVVMCompiledBindingLibrary.cpp b/Engine/Plugins/Runtime/ModelViewViewModel/Source/ModelViewViewModel/Private/Bindings/MVVMCompiledBindingLibrary.cpp
--- a/Engine/Plugins/Runtime/ModelViewViewModel/Source/ModelViewViewModel/Private/Bindings/MVVMCompiledBindingLibrary.cpp
+++ b/Engine/Plugins/Runtime/ModelViewViewModel/Source/ModelViewViewModel/Private/Bindings/MVVMCompiledBindingLibrary.cpp
@@ -64,6 +64,14 @@ FMVVMCompiledBindingLibrary::FMVVMCompiledBindingLibrary()
 
 }
 
+#if WITH_EDITOR
+FMVVMCompiledBindingLibrary::FMVVMCompiledBindingLibrary(FGuid LibraryId)
+	: CompiledBindingLibraryId(LibraryId)
+{
+
+}
+#endif
+
 
 FText FMVVMCompiledBindingLibrary::LexToText(EExecutionFailingReason Reason)
 {
diff --git a/Engine/Plugins/Runtime/ModelViewViewModel/Source/ModelViewViewModel/Public/Bindings/MVVMCompiledBindingLibrary.h b/Engine/Plugins/Runtime/ModelViewViewModel/Source/ModelViewViewModel/Public/Bindings/MVVMCompiledBindingLibrary.h
--- a/Engine/Plugins/Runtime/ModelViewViewModel/Source/ModelViewViewModel/Public/Bindings/MVVMCompiledBindingLibrary.h
+++ b/Engine/Plugins/Runtime/ModelViewViewModel/Source/ModelViewViewModel/Public/Bindings/MVVMCompiledBindingLibrary.h
@@ -277,6 +277,9 @@ struct MODELVIEWVIEWMODEL_API FMVVMCompiledBindingLibrary
 
 public:
 	FMVVMCompiledBindingLibrary();
+#if WITH_EDITOR
+	FMVVMCompiledBindingLibrary(FGuid LibraryId);
+#endif
 
 public:
 	/** */
diff --git a/Engine/Plugins/Runtime/ModelViewViewModel/Source/ModelViewViewModelBlueprint/Private/Bindings/MVVMCompiledBindingLibraryCompiler.cpp b/Engine/Plugins/Runtime/ModelViewViewModel/Source/ModelViewViewModelBlueprint/Private/Bindings/MVVMCompiledBindingLibraryCompiler.cpp
--- a/Engine/Plugins/Runtime/ModelViewViewModel/Source/ModelViewViewModelBlueprint/Private/Bindings/MVVMCompiledBindingLibraryCompiler.cpp
+++ b/Engine/Plugins/Runtime/ModelViewViewModel/Source/ModelViewViewModelBlueprint/Private/Bindings/MVVMCompiledBindingLibraryCompiler.cpp
@@ -192,6 +192,14 @@ int32 FCompiledBindingLibraryCompiler::FBindingHandle::IdGenerator = 0;
 int32 FCompiledBindingLibraryCompiler::FFieldPathHandle::IdGenerator = 0;
 int32 FCompiledBindingLibraryCompiler::FFieldIdHandle::IdGenerator = 0;
 
+/**
+ *
+ */
+FCompiledBindingLibraryCompiler::FCompileResult::FCompileResult(FGuid LibraryId)
+	: Library(LibraryId)
+{
+
+}
 
 /**
  *
@@ -644,7 +652,7 @@ TValueOrError<FCompiledBindingLibraryCompiler::FBindingHandle, FText> FCompiledB
 }
 
 
-TValueOrError<FCompiledBindingLibraryCompiler::FCompileResult, FText> FCompiledBindingLibraryCompiler::Compile()
+TValueOrError<FCompiledBindingLibraryCompiler::FCompileResult, FText> FCompiledBindingLibraryCompiler::Compile(FGuid LibraryId)
 {
 	Impl->bCompiled = false;
 
@@ -710,7 +718,7 @@ TValueOrError<FCompiledBindingLibraryCompiler::FCompileResult, FText> FCompiledB
 		// then: merge the 2 and restart the algo
 
 
-	FCompileResult Result;
+	FCompileResult Result = FCompileResult(LibraryId);
 
 	// Create FMVVMCompiledBindingLibrary::CompiledFields and FMVVMCompiledBindingLibrary::CompiledFieldNames
 	int32 TotalNumberOfProperties = 0;
diff --git a/Engine/Plugins/Runtime/ModelViewViewModel/Source/ModelViewViewModelBlueprint/Private/MVVMBlueprintView.cpp b/Engine/Plugins/Runtime/ModelViewViewModel/Source/ModelViewViewModelBlueprint/Private/MVVMBlueprintView.cpp
--- a/Engine/Plugins/Runtime/ModelViewViewModel/Source/ModelViewViewModelBlueprint/Private/MVVMBlueprintView.cpp
+++ b/Engine/Plugins/Runtime/ModelViewViewModel/Source/ModelViewViewModelBlueprint/Private/MVVMBlueprintView.cpp
@@ -16,6 +16,7 @@
 UMVVMBlueprintView::UMVVMBlueprintView()
 {
 	Settings = CreateDefaultSubobject<UMVVMBlueprintViewSettings>("Settings");
+	CompiledBindingLibraryId = FGuid::NewGuid();
 }
 
 FMVVMBlueprintViewModelContext* UMVVMBlueprintView::FindViewModel(FGuid ViewModelId)
diff --git a/Engine/Plugins/Runtime/ModelViewViewModel/Source/ModelViewViewModelBlueprint/Private/MVVMViewBlueprintCompiler.cpp b/Engine/Plugins/Runtime/ModelViewViewModel/Source/ModelViewViewModelBlueprint/Private/MVVMViewBlueprintCompiler.cpp
--- a/Engine/Plugins/Runtime/ModelViewViewModel/Source/ModelViewViewModelBlueprint/Private/MVVMViewBlueprintCompiler.cpp
+++ b/Engine/Plugins/Runtime/ModelViewViewModel/Source/ModelViewViewModelBlueprint/Private/MVVMViewBlueprintCompiler.cpp
@@ -949,7 +949,7 @@ bool FMVVMViewBlueprintCompiler::Compile(UWidgetBlueprintGeneratedClass* Class,
 		return false;
 	}
 
-	TValueOrError<FCompiledBindingLibraryCompiler::FCompileResult, FText> CompileResult = BindingLibraryCompiler.Compile();
+	TValueOrError<FCompiledBindingLibraryCompiler::FCompileResult, FText> CompileResult = BindingLibraryCompiler.Compile(BlueprintView->GetCompiledBindingLibraryId());
 	if (CompileResult.HasError())
 	{
 		WidgetBlueprintCompilerContext.MessageLog.Error(*FText::Format(LOCTEXT("BindingCompilationFailed", "The binding compilation failed. {1}"), CompileResult.GetError()).ToString());
diff --git a/Engine/Plugins/Runtime/ModelViewViewModel/Source/ModelViewViewModelBlueprint/Private/MVVMWidgetBlueprintExtension_View.cpp b/Engine/Plugins/Runtime/ModelViewViewModel/Source/ModelViewViewModelBlueprint/Private/MVVMWidgetBlueprintExtension_View.cpp
--- a/Engine/Plugins/Runtime/ModelViewViewModel/Source/ModelViewViewModelBlueprint/Private/MVVMWidgetBlueprintExtension_View.cpp
+++ b/Engine/Plugins/Runtime/ModelViewViewModel/Source/ModelViewViewModelBlueprint/Private/MVVMWidgetBlueprintExtension_View.cpp
@@ -130,6 +130,14 @@ void UMVVMWidgetBlueprintExtension_View::HandleFinishCompilingClass(UWidgetBluep
 		bool bCompiled = false;
 		if (CurrentCompilerContext->PreCompile(Class, BlueprintView))
 		{
+			FName ClassName = "ViewClass";
+			if (UObject* PreviousObj = StaticFindObjectFastInternal(nullptr, Class, ClassName, true))
+			{
+				// Remove previous object.
+				ERenameFlags RenameFlags = REN_ForceNoResetLoaders | REN_NonTransactional | REN_DoNotDirty | REN_DontCreateRedirectors;
+				FName TrashName = MakeUniqueObjectName(GetTransientPackage(), PreviousObj->GetClass(), *FString::Printf(TEXT("TRASH_%s"), *PreviousObj->GetName()));
+				PreviousObj->Rename(*TrashName.ToString(), GetTransientPackage(), RenameFlags);
+			}
 			ViewExtension = NewObject<UMVVMViewClass>(Class);
 			bCompiled = CurrentCompilerContext->Compile(Class, BlueprintView, ViewExtension);
 		}
diff --git a/Engine/Plugins/Runtime/ModelViewViewModel/Source/ModelViewViewModelBlueprint/Public/Bindings/MVVMCompiledBindingLibraryCompiler.h b/Engine/Plugins/Runtime/ModelViewViewModel/Source/ModelViewViewModelBlueprint/Public/Bindings/MVVMCompiledBindingLibraryCompiler.h
--- a/Engine/Plugins/Runtime/ModelViewViewModel/Source/ModelViewViewModelBlueprint/Public/Bindings/MVVMCompiledBindingLibraryCompiler.h
+++ b/Engine/Plugins/Runtime/ModelViewViewModel/Source/ModelViewViewModelBlueprint/Public/Bindings/MVVMCompiledBindingLibraryCompiler.h
@@ -175,6 +175,7 @@ public:
 
 	struct FCompileResult
 	{
+		FCompileResult(FGuid LibraryId);
 		FMVVMCompiledBindingLibrary Library;
 		TMap<FFieldPathHandle, FMVVMVCompiledFieldPath> FieldPaths;
 		TMap<FBindingHandle, FMVVMVCompiledBinding> Bindings;
@@ -182,7 +183,7 @@ public:
 	};
 
 	/** */
-	TValueOrError<FCompileResult, FText> Compile();
+	TValueOrError<FCompileResult, FText> Compile(FGuid LibraryId);
 
 private:
 	/** */
diff --git a/Engine/Plugins/Runtime/ModelViewViewModel/Source/ModelViewViewModelBlueprint/Public/MVVMBlueprintView.h b/Engine/Plugins/Runtime/ModelViewViewModel/Source/ModelViewViewModelBlueprint/Public/MVVMBlueprintView.h
--- a/Engine/Plugins/Runtime/ModelViewViewModel/Source/ModelViewViewModelBlueprint/Public/MVVMBlueprintView.h
+++ b/Engine/Plugins/Runtime/ModelViewViewModel/Source/ModelViewViewModelBlueprint/Public/MVVMBlueprintView.h
@@ -135,6 +135,11 @@ public:
 	void AddMessageToBinding(FGuid Id, UE::MVVM::FBindingMessage MessageToAdd);
 	void ResetBindingMessages();
 
+	FGuid GetCompiledBindingLibraryId() const
+	{
+		return CompiledBindingLibraryId;
+	}
+
 #if WITH_EDITOR
 	virtual void PostLoad() override;
 	virtual void PreSave(FObjectPreSaveContext Context) override;
@@ -173,6 +178,9 @@ private:
 	UPROPERTY(EditAnywhere, Category = "Viewmodel")
 	TArray<FMVVMBlueprintViewModelContext> AvailableViewModels;
 
+	UPROPERTY(VisibleAnywhere, Category = "Viewmodel", meta = (IgnoreForMemberInitializationTest))
+	FGuid CompiledBindingLibraryId;
+
 	TMap<FGuid, TArray<UE::MVVM::FBindingMessage>> BindingMessages;
 
 	bool bIsContextSensitive;
