MVVM: Add callback when a view source changed.

diff --git a/Engine/Plugins/Runtime/ModelViewViewModel/Source/ModelViewViewModel/Private/Debugging/MVVMDebugging.cpp b/Engine/Plugins/Runtime/ModelViewViewModel/Source/ModelViewViewModel/Private/Debugging/MVVMDebugging.cpp
--- a/Engine/Plugins/Runtime/ModelViewViewModel/Source/ModelViewViewModel/Private/Debugging/MVVMDebugging.cpp
+++ b/Engine/Plugins/Runtime/ModelViewViewModel/Source/ModelViewViewModel/Private/Debugging/MVVMDebugging.cpp
@@ -39,6 +39,19 @@ void FDebugging::BroadcastViewBeginDestruction(const UMVVMView* InView)
 	OnViewBeginDestruction.Broadcast(InView, FViewBeginDestructionArgs());
 }
 
+/** */
+FDebugging::FViewSourceValueArgs::FViewSourceValueArgs(const FMVVMViewSource& InViewSource)
+	: ViewSource(InViewSource)
+{
+}
+
+FDebugging::FViewSourceValueChanged FDebugging::OnViewSourceValueChanged;
+
+void FDebugging::BroadcastViewSourceValueChanged(const UMVVMView* InView, const FMVVMViewSource& ViewSource)
+{
+	OnViewSourceValueChanged.Broadcast(InView, FViewSourceValueArgs(ViewSource));
+}
+
 /** */
 FDebugging::FLibraryBindingRegisteredArgs::FLibraryBindingRegisteredArgs(const FMVVMViewClass_CompiledBinding& InBinding, ERegisterLibraryBindingResult InResult)
 	: Binding(InBinding)
diff --git a/Engine/Plugins/Runtime/ModelViewViewModel/Source/ModelViewViewModel/Private/View/MVVMView.cpp b/Engine/Plugins/Runtime/ModelViewViewModel/Source/ModelViewViewModel/Private/View/MVVMView.cpp
--- a/Engine/Plugins/Runtime/ModelViewViewModel/Source/ModelViewViewModel/Private/View/MVVMView.cpp
+++ b/Engine/Plugins/Runtime/ModelViewViewModel/Source/ModelViewViewModel/Private/View/MVVMView.cpp
@@ -59,6 +59,10 @@ void UMVVMView::Construct()
 	InitializeEvents();
 
 	bConstructed = true;
+
+#if UE_WITH_MVVM_DEBUGGING
+	UE::MVVM::FDebugging::BroadcastViewConstructed(this);
+#endif
 }
 
 
@@ -114,15 +118,15 @@ void UMVVMView::InitializeSources()
 				}
 			}
 			ViewSource.Source = NewSource;
+
+#if UE_WITH_MVVM_DEBUGGING
+			UE::MVVM::FDebugging::BroadcastViewSourceValueChanged(this, ViewSource);
+#endif
 		}
 	}
 
 	bSourcesInitialized = true;
 
-#if UE_WITH_MVVM_DEBUGGING
-	UE::MVVM::FDebugging::BroadcastViewConstructed(this);
-#endif
-
 	if (ClassExtension->InitializeBindingsOnConstruct())
 	{
 		InitializeBindings();
@@ -210,30 +214,34 @@ void UMVVMView::UninitializeInternal(bool bUninitializeSources)
 	for (int32 Index = 0; Index < AllViewModelCreators.Num(); ++Index)
 	{
 		const FMVVMViewClass_SourceCreator& Item = AllViewModelCreators[Index];
-		FMVVMViewSource& Source = Sources[Index];
+		FMVVMViewSource& ViewSource = Sources[Index];
 
-		if (Source.RegisteredCount > 0 && Source.Source && bUninitializeBindings)
+		if (ViewSource.RegisteredCount > 0 && ViewSource.Source && bUninitializeBindings)
 		{
-			TScriptInterface<INotifyFieldValueChanged> SourceAsInterface = Source.Source;
+			TScriptInterface<INotifyFieldValueChanged> SourceAsInterface = ViewSource.Source;
 			checkf(SourceAsInterface.GetInterface(), TEXT("It was added as a INotifyFieldValueChanged. It should still be."));
 			SourceAsInterface->RemoveAllFieldValueChangedDelegates(this);
 		}
 
 		// For GC release any object used by the view
-		if (!Source.bSetManually && bUninitializeSources)
+		if (!ViewSource.bSetManually && bUninitializeSources)
 		{
-			Item.DestroyInstance(Source.Source, this);
+			Item.DestroyInstance(ViewSource.Source, this);
 
-			Source.Source = nullptr;
-			if (Source.bAssignedToUserWidgetProperty)
+			ViewSource.Source = nullptr;
+			if (ViewSource.bAssignedToUserWidgetProperty)
 			{
-				FObjectPropertyBase* FoundObjectProperty = FindFProperty<FObjectPropertyBase>(UserWidget->GetClass(), Source.SourceName);
+				FObjectPropertyBase* FoundObjectProperty = FindFProperty<FObjectPropertyBase>(UserWidget->GetClass(), ViewSource.SourceName);
 				if (ensureAlwaysMsgf(FoundObjectProperty, TEXT("The compiler should have added the property")))
 				{
 					FoundObjectProperty->SetObjectPropertyValue_InContainer(UserWidget, nullptr);
 				}
-				Source.bAssignedToUserWidgetProperty = false;
+				ViewSource.bAssignedToUserWidgetProperty = false;
 			}
+
+#if UE_WITH_MVVM_DEBUGGING
+			UE::MVVM::FDebugging::BroadcastViewSourceValueChanged(this, ViewSource);
+#endif
 		}
 	}
 
@@ -496,6 +504,10 @@ bool UMVVMView::SetSourceInternal(FName ViewModelName, TScriptInterface<INotifyF
 				GEngine->GetEngineSubsystem<UMVVMBindingSubsystem>()->RemoveViewWithEveryTickBinding(this);
 			}
 		}
+
+#if UE_WITH_MVVM_DEBUGGING
+		UE::MVVM::FDebugging::BroadcastViewSourceValueChanged(this, ViewSource);
+#endif
 	}
 	return true;
 }
diff --git a/Engine/Plugins/Runtime/ModelViewViewModel/Source/ModelViewViewModel/Public/Debugging/MVVMDebugging.h b/Engine/Plugins/Runtime/ModelViewViewModel/Source/ModelViewViewModel/Public/Debugging/MVVMDebugging.h
--- a/Engine/Plugins/Runtime/ModelViewViewModel/Source/ModelViewViewModel/Public/Debugging/MVVMDebugging.h
+++ b/Engine/Plugins/Runtime/ModelViewViewModel/Source/ModelViewViewModel/Public/Debugging/MVVMDebugging.h
@@ -15,6 +15,7 @@
 
 class UMVVMView;
 struct FMVVMViewClass_CompiledBinding;
+struct FMVVMViewSource;
 class UUserWidget;
 
 namespace UE::MVVM
@@ -64,6 +65,19 @@ public:
 	static FViewDestructing OnViewBeginDestruction;
 	static void BroadcastViewBeginDestruction(const UMVVMView* View);
 
+public:
+	struct FViewSourceValueArgs
+	{
+		FViewSourceValueArgs() = delete;
+		FViewSourceValueArgs(const FMVVMViewSource& ViewSource);
+		const FMVVMViewSource& ViewSource;
+	};
+
+	DECLARE_EVENT_TwoParams(FDebugging, FViewSourceValueChanged, const FView&, const FViewSourceValueArgs&);
+	/** Broadcast when a source changes. */
+	static FViewSourceValueChanged OnViewSourceValueChanged;
+	static void BroadcastViewSourceValueChanged(const UMVVMView* View, const FMVVMViewSource& ViewSource);
+
 public:
 	enum class ERegisterLibraryBindingResult : uint8
 	{
diff --git a/Engine/Plugins/Runtime/ModelViewViewModel/Source/ModelViewViewModelEditor/Private/Widgets/SMVVMPreviewSourcePanel.cpp b/Engine/Plugins/Runtime/ModelViewViewModel/Source/ModelViewViewModelEditor/Private/Widgets/SMVVMPreviewSourcePanel.cpp
--- a/Engine/Plugins/Runtime/ModelViewViewModel/Source/ModelViewViewModelEditor/Private/Widgets/SMVVMPreviewSourcePanel.cpp
+++ b/Engine/Plugins/Runtime/ModelViewViewModel/Source/ModelViewViewModelEditor/Private/Widgets/SMVVMPreviewSourcePanel.cpp
@@ -69,6 +69,7 @@ void SPreviewSourcePanel::Construct(const FArguments& InArgs, TSharedPtr<FWidget
 		HandlePreviewWidgetChanged();
 		Context->OnPreviewWidgetChanged().AddSP(this, &SPreviewSourcePanel::HandlePreviewWidgetChanged);
 		Context->OnSelectedObjectChanged().AddSP(this, &SPreviewSourcePanel::HandleSelectedObjectChanged);
+		FDebugging::OnViewSourceValueChanged.AddSP(this, &SPreviewSourcePanel::HandleViewChanged);
 	}
 
 	ChildSlot
@@ -90,6 +91,7 @@ void SPreviewSourcePanel::Construct(const FArguments& InArgs, TSharedPtr<FWidget
 void SPreviewSourcePanel::HandlePreviewWidgetChanged()
 {
 	SourceList.Reset();
+	WeakView.Reset();
 
 	if (TSharedPtr<FWidgetBlueprintEditor> Editor = WeakEditor.Pin())
 	{
@@ -97,6 +99,7 @@ void SPreviewSourcePanel::HandlePreviewWidgetChanged()
 		{
 			if (UMVVMView* View = UMVVMSubsystem::GetViewFromUserWidget(NewWidget))
 			{
+				WeakView = View;
 				for (const FMVVMViewSource& Source : View->GetSources())
 				{
 					SourceList.Emplace(MakeShared<Private::SPreviewSourceEntry>(Source.Source, Source.SourceName));
@@ -169,10 +172,24 @@ void SPreviewSourcePanel::HandleSourceSelectionChanged(TSharedPtr<Private::SPrev
 }
 
 
+void SPreviewSourcePanel::HandleViewChanged(const FDebugging::FView& View, const FDebugging::FViewSourceValueArgs& Args)
+{
+	if (SourceListView)
+	{
+		if (View.GetView() == WeakView.Get())
+		{
+			SourceListView->RebuildList(); // to prevent access to invalid class, rebuild everything.
+		}
+	}
+}
+
+
 TSharedRef<ITableRow> SPreviewSourcePanel::GenerateWidget(TSharedPtr<Private::SPreviewSourceEntry> Entry, const TSharedRef<STableViewBase>& OwnerTable) const
 {
 	typedef STableRow<TSharedPtr<Private::SPreviewSourceEntry>> RowType;
 
+	TSharedRef<SWidget> FieldIcon = Entry->GetClass() ? SNew(UE::PropertyViewer::SFieldIcon, Entry->GetClass()) : SNullWidget::NullWidget;
+
 	TSharedRef<RowType> NewRow = SNew(RowType, OwnerTable);
 	NewRow->SetContent(SNew(SHorizontalBox)
 		+ SHorizontalBox::Slot()
@@ -180,7 +197,7 @@ TSharedRef<ITableRow> SPreviewSourcePanel::GenerateWidget(TSharedPtr<Private::SP
 		.HAlign(HAlign_Right)
 		.VAlign(VAlign_Center)
 		[
-			SNew(UE::PropertyViewer::SFieldIcon, Entry->GetClass())
+			FieldIcon
 		]
 		+ SHorizontalBox::Slot()
 		.Padding(4.0f)
diff --git a/Engine/Plugins/Runtime/ModelViewViewModel/Source/ModelViewViewModelEditor/Private/Widgets/SMVVMPreviewSourcePanel.h b/Engine/Plugins/Runtime/ModelViewViewModel/Source/ModelViewViewModelEditor/Private/Widgets/SMVVMPreviewSourcePanel.h
--- a/Engine/Plugins/Runtime/ModelViewViewModel/Source/ModelViewViewModelEditor/Private/Widgets/SMVVMPreviewSourcePanel.h
+++ b/Engine/Plugins/Runtime/ModelViewViewModel/Source/ModelViewViewModelEditor/Private/Widgets/SMVVMPreviewSourcePanel.h
@@ -2,11 +2,13 @@
 
 #pragma once
 
+#include "Debugging/MVVMDebugging.h"
 #include "Widgets/SCompoundWidget.h"
 
 class STableViewBase;
 class ITableRow;
 class FWidgetBlueprintEditor;
+class UMVVMView;
 class UUserWidget;
 
 namespace UE::MVVM::Private { class SPreviewSourceView; }
@@ -30,8 +32,10 @@ private:
 	void HandlePreviewWidgetChanged();
 	void HandleSelectedObjectChanged();
 	void HandleSourceSelectionChanged(TSharedPtr<Private::SPreviewSourceEntry> Entry, ESelectInfo::Type SelectionType) const;
+	void HandleViewChanged(const FDebugging::FView&, const FDebugging::FViewSourceValueArgs&);
 	TSharedRef<ITableRow> GenerateWidget(TSharedPtr<Private::SPreviewSourceEntry> Entry, const TSharedRef<STableViewBase>& OwnerTable) const;
 
+	TWeakObjectPtr<UMVVMView> WeakView;
 	TWeakPtr<FWidgetBlueprintEditor> WeakEditor;
 	TSharedPtr<Private::SPreviewSourceView> SourceListView;
 	TArray<TSharedPtr<Private::SPreviewSourceEntry>> SourceList;
