MVVM: Add instanced viewmodel. A way to create your viewmodel inside the view. The viewmodel is private and can only be used by the view. The feature is enabled by a setting. Change how the customization for the viewmodel context hide properties. Add a new submenu for the viewmodel panel.

diff --git a/Engine/Plugins/Runtime/ModelViewViewModel/ModelViewViewModel.uplugin b/Engine/Plugins/Runtime/ModelViewViewModel/ModelViewViewModel.uplugin
--- a/Engine/Plugins/Runtime/ModelViewViewModel/ModelViewViewModel.uplugin
+++ b/Engine/Plugins/Runtime/ModelViewViewModel/ModelViewViewModel.uplugin
@@ -41,5 +41,13 @@
 			"Type": "Editor",
 			"LoadingPhase": "Default"
 		}
+	],
+
+	"Plugins": [
+		{
+			"Name": "StructUtils",
+			"Enabled": true
+		}
 	]
 }
+
diff --git a/Engine/Plugins/Runtime/ModelViewViewModel/Source/ModelViewViewModel/Private/ViewModel/MVVMInstancedViewModelGeneratedClass.cpp b/Engine/Plugins/Runtime/ModelViewViewModel/Source/ModelViewViewModel/Private/ViewModel/MVVMInstancedViewModelGeneratedClass.cpp
new file mode 100644
--- /dev/null
+++ b/Engine/Plugins/Runtime/ModelViewViewModel/Source/ModelViewViewModel/Private/ViewModel/MVVMInstancedViewModelGeneratedClass.cpp
@@ -0,0 +1,75 @@
+// Copyright Epic Games, Inc. All Rights Reserved.
+
+#include "ViewModel/MVVMInstancedViewModelGeneratedClass.h"
+
+#include "INotifyFieldValueChanged.h"
+
+#include UE_INLINE_GENERATED_CPP_BY_NAME(MVVMInstancedViewModelGeneratedClass)
+
+#if WITH_EDITOR
+UClass* UMVVMInstancedViewModelGeneratedClass::GetAuthoritativeClass()
+{
+	return this;
+}
+#endif
+
+void UMVVMInstancedViewModelGeneratedClass::Link(FArchive& Ar, bool bRelinkExistingProperties)
+{
+	Super::Link(Ar, bRelinkExistingProperties);
+
+	for (TFieldIterator<FProperty> PropertyIter(this, EFieldIteratorFlags::ExcludeSuper); PropertyIter; ++PropertyIter)
+	{
+		if (!PropertyIter->RepNotifyFunc.IsNone())
+		{
+			if (!NativeFunctionLookupTable.ContainsByPredicate([ToFind = PropertyIter->RepNotifyFunc](const FNativeFunctionLookup& Other) { return Other.Name == ToFind; }))
+			{
+				NativeFunctionLookupTable.Emplace(PropertyIter->RepNotifyFunc, &UMVVMInstancedViewModelGeneratedClass::K2_CallNativeOnRep);
+			}
+			//NewFunction->RPCId = Params.RPCId;
+			//NewFunction->RPCResponseId = Params.RPCResponseId;
+		}
+	}
+}
+
+DEFINE_FUNCTION(UMVVMInstancedViewModelGeneratedClass::K2_CallNativeOnRep)
+{
+	UObject* CallingObject = P_THIS_OBJECT;
+	if (CallingObject == nullptr || !CallingObject->GetClass()->ImplementsInterface(UNotifyFieldValueChanged::StaticClass()))
+	{
+		return;
+	}
+	UMVVMInstancedViewModelGeneratedClass* GeneratedClass = Cast<UMVVMInstancedViewModelGeneratedClass>(CallingObject->GetClass());
+	ensure(GeneratedClass);
+	if (!GeneratedClass)
+	{
+		return;
+	}
+
+	FName PropertyName;
+	{
+		const UFunction* Func = Stack.CurrentNativeFunction;
+		FString FunctionName = Func->GetName();
+		if (FunctionName.RemoveFromStart(TEXT("__OnRep_")))
+		{
+			PropertyName = FName(*FunctionName, EFindName::FNAME_Find);
+			if (!PropertyName.IsNone())
+			{
+				FProperty* FoundProperty = CallingObject->GetClass()->FindPropertyByName(PropertyName);
+				if (FoundProperty)
+				{
+					GeneratedClass->OnPropertyReplicated(CallingObject, FoundProperty);
+				}
+			}
+		}
+	}
+}
+
+void UMVVMInstancedViewModelGeneratedClass::OnPropertyReplicated(UObject* Object, const FProperty* Property)
+{
+	TScriptInterface<INotifyFieldValueChanged> CallingInterface = Object;
+	UE::FieldNotification::FFieldId FieldId = CallingInterface->GetFieldNotificationDescriptor().GetField(Object->GetClass(), Property->GetFName());
+	if (FieldId.IsValid())
+	{
+		CallingInterface->BroadcastFieldValueChanged(FieldId);
+	}
+}
diff --git a/Engine/Plugins/Runtime/ModelViewViewModel/Source/ModelViewViewModel/Public/ViewModel/MVVMInstancedViewModelGeneratedClass.h b/Engine/Plugins/Runtime/ModelViewViewModel/Source/ModelViewViewModel/Public/ViewModel/MVVMInstancedViewModelGeneratedClass.h
new file mode 100644
--- /dev/null
+++ b/Engine/Plugins/Runtime/ModelViewViewModel/Source/ModelViewViewModel/Public/ViewModel/MVVMInstancedViewModelGeneratedClass.h
@@ -0,0 +1,28 @@
+// Copyright Epic Games, Inc. All Rights Reserved.
+
+#pragma once
+
+#include "Engine/BlueprintGeneratedClass.h"
+#include "UObject/Package.h"
+#include "MVVMInstancedViewModelGeneratedClass.generated.h"
+
+/**
+ *
+ */
+UCLASS()
+class MODELVIEWVIEWMODEL_API UMVVMInstancedViewModelGeneratedClass : public UBlueprintGeneratedClass
+{
+	GENERATED_BODY()
+
+public:
+	virtual void Link(FArchive& Ar, bool bRelinkExistingProperties) override;
+	
+#if WITH_EDITOR
+	virtual UClass* GetAuthoritativeClass() override;
+#endif
+
+public:
+	DECLARE_FUNCTION(K2_CallNativeOnRep);
+
+	virtual void OnPropertyReplicated(UObject* Object, const FProperty* Property);
+};
diff --git a/Engine/Plugins/Runtime/ModelViewViewModel/Source/ModelViewViewModelBlueprint/ModelViewViewModelBlueprint.Build.cs b/Engine/Plugins/Runtime/ModelViewViewModel/Source/ModelViewViewModelBlueprint/ModelViewViewModelBlueprint.Build.cs
--- a/Engine/Plugins/Runtime/ModelViewViewModel/Source/ModelViewViewModelBlueprint/ModelViewViewModelBlueprint.Build.cs
+++ b/Engine/Plugins/Runtime/ModelViewViewModel/Source/ModelViewViewModelBlueprint/ModelViewViewModelBlueprint.Build.cs
@@ -17,6 +17,7 @@ public class ModelViewViewModelBlueprint : ModuleRules
 				"Engine",
 				"FieldNotification",
 				"ModelViewViewModel",
+				"StructUtils",
 			});
 
 		PrivateDependencyModuleNames.AddRange(
diff --git a/Engine/Plugins/Runtime/ModelViewViewModel/Source/ModelViewViewModelBlueprint/Private/MVVMBlueprintInstancedViewModel.cpp b/Engine/Plugins/Runtime/ModelViewViewModel/Source/ModelViewViewModelBlueprint/Private/MVVMBlueprintInstancedViewModel.cpp
new file mode 100644
--- /dev/null
+++ b/Engine/Plugins/Runtime/ModelViewViewModel/Source/ModelViewViewModelBlueprint/Private/MVVMBlueprintInstancedViewModel.cpp
@@ -0,0 +1,306 @@
+// Copyright Epic Games, Inc. All Rights Reserved.
+
+#include "MVVMBlueprintInstancedViewModel.h"
+
+#include "MVVMBlueprintView.h"
+#include "MVVMViewModelBase.h"
+#include "MVVMWidgetBlueprintExtension_View.h"
+#include "ViewModel/MVVMInstancedViewModelGeneratedClass.h"
+
+#include "BlueprintActionDatabase.h"
+#include "Engine/Engine.h"
+#include "Kismet2/BlueprintEditorUtils.h"
+#include "Kismet2/KismetEditorUtilities.h"
+
+#include UE_INLINE_GENERATED_CPP_BY_NAME(MVVMBlueprintInstancedViewModel)
+
+/**
+ * 
+ */
+UMVVMBlueprintInstancedViewModelBase::UMVVMBlueprintInstancedViewModelBase()
+{
+	ParentClass = UMVVMViewModelBase::StaticClass();
+}
+
+void UMVVMBlueprintInstancedViewModelBase::GenerateClass()
+{
+	if (ParentClass.Get() == nullptr)
+	{
+		ParentClass = UMVVMViewModelBase::StaticClass();
+	}
+	if (!FKismetEditorUtilities::CanCreateBlueprintOfClass(ParentClass))
+	{
+		ParentClass = UMVVMViewModelBase::StaticClass();
+	}
+	if (GeneratedClassType.Get() == nullptr)
+	{
+		GeneratedClassType = UMVVMInstancedViewModelGeneratedClass::StaticClass();
+	}
+
+	auto SafeRename = [](UObject* Object)
+	{
+		ERenameFlags RenameFlags = REN_ForceNoResetLoaders | REN_NonTransactional | REN_DoNotDirty | REN_DontCreateRedirectors;
+		FName TrashName = MakeUniqueObjectName(GetTransientPackage(), Object->GetClass(), *FString::Printf(TEXT("TRASH_%s"), *Object->GetName()));
+		Object->Rename(*TrashName.ToString(), GetTransientPackage(), RenameFlags);
+	};
+
+	bool bGeneratedNewClass = false;
+	if (GeneratedClass == nullptr)
+	{
+		bGeneratedNewClass = true;
+		GeneratedClass = NewObject<UMVVMInstancedViewModelGeneratedClass>(GetOutermost(), GeneratedClassType.Get());
+	}
+
+	UObject* PreviousDefaultObject = GeneratedClass->GetDefaultObject(false);
+	if (PreviousDefaultObject)
+	{
+		SafeRename(PreviousDefaultObject);
+	}
+	for (TFieldIterator<UFunction> FunctionIter(GeneratedClass, EFieldIteratorFlags::ExcludeSuper); FunctionIter; ++FunctionIter)
+	{
+		FunctionIter->FunctionFlags &= ~FUNC_Native;
+		SafeRename(*FunctionIter);
+	}
+
+	// Clean class and reset basic properties
+	GeneratedClass->PurgeClass(false);
+	GeneratedClass->PropertyLink = ParentClass->PropertyLink;
+	GeneratedClass->SetSuperStruct(ParentClass);
+	GeneratedClass->ClassWithin = UObject::StaticClass();
+	GeneratedClass->ClassConfigName = ParentClass->ClassConfigName;
+	GeneratedClass->ClassFlags |= CLASS_NotPlaceable;
+
+	// Clean up temporary generate variables
+	{
+		FromPropertyToCreatedProperty.Empty();
+	}
+
+	// Create Properties
+	{
+		PreAddProperties();
+		for (TPropertyValueIterator<FProperty> It(GetSourceStruct(), GetSourceDefaults(), EPropertyValueIteratorFlags::NoRecursion, EFieldIteratorFlags::ExcludeDeprecated); It; ++It)
+		{
+			const FProperty* Property = It.Key();
+			void const* ValuePtr = It.Value();
+
+			AddProperty(Property);
+		}
+		PostAddProperties();
+	}
+
+	// Update the class
+	GeneratedClass->Bind();
+	GeneratedClass->StaticLink(true);
+	GeneratedClass->AssembleReferenceTokenStream();
+	ensure(GeneratedClass->ClassGeneratedBy == nullptr);
+	ensure(GeneratedClass->GetDefaultObject(false) == nullptr);
+	GeneratedClass->GetDefaultObject(true);
+	GeneratedClass->UpdateCustomPropertyListForPostConstruction();
+
+	// The class is not public and can only be access inside the UMG. What about inherited UMG???
+	GeneratedClass->ClearFlags(RF_Public | RF_Transactional);
+	GeneratedClass->GetDefaultObject(true)->ClearFlags(RF_Public);
+
+	// Relink the new default object
+	if (PreviousDefaultObject)
+	{
+		FLinkerLoad::PRIVATE_PatchNewObjectIntoExport(PreviousDefaultObject, GeneratedClass->GetDefaultObject(true));
+	}
+
+	// Initialize the default object value
+	{
+		PreSetDefaultValues();
+		for (TPropertyValueIterator<FProperty> It(GetSourceStruct(), GetSourceDefaults(), EPropertyValueIteratorFlags::NoRecursion, EFieldIteratorFlags::ExcludeDeprecated); It; ++It)
+		{
+			const FProperty* Property = It.Key();
+			void const* ValuePtr = It.Value();
+
+			SetDefaultValue(Property, ValuePtr);
+		}
+		PostSetDefaultValues();
+	}
+
+	// Initialize the fields for MVVM
+	GeneratedClass->InitializeFieldNotifies();
+
+	// Inform that the class changed
+	{
+		TMap<UObject*, UObject*> OldToNew;
+		if (PreviousDefaultObject)
+		{
+			OldToNew.Emplace(PreviousDefaultObject, GeneratedClass->GetDefaultObject(true));
+		}
+		if (bGeneratedNewClass)
+		{
+			OldToNew.Emplace(GeneratedClass, GeneratedClass);
+		}
+		if (GEngine && OldToNew.Num() > 0)
+		{
+			GEngine->NotifyToolsOfObjectReplacement(OldToNew);
+		}
+	}
+
+	// Rebuild the BP action
+	if (FBlueprintActionDatabase* ActionDB = FBlueprintActionDatabase::TryGet())
+	{
+		// Notify Blueprints that there is a new class to add to the action list
+		ActionDB->RefreshClassActions(GeneratedClass);
+	}
+
+	// Clean up temporary generate variables
+	{
+		FromPropertyToCreatedProperty.Empty();
+	}
+
+	//GeneratedClass->DestroyPropertiesPendingDestruction();
+}
+
+bool UMVVMBlueprintInstancedViewModelBase::IsValidFieldName(const FName NewPropertyName) const
+{
+	if (!FName::IsValidXName(NewPropertyName, INVALID_NAME_CHARACTERS))
+	{
+		return false;
+	}
+
+	// Check if the name already exist. If it does, do not add the property.
+	for (TFieldIterator<FField> PropertyIter(GeneratedClass, EFieldIteratorFlags::IncludeSuper); PropertyIter; ++PropertyIter)
+	{
+		if (PropertyIter->GetFName() == NewPropertyName)
+		{
+			return false;
+		}
+	}
+	for (TFieldIterator<UField> FunctionIter(GeneratedClass, EFieldIteratorFlags::IncludeSuper); FunctionIter; ++FunctionIter)
+	{
+		if (FunctionIter->GetFName() == NewPropertyName)
+		{
+			return false;
+		}
+	}
+
+	return true;
+}
+
+void UMVVMBlueprintInstancedViewModelBase::PreAddProperties()
+{
+
+}
+
+void UMVVMBlueprintInstancedViewModelBase::PostAddProperties()
+{
+}
+
+void UMVVMBlueprintInstancedViewModelBase::AddProperty(const FProperty* FromProperty)
+{
+	if (!IsValidFieldName(FromProperty->GetFName()))
+	{
+		return;
+	}
+
+	FProperty* NewProperty = CastFieldChecked<FProperty>(FField::Duplicate(FromProperty, GeneratedClass, FromProperty->GetFName()));
+#if WITH_EDITOR
+	FField::CopyMetaData(FromProperty, NewProperty);
+#endif
+	FInitializePropertyArgs Args;
+	Args.PropertyName = FromProperty->GetFName();
+	Args.DisplayName = FromProperty->GetMetaData(FBlueprintMetadata::MD_DisplayName);
+	Args.bNetwork = true;
+	InitializeProperty(NewProperty, Args);
+
+	LinkProperty(NewProperty);
+	FromPropertyToCreatedProperty.Add(FromProperty, NewProperty);
+}
+
+
+void UMVVMBlueprintInstancedViewModelBase::PreSetDefaultValues()
+{
+}
+
+void UMVVMBlueprintInstancedViewModelBase::PostSetDefaultValues()
+{
+}
+
+void UMVVMBlueprintInstancedViewModelBase::SetDefaultValue(const FProperty* Property, void const* ValuePtr)
+{
+	FProperty* NewProperty = nullptr;
+	if (FProperty** NewPropertyPtr = FromPropertyToCreatedProperty.Find(Property))
+	{
+		NewProperty = *NewPropertyPtr;
+	}
+	else
+	{
+		NewProperty = GeneratedClass->FindPropertyByName(Property->GetFName());
+	}
+
+	if (NewProperty)
+	{
+		//void* DestinationPtr = reinterpret_cast<uint8*>(GeneratedClass->GetDefaultObject()) + NewProperty->GetOffset_ForInternal();
+		void* DestinationPtr = NewProperty->ContainerPtrToValuePtr<void>(GeneratedClass->GetDefaultObject());
+		NewProperty->CopyCompleteValue(DestinationPtr, ValuePtr);
+	}
+}
+
+void UMVVMBlueprintInstancedViewModelBase::InitializeProperty(FProperty* NewProperty, FInitializePropertyArgs& Args)
+{
+	if (Args.bNetwork)
+	{
+		NewProperty->RepNotifyFunc = AddOnRepFunction(Args.PropertyName);
+	}
+
+	EPropertyFlags NewFlags = NewProperty->GetPropertyFlags() | EPropertyFlags::CPF_Edit;
+	NewFlags |= !NewProperty->RepNotifyFunc.IsNone() ? EPropertyFlags::CPF_Net | EPropertyFlags::CPF_RepNotify : EPropertyFlags::CPF_None;
+	NewFlags |= !Args.bPrivate ? EPropertyFlags::CPF_BlueprintVisible : EPropertyFlags::CPF_None;
+	NewFlags |= Args.bReadOnly ? EPropertyFlags::CPF_BlueprintReadOnly : EPropertyFlags::CPF_None;
+	NewProperty->SetPropertyFlags(NewFlags);
+	if (Args.bFieldNotify)
+	{
+		NewProperty->SetMetaData(FBlueprintMetadata::MD_FieldNotify, TEXT(""));
+		GeneratedClass->FieldNotifies.Add(FFieldNotificationId(Args.PropertyName));
+	}
+}
+
+void UMVVMBlueprintInstancedViewModelBase::LinkProperty(FProperty* NewProperty)
+{
+	NewProperty->Next = GeneratedClass->ChildProperties;
+	GeneratedClass->ChildProperties = NewProperty;
+}
+
+FName UMVVMBlueprintInstancedViewModelBase::AddOnRepFunction(FName PropertyName)
+{
+	FString OnRepCallFunctionName = FString::Printf(TEXT("__OnRep_%s"), *PropertyName.ToString());
+	FName Name_OnRepCallFunctionName = *OnRepCallFunctionName;
+	UObject* PreviousObj = StaticFindObjectFastInternal(nullptr, GeneratedClass, Name_OnRepCallFunctionName, true);
+	if (PreviousObj)
+	{
+		// The function or property already exist. Something is wrong.
+		return FName();
+	}
+	UFunction* Func = NewObject<UFunction>(GeneratedClass, Name_OnRepCallFunctionName);
+	Func->FunctionFlags |= FUNC_Native | FUNC_Event | FUNC_BlueprintEvent | FUNC_BlueprintCallable;
+	GeneratedClass->AddNativeFunction(*OnRepCallFunctionName, &UMVVMInstancedViewModelGeneratedClass::K2_CallNativeOnRep);
+	GeneratedClass->AddFunctionToFunctionMap(Func, Func->GetFName());
+	Func->Bind();
+	Func->StaticLink(true);
+	Func->Next = GeneratedClass->Children;
+	GeneratedClass->Children = Func;
+	return Func->GetFName();
+}
+
+#if WITH_EDITOR
+void UMVVMBlueprintInstancedViewModel::PostEditChangeChainProperty(FPropertyChangedChainEvent& PropertyChangedEvent)
+{
+	Super::PostEditChangeChainProperty(PropertyChangedEvent);
+
+	FProperty* VariableProperty = GetClass()->FindPropertyByName(GET_MEMBER_NAME_CHECKED(UMVVMBlueprintInstancedViewModel, Variables));
+	FEditPropertyChain::TDoubleLinkedListNode* CurrentPropertyNode = PropertyChangedEvent.PropertyChain.GetActiveMemberNode();
+	while (CurrentPropertyNode)
+	{
+		if (CurrentPropertyNode->GetValue() == VariableProperty)
+		{
+			FBlueprintEditorUtils::MarkBlueprintAsModified(GetOuterUMVVMBlueprintView()->GetOuterUMVVMWidgetBlueprintExtension_View()->GetWidgetBlueprint());
+			break;
+		}
+		CurrentPropertyNode = CurrentPropertyNode->GetNextNode();
+	}
+}
+#endif
\ No newline at end of file
diff --git a/Engine/Plugins/Runtime/ModelViewViewModel/Source/ModelViewViewModelBlueprint/Private/MVVMBlueprintView.cpp b/Engine/Plugins/Runtime/ModelViewViewModel/Source/ModelViewViewModelBlueprint/Private/MVVMBlueprintView.cpp
--- a/Engine/Plugins/Runtime/ModelViewViewModel/Source/ModelViewViewModelBlueprint/Private/MVVMBlueprintView.cpp
+++ b/Engine/Plugins/Runtime/ModelViewViewModel/Source/ModelViewViewModelBlueprint/Private/MVVMBlueprintView.cpp
@@ -4,6 +4,7 @@
 
 #include "Components/Widget.h"
 #include "Kismet2/BlueprintEditorUtils.h"
+#include "MVVMBlueprintInstancedViewModel.h"
 #include "MVVMBlueprintViewConversionFunction.h"
 #include "MVVMBlueprintViewEvent.h"
 #include "MVVMWidgetBlueprintExtension_View.h"
@@ -47,20 +48,47 @@ void UMVVMBlueprintView::AddViewModel(const FMVVMBlueprintViewModelContext& NewC
 	OnViewModelsUpdated.Broadcast();
 }
 
-
-bool UMVVMBlueprintView::RemoveViewModel(FGuid ViewModelId)
+namespace UE::MVVM::Private
 {
-	int32 Count = AvailableViewModels.RemoveAll([ViewModelId](const FMVVMBlueprintViewModelContext& VM)
+	bool RemoveViewModelInternal(FGuid ViewModelId, TArray<FMVVMBlueprintViewModelContext>& ViewModelContexts)
+	{
+		bool bResult = false;
+		for (int32 Index = ViewModelContexts.Num() - 1; Index >= 0; --Index)
 		{
-			return VM.GetViewModelId() == ViewModelId;
-		});
+			if (ViewModelContexts[Index].GetViewModelId() == ViewModelId)
+			{
+				UMVVMBlueprintInstancedViewModelBase* InstancedViewModel = ViewModelContexts[Index].InstancedViewModel;
+				if (InstancedViewModel)
+				{
+					auto RenameToTransient = [](UObject* ObjectToRename)
+					{
+						FName TrashName = MakeUniqueObjectName(GetTransientPackage(), ObjectToRename->GetClass(), *FString::Printf(TEXT("TRASH_%s"), *ObjectToRename->GetName()));
+						ObjectToRename->Rename(*TrashName.ToString(), GetTransientPackage());
+					};
+					if (InstancedViewModel->GetGeneratedClass())
+					{
+						RenameToTransient(InstancedViewModel->GetGeneratedClass());
+					}
+					RenameToTransient(InstancedViewModel);
+				}
 
-	if (Count > 0)
+				ViewModelContexts.RemoveAt(Index);
+				bResult = true;
+			}
+		}
+		return bResult;
+	}
+}
+
+bool UMVVMBlueprintView::RemoveViewModel(FGuid ViewModelId)
+{
+	bool bRemoved = UE::MVVM::Private::RemoveViewModelInternal(ViewModelId, AvailableViewModels);
+	if (bRemoved)
 	{
 		FBlueprintEditorUtils::MarkBlueprintAsStructurallyModified(GetOuterUMVVMWidgetBlueprintExtension_View()->GetWidgetBlueprint());
 		OnViewModelsUpdated.Broadcast();
 	}
-	return Count > 0;
+	return bRemoved;
 }
 
 int32 UMVVMBlueprintView::RemoveViewModels(const TArrayView<FGuid> ViewModelIds)
@@ -68,10 +96,10 @@ int32 UMVVMBlueprintView::RemoveViewModels(const TArrayView<FGuid> ViewModelIds)
 	int32 Count = 0;
 	for (const FGuid& ViewModelId : ViewModelIds)
 	{
-		Count += AvailableViewModels.RemoveAll([ViewModelId](const FMVVMBlueprintViewModelContext& VM)
-			{
-				return VM.GetViewModelId() == ViewModelId;
-			});
+		if (UE::MVVM::Private::RemoveViewModelInternal(ViewModelId, AvailableViewModels))
+		{
+			++Count;
+		}
 	}
 
 	if (Count > 0)
@@ -131,6 +159,8 @@ void UMVVMBlueprintView::RemoveBindingAt(int32 Index)
 
 		Bindings.RemoveAt(Index);
 		OnBindingsUpdated.Broadcast();
+
+		FBlueprintEditorUtils::MarkBlueprintAsModified(GetOuterUMVVMWidgetBlueprintExtension_View()->GetWidgetBlueprint());
 	}
 }
 
@@ -157,6 +187,9 @@ FMVVMBlueprintViewBinding& UMVVMBlueprintView::AddBinding(const UWidget* Widget,
 
 	OnBindingsAdded.Broadcast();
 	OnBindingsUpdated.Broadcast();
+
+	FBlueprintEditorUtils::MarkBlueprintAsModified(GetOuterUMVVMWidgetBlueprintExtension_View()->GetWidgetBlueprint());
+
 	return NewBinding;
 }
 
@@ -167,6 +200,9 @@ FMVVMBlueprintViewBinding& UMVVMBlueprintView::AddDefaultBinding()
 
 	OnBindingsAdded.Broadcast();
 	OnBindingsUpdated.Broadcast();
+
+	FBlueprintEditorUtils::MarkBlueprintAsModified(GetOuterUMVVMWidgetBlueprintExtension_View()->GetWidgetBlueprint());
+
 	return NewBinding;
 }
 
diff --git a/Engine/Plugins/Runtime/ModelViewViewModel/Source/ModelViewViewModelBlueprint/Private/MVVMViewBlueprintCompiler.cpp b/Engine/Plugins/Runtime/ModelViewViewModel/Source/ModelViewViewModelBlueprint/Private/MVVMViewBlueprintCompiler.cpp
--- a/Engine/Plugins/Runtime/ModelViewViewModel/Source/ModelViewViewModelBlueprint/Private/MVVMViewBlueprintCompiler.cpp
+++ b/Engine/Plugins/Runtime/ModelViewViewModel/Source/ModelViewViewModelBlueprint/Private/MVVMViewBlueprintCompiler.cpp
@@ -312,10 +312,10 @@ void FMVVMViewBlueprintCompiler::CreateVariables(const FWidgetBlueprintCompilerC
 			{
 				NewProperty->SetMetaData(FBlueprintMetadata::MD_ExposeOnSpawn, TEXT("true"));
 			}
-			//if (!SourceContext.bPublicGetter)
-			//{
-			//	NewProperty->SetMetaData(FBlueprintMetadata::MD_Private, TEXT("true"));
-			//}
+			if (SourceContext.bPrivate)
+			{
+				NewProperty->SetMetaData(FBlueprintMetadata::MD_Private, TEXT("true"));
+			}
 #endif
 		}
 		return NewProperty;
@@ -471,7 +471,7 @@ void FMVVMViewBlueprintCompiler::CreateSourceLists(const FWidgetBlueprintCompile
 		SourceVariable.DisplayName = ViewModelContext.GetDisplayName();
 		SourceVariable.CategoryName = TEXT("Viewmodel");
 		SourceVariable.bExposeOnSpawn = bCreateSetterFunction;
-		//SourceVariable.bPublicGetter = ViewModelContext.bCreateGetterFunction;
+		SourceVariable.bPrivate = ViewModelContext.InstancedViewModel != nullptr;
 		SourceVariable.BlueprintSetter = CompilerSourceCreatorContexts[FoundSourceCreatorContextIndex].SetterFunctionName;
 		SourceVariable.ViewModelId = ViewModelContext.GetViewModelId();
 		CompilerUserWidgetPropertyContexts.Emplace(MoveTemp(SourceVariable));
@@ -510,7 +510,7 @@ void FMVVMViewBlueprintCompiler::CreateSourceLists(const FWidgetBlueprintCompile
 				SourceVariable.DisplayName = FText::FromString(Widget->GetDisplayLabel());
 				SourceVariable.CategoryName = TEXT("Widget");
 				SourceVariable.ViewModelId = FGuid();
-				SourceVariable.bPublicGetter = false;
+				SourceVariable.bPrivate = false;
 				Self->CompilerUserWidgetPropertyContexts.Emplace(MoveTemp(SourceVariable));
 			}
 		}
diff --git a/Engine/Plugins/Runtime/ModelViewViewModel/Source/ModelViewViewModelBlueprint/Private/MVVMViewBlueprintCompiler.h b/Engine/Plugins/Runtime/ModelViewViewModel/Source/ModelViewViewModelBlueprint/Private/MVVMViewBlueprintCompiler.h
--- a/Engine/Plugins/Runtime/ModelViewViewModel/Source/ModelViewViewModelBlueprint/Private/MVVMViewBlueprintCompiler.h
+++ b/Engine/Plugins/Runtime/ModelViewViewModel/Source/ModelViewViewModelBlueprint/Private/MVVMViewBlueprintCompiler.h
@@ -112,7 +112,7 @@ private:
 		FGuid ViewModelId;
 
 		bool bExposeOnSpawn = false;
-		bool bPublicGetter = false;
+		bool bPrivate = false;
 	};
 	TArray<FCompilerUserWidgetPropertyContext> CompilerUserWidgetPropertyContexts;
 
diff --git a/Engine/Plugins/Runtime/ModelViewViewModel/Source/ModelViewViewModelBlueprint/Private/MVVMWidgetBlueprintExtension_View.cpp b/Engine/Plugins/Runtime/ModelViewViewModel/Source/ModelViewViewModelBlueprint/Private/MVVMWidgetBlueprintExtension_View.cpp
--- a/Engine/Plugins/Runtime/ModelViewViewModel/Source/ModelViewViewModelBlueprint/Private/MVVMWidgetBlueprintExtension_View.cpp
+++ b/Engine/Plugins/Runtime/ModelViewViewModel/Source/ModelViewViewModelBlueprint/Private/MVVMWidgetBlueprintExtension_View.cpp
@@ -1,12 +1,17 @@
 // Copyright Epic Games, Inc. All Rights Reserved.
 
 #include "MVVMWidgetBlueprintExtension_View.h"
+
+#include "FindInBlueprintManager.h"
+#include "MVVMBlueprintInstancedViewModel.h"
 #include "MVVMBlueprintView.h"
+#include "MVVMBlueprintViewModel.h"
 #include "MVVMBlueprintViewConversionFunction.h"
 #include "MVVMViewBlueprintCompiler.h"
 #include "View/MVVMViewClass.h"
+#include "Kismet2/BlueprintEditorUtils.h"
+#include "Kismet2/KismetEditorUtilities.h"
 
-#include "FindInBlueprintManager.h"
 
 #include UE_INLINE_GENERATED_CPP_BY_NAME(MVVMWidgetBlueprintExtension_View)
 
@@ -27,6 +32,7 @@ void UMVVMWidgetBlueprintExtension_View::CreateBlueprintViewInstance()
 {
 	BlueprintView = NewObject<UMVVMBlueprintView>(this, FName(), RF_Transactional);
 	BlueprintViewChangedDelegate.Broadcast();
+	FBlueprintEditorUtils::MarkBlueprintAsModified(GetWidgetBlueprint());
 }
 
 
@@ -51,6 +57,14 @@ void UMVVMWidgetBlueprintExtension_View::HandlePreloadObjectsForCompilation(UBlu
 
 		for (const FMVVMBlueprintViewModelContext& AvailableViewModel : BlueprintView->GetViewModels())
 		{
+			if (AvailableViewModel.InstancedViewModel)
+			{
+				UBlueprint::ForceLoad(AvailableViewModel.InstancedViewModel);
+				if (AvailableViewModel.InstancedViewModel->GetGeneratedClass())
+				{
+					AvailableViewModel.InstancedViewModel->GetGeneratedClass()->ConditionalPostLoad();
+				}
+			}
 			if (AvailableViewModel.GetViewModelClass())
 			{
 				UBlueprint::ForceLoad(AvailableViewModel.GetViewModelClass());
@@ -73,6 +87,14 @@ void UMVVMWidgetBlueprintExtension_View::HandlePreloadObjectsForCompilation(UBlu
 
 void UMVVMWidgetBlueprintExtension_View::HandleBeginCompilation(FWidgetBlueprintCompilerContext& InCreationContext)
 {
+	for (const FMVVMBlueprintViewModelContext& AvailableViewModel : BlueprintView->GetViewModels())
+	{
+		if (AvailableViewModel.InstancedViewModel)
+		{
+			AvailableViewModel.InstancedViewModel->GenerateClass();
+		}
+	}
+
 	CurrentCompilerContext.Reset();
 	if (BlueprintView)
 	{
diff --git a/Engine/Plugins/Runtime/ModelViewViewModel/Source/ModelViewViewModelBlueprint/Public/MVVMBlueprintInstancedViewModel.h b/Engine/Plugins/Runtime/ModelViewViewModel/Source/ModelViewViewModelBlueprint/Public/MVVMBlueprintInstancedViewModel.h
new file mode 100644
--- /dev/null
+++ b/Engine/Plugins/Runtime/ModelViewViewModel/Source/ModelViewViewModelBlueprint/Public/MVVMBlueprintInstancedViewModel.h
@@ -0,0 +1,114 @@
+// Copyright Epic Games, Inc. All Rights Reserved.
+
+#pragma once
+
+#include "Engine/Blueprint.h"
+#include "Engine/BlueprintGeneratedClass.h"
+#include "PropertyBag.h"
+#include "Templates/SubclassOf.h"
+#include "ViewModel/MVVMInstancedViewModelGeneratedClass.h"
+#include "MVVMBlueprintInstancedViewModel.generated.h"
+
+/**
+ *
+ */
+UCLASS(Abstract, Within = MVVMBlueprintView)
+class MODELVIEWVIEWMODELBLUEPRINT_API UMVVMBlueprintInstancedViewModelBase : public UObject
+{
+	GENERATED_BODY()
+
+public:
+	UMVVMBlueprintInstancedViewModelBase();
+
+	void GenerateClass();
+
+	UClass* GetGeneratedClass() const
+	{
+		return GeneratedClass;
+	}
+
+	virtual const UStruct* GetSourceStruct() const
+	{
+		PURE_VIRTUAL(GetSourceStruct, return nullptr;);
+	}
+
+	virtual const uint8* GetSourceDefaults() const
+	{
+		PURE_VIRTUAL(GetSourceDefaults, return nullptr;);
+	}
+
+protected:
+	virtual void PreAddProperties();
+	virtual void AddProperty(const FProperty* FromProperty);
+	virtual void PostAddProperties();
+
+	virtual void PreSetDefaultValues();
+	virtual void SetDefaultValue(const FProperty* FromProperty, void const* SourceValuePtr);
+	virtual void PostSetDefaultValues();
+
+	virtual bool IsValidFieldName(const FName NewPropertyName) const;
+
+protected:
+	struct FInitializePropertyArgs
+	{
+		FName PropertyName;
+		FString DisplayName;
+		bool bFieldNotify = true;
+		bool bReadOnly = false;
+		bool bNetwork = false;
+		bool bPrivate = false;
+	};
+	void InitializeProperty(FProperty* NewProperty, FInitializePropertyArgs& Args);
+	void LinkProperty(FProperty* NewProperty);
+	FName AddOnRepFunction(FName PropertyName);
+
+	TMap<const FProperty*, FProperty*> FromPropertyToCreatedProperty;
+
+public:
+	/** The base object of the generated class. */
+	UPROPERTY(EditAnywhere, Category = "Viewmodel", meta=(AllowedClasses = "/Script/FieldNotification.NotifyFieldValueChanged", DisallowedClasses = "/Script/UMG.Widget"))
+	TSubclassOf<UObject> ParentClass;
+
+	UPROPERTY()
+	TObjectPtr<UMVVMInstancedViewModelGeneratedClass> GeneratedClass;
+
+protected:
+	//~ if it changes, the GeneratedClass needs to be removed.
+	UPROPERTY()
+	TSubclassOf<UMVVMInstancedViewModelGeneratedClass> GeneratedClassType;
+};
+
+/**
+ *
+ */
+UCLASS()
+class MODELVIEWVIEWMODELBLUEPRINT_API UMVVMBlueprintInstancedViewModel : public UMVVMBlueprintInstancedViewModelBase
+{
+	GENERATED_BODY()
+
+public:
+	virtual const UStruct* GetSourceStruct() const
+	{
+		return Variables.GetValue().GetScriptStruct();
+}
+
+	virtual const uint8* GetSourceDefaults() const
+	{
+		return Variables.GetValue().GetMemory();
+	}
+
+public:
+	UPROPERTY(EditAnywhere, Category = "Viewmodel")
+	FInstancedPropertyBag Variables;	// todo, this should be a base and we have 2 implementation, one for the normal with a struct builder, and one for verse with the class ptr
+
+
+public:
+#if WITH_EDITOR
+	virtual void PostEditChangeChainProperty(FPropertyChangedChainEvent& PropertyChangedEvent) override;
+#endif
+};
+
+#if UE_ENABLE_INCLUDE_ORDER_DEPRECATED_IN_5_2
+#include "CoreMinimal.h"
+#include "Templates/SubclassOf.h"
+#endif
diff --git a/Engine/Plugins/Runtime/ModelViewViewModel/Source/ModelViewViewModelBlueprint/Public/MVVMBlueprintViewModelContext.h b/Engine/Plugins/Runtime/ModelViewViewModel/Source/ModelViewViewModelBlueprint/Public/MVVMBlueprintViewModelContext.h
--- a/Engine/Plugins/Runtime/ModelViewViewModel/Source/ModelViewViewModelBlueprint/Public/MVVMBlueprintViewModelContext.h
+++ b/Engine/Plugins/Runtime/ModelViewViewModel/Source/ModelViewViewModelBlueprint/Public/MVVMBlueprintViewModelContext.h
@@ -7,6 +7,7 @@
 
 #include "MVVMBlueprintViewModelContext.generated.h"
 
+class UMVVMBlueprintInstancedViewModelBase;
 class UMVVMViewModelContextResolver;
 
 /**
@@ -89,10 +90,11 @@ public:
 
 private:
 	/** When the view is spawn, create an instance of the viewmodel. */
-	UPROPERTY(VisibleAnywhere, AdvancedDisplay, Category = "Viewmodel", meta = (DisplayName = "Viewmodel Context Id"))
+	UPROPERTY(VisibleAnywhere, AdvancedDisplay, Category = "Viewmodel", meta = (DisplayName = "Viewmodel Context Id", NoResetToDefault))
 	FGuid ViewModelContextId;
 
-	UPROPERTY(VisibleAnywhere, Category = "Viewmodel", NoClear, meta = (DisallowCreateNew, AllowedClasses = "/Script/UMG.NotifyFieldValueChanged", DisallowedClasses = "/Script/UMG.Widget"))
+public:
+	UPROPERTY(VisibleAnywhere, Category = "Viewmodel", NoClear, meta = (DisallowCreateNew, AllowedClasses = "/Script/FieldNotification.NotifyFieldValueChanged", DisallowedClasses = "/Script/UMG.Widget", NoResetToDefault))
 	TObjectPtr<UClass> NotifyFieldValueClass = nullptr;
 
 	UPROPERTY()
@@ -103,7 +105,7 @@ private:
 
 public:
 	/** Property name that will be generated. */
-	UPROPERTY(EditAnywhere, Category = "Viewmodel", meta = (DisplayName = "Viewmodel Name"))
+	UPROPERTY(EditAnywhere, Category = "Viewmodel", meta = (DisplayName = "Viewmodel Name", NoResetToDefault))
 	FName ViewModelName;
 
 	/** When the view is spawn, create an instance of the viewmodel. */
@@ -111,28 +113,31 @@ public:
 	EMVVMBlueprintViewModelContextCreationType CreationType = EMVVMBlueprintViewModelContextCreationType::CreateInstance;
 
 	/** Identifier of an already registered viewmodel. */
-	UPROPERTY(EditAnywhere, Category = "Viewmodel", AdvancedDisplay, meta = (DisplayName = "Global Viewmodel Identifier", EditCondition = "CreationType == EMVVMBlueprintViewModelContextCreationType::GlobalViewModelCollection && bCanEdit", EditConditionHides))
+	UPROPERTY(EditAnywhere, Category = "Viewmodel", AdvancedDisplay, meta = (DisplayName = "Global Viewmodel Identifier"))
 	FName GlobalViewModelIdentifier;
 
 	/** The Path to get the viewmodel instance. */
-	UPROPERTY(EditAnywhere, Category = "Viewmodel", AdvancedDisplay, meta = (DisplayName = "Viewmodel Property Path", EditCondition = "CreationType == EMVVMBlueprintViewModelContextCreationType::PropertyPath && bCanEdit", EditConditionHides))
+	UPROPERTY(EditAnywhere, Category = "Viewmodel", AdvancedDisplay, meta = (DisplayName = "Viewmodel Property Path"))
 	FString ViewModelPropertyPath;
 
-	UPROPERTY(EditAnywhere, Category = "Viewmodel", Instanced, AdvancedDisplay, meta = (EditInline, EditCondition = "CreationType == EMVVMBlueprintViewModelContextCreationType::Resolver && bCanEdit", EditConditionHides))
+	UPROPERTY(EditAnywhere, Category = "Viewmodel", Instanced, AdvancedDisplay, meta = (EditInline))
 	TObjectPtr<UMVVMViewModelContextResolver> Resolver = nullptr;
 
+	UPROPERTY(EditAnywhere, Category = "Viewmodel", Instanced, AdvancedDisplay, NoClear, meta = (ShowOnlyInnerProperties))
+	TObjectPtr<UMVVMBlueprintInstancedViewModelBase> InstancedViewModel;
+
 	/**
 	 * Generate a public setter for this viewmodel.
 	 * @note Always true when the Creation Type is Manual.
 	 */
-	UPROPERTY(EditAnywhere, Category = "Viewmodel", AdvancedDisplay, meta = (DisplanName="Create Public Setter", EditCondition = "CreationType != EMVVMBlueprintViewModelContextCreationType::Manual && bCanEdit", EditConditionHides))
+	UPROPERTY(EditAnywhere, Category = "Viewmodel", AdvancedDisplay, meta = (DisplanName="Create Public Setter"))
 	bool bCreateSetterFunction = false;
 
 	/**
 	 * Optional. Will not warn if the instance is not set or found.
 	 * @note Always true when the Creation Type is Manual.
 	 */
-	UPROPERTY(EditAnywhere, Category = "Viewmodel", AdvancedDisplay, meta = (EditCondition = "(CreationType == EMVVMBlueprintViewModelContextCreationType::GlobalViewModelCollection || CreationType == EMVVMBlueprintViewModelContextCreationType::PropertyPath) && bCanEdit", EditConditionHides))
+	UPROPERTY(EditAnywhere, Category = "Viewmodel", AdvancedDisplay)
 	bool bOptional = false;
 
 	/** Can change the name in the editor. */
diff --git a/Engine/Plugins/Runtime/ModelViewViewModel/Source/ModelViewViewModelBlueprint/Public/MVVMDeveloperProjectSettings.h b/Engine/Plugins/Runtime/ModelViewViewModel/Source/ModelViewViewModelBlueprint/Public/MVVMDeveloperProjectSettings.h
--- a/Engine/Plugins/Runtime/ModelViewViewModel/Source/ModelViewViewModelBlueprint/Public/MVVMDeveloperProjectSettings.h
+++ b/Engine/Plugins/Runtime/ModelViewViewModel/Source/ModelViewViewModelBlueprint/Public/MVVMDeveloperProjectSettings.h
@@ -131,6 +131,10 @@ public:
 	/** When binding to a multicast delegate property, allow to create an event. */
 	UPROPERTY(EditAnywhere, config, Category = "Viewmodel")
 	bool bAllowBindingEvent = true;
+	
+	/** Allow to create an instanced viewmodel directly in the view editor. */
+	UPROPERTY(EditAnywhere, config, Category = "Viewmodel")
+	bool bCanCreateViewModelInView = false;
 
 	/** Permission list for filtering which execution mode is allowed. */
 	UPROPERTY(EditAnywhere, config, Category = "Viewmodel")
diff --git a/Engine/Plugins/Runtime/ModelViewViewModel/Source/ModelViewViewModelEditor/Private/Customizations/MVVMBlueprintViewModelContextCustomization.cpp b/Engine/Plugins/Runtime/ModelViewViewModel/Source/ModelViewViewModelEditor/Private/Customizations/MVVMBlueprintViewModelContextCustomization.cpp
--- a/Engine/Plugins/Runtime/ModelViewViewModel/Source/ModelViewViewModelEditor/Private/Customizations/MVVMBlueprintViewModelContextCustomization.cpp
+++ b/Engine/Plugins/Runtime/ModelViewViewModel/Source/ModelViewViewModelEditor/Private/Customizations/MVVMBlueprintViewModelContextCustomization.cpp
@@ -9,6 +9,7 @@
 #include "IPropertyAccessEditor.h"
 #include "IPropertyTypeCustomization.h"
 #include "MVVMBlueprintView.h"
+#include "MVVMBlueprintInstancedViewModel.h"
 #include "MVVMDeveloperProjectSettings.h"
 #include "MVVMEditorSubsystem.h"
 #include "PropertyHandle.h"
@@ -127,130 +128,139 @@ void FBlueprintViewModelContextDetailCustomization::CustomizeChildren(TSharedRef
 	uint32 NumChildren = 0;
 	PropertyHandle->GetNumChildren(NumChildren);
 
-	FName Name_NotifyFieldValueClass = TEXT("NotifyFieldValueClass");
-	FName Name_ViewModelContextId = TEXT("ViewModelContextId");
-	FName Name_ViewModelName = GET_MEMBER_NAME_CHECKED(FMVVMBlueprintViewModelContext, ViewModelName);
-	FName Name_ViewModelPropertyPath = GET_MEMBER_NAME_CHECKED(FMVVMBlueprintViewModelContext, ViewModelPropertyPath);
-	FName Name_CreationType = GET_MEMBER_NAME_CHECKED(FMVVMBlueprintViewModelContext, CreationType);
-	FName Name_CreateSetterFunction = GET_MEMBER_NAME_CHECKED(FMVVMBlueprintViewModelContext, bCreateSetterFunction);
-
-	FGuid ViewModelContextId;
-	FName ViewModelPropertyName;
-	for (uint32 ChildIndex = 0; ChildIndex < NumChildren; ++ChildIndex)
+	FMVVMBlueprintViewModelContext* ContextPtr = nullptr;
 	{
-		TSharedPtr<IPropertyHandle> ChildHandle = PropertyHandle->GetChildHandle(ChildIndex);
-		const FName PropertyName = ChildHandle->GetProperty()->GetFName();
-
-		if (PropertyName == Name_NotifyFieldValueClass)
+		ensure(CastField<FStructProperty>(PropertyHandle->GetProperty()) && CastField<FStructProperty>(PropertyHandle->GetProperty())->Struct == FMVVMBlueprintViewModelContext::StaticStruct());
+		void* Buffer = nullptr;
+		if (PropertyHandle->GetValueData(Buffer) == FPropertyAccess::Success)
 		{
-			ensure(CastField<FClassProperty>(ChildHandle->GetProperty()));
-			NotifyFieldValueClassHandle = ChildHandle;
-			UObject* Object = nullptr;
-			if (ChildHandle->GetValue(Object) == FPropertyAccess::Success)
-			{
-				UClass* ViewModelClass = Cast<UClass>(Object);
-				if (ViewModelClass)
-				{
-					AllowedCreationTypes = GetAllowedContextCreationType(ViewModelClass);
-				}
-				PropertyAccessEditor.ClassToLookFor = ViewModelClass;
-			}
-			ChildHandle->SetOnPropertyValueChanged(FSimpleDelegate::CreateSP(this, &FBlueprintViewModelContextDetailCustomization::HandleClassChanged));
+			ContextPtr = reinterpret_cast<FMVVMBlueprintViewModelContext*>(Buffer);
 		}
+	}
 
-		else if (PropertyName == Name_ViewModelName)
-		{
-			ensure(CastField<FNameProperty>(ChildHandle->GetProperty()));
-			ViewModelNameHandle = ChildHandle;
-			ChildHandle->GetValue(ViewModelPropertyName);
-		}
-		else if (PropertyName == Name_ViewModelPropertyPath)
-		{
-			ensure(CastField<FStrProperty>(ChildHandle->GetProperty()));
-			PropertyPathHandle = ChildHandle;
-		}
-		else if (PropertyName == Name_CreationType)
-		{
-			ensure(CastField<FEnumProperty>(ChildHandle->GetProperty()));
-			CreationTypeHandle = ChildHandle;
-		}
-		else if (PropertyName == Name_ViewModelContextId)
+	if (ContextPtr == nullptr)
+	{
+		return;
+	}
+
+	FGuid ViewModelContextId = ContextPtr->GetViewModelId();
+	FName ViewModelPropertyName = ContextPtr->GetViewModelName();
+	bool bCanEdit = ContextPtr->bCanEdit;
+	bool bCanRename = ContextPtr->bCanRename;
+
+	NotifyFieldValueClassHandle = PropertyHandle->GetChildHandle(TEXT("NotifyFieldValueClass"), false);
+	PropertyPathHandle = PropertyHandle->GetChildHandle(GET_MEMBER_NAME_CHECKED(FMVVMBlueprintViewModelContext, ViewModelPropertyPath), false);
+	CreationTypeHandle = PropertyHandle->GetChildHandle(GET_MEMBER_NAME_CHECKED(FMVVMBlueprintViewModelContext, CreationType), false);
+	ViewModelNameHandle = PropertyHandle->GetChildHandle(GET_MEMBER_NAME_CHECKED(FMVVMBlueprintViewModelContext, ViewModelName), false);
+
+	if (ensure(NotifyFieldValueClassHandle))
+	{
+		UObject* Object = nullptr;
+		if (NotifyFieldValueClassHandle->GetValue(Object) == FPropertyAccess::Success)
 		{
-			ensure(CastField<FStructProperty>(ChildHandle->GetProperty()) && CastField<FStructProperty>(ChildHandle->GetProperty())->Struct->GetFName() == "Guid");
-			void* Buffer = nullptr;
-			if (ChildHandle->GetValueData(Buffer) == FPropertyAccess::Success)
+			UClass* ViewModelClass = Cast<UClass>(Object);
+			if (ViewModelClass)
 			{
-				ViewModelContextId = *reinterpret_cast<FGuid*>(Buffer);
+				AllowedCreationTypes = GetAllowedContextCreationType(ViewModelClass);
 			}
+			PropertyAccessEditor.ClassToLookFor = ViewModelClass;
 		}
+		NotifyFieldValueClassHandle->SetOnPropertyValueChanged(FSimpleDelegate::CreateSP(this, &FBlueprintViewModelContextDetailCustomization::HandleClassChanged));
 	}
 
-	bool bCanEdit = false;
-	bool bCanRename = false;
+	if (ensure(ViewModelNameHandle))
 	{
+		IDetailPropertyRow& PropertyRow = ChildBuilder.AddProperty(ViewModelNameHandle.ToSharedRef())
+			.IsEnabled(bCanRename);
+
+		TSharedPtr<SWidget> NameWidget, ValueWidget;
+		PropertyRow.GetDefaultWidgets(NameWidget, ValueWidget);
+		PropertyRow.CustomWidget()
+			.NameContent()
+			[
+				NameWidget.ToSharedRef()
+			]
+			.ValueContent()
+			[
+				SNew(SHorizontalBox)
+				+ SHorizontalBox::Slot()
+				.FillWidth(1.f)
+				[
+					SNew(SEditableTextBox)
+					.Text(this, &FBlueprintViewModelContextDetailCustomization::GetViewModalNameValueAsText)
+					.Font(CustomizationUtils.GetRegularFont())
+					.SelectAllTextWhenFocused(true)
+					.ClearKeyboardFocusOnCommit(false)
+					.OnTextCommitted(this, &FBlueprintViewModelContextDetailCustomization::HandleNameTextCommitted)
+					.OnVerifyTextChanged(this, &FBlueprintViewModelContextDetailCustomization::HandleNameVerifyTextChanged)
+					.SelectAllTextOnCommit(true)
+					.IsEnabled(bCanRename && ViewModelContextId.IsValid())
+				]
+			];
+	}
 
-		TSharedPtr<FWidgetBlueprintEditor> Editor = WidgetBlueprintEditor.Pin();
-		const UWidgetBlueprint* WidgetBP = Editor ? Editor->GetWidgetBlueprintObj() : nullptr;
-		const UMVVMBlueprintView* BlueprintView = WidgetBP ? GEditor->GetEditorSubsystem<UMVVMEditorSubsystem>()->GetView(WidgetBP) : nullptr;
-		const FMVVMBlueprintViewModelContext* ViewModel = BlueprintView ? BlueprintView->FindViewModel(ViewModelContextId) : nullptr;
-		if (ViewModel)
-		{
-			bCanEdit = ViewModel->bCanEdit;
-			bCanRename = ViewModel->bCanRename;
-		}
+	if (ensure(NotifyFieldValueClassHandle))
+	{
+		ChildBuilder.AddProperty(NotifyFieldValueClassHandle.ToSharedRef())
+			.IsEnabled(false)
+			.Visibility(MakeAttributeLambda([ContextPtr](){ return ContextPtr->InstancedViewModel != nullptr ? EVisibility::Collapsed : EVisibility::Visible; }));
 	}
 
-	if (NotifyFieldValueClassHandle)
+	if (ContextPtr->InstancedViewModel == nullptr)
 	{
-		for (uint32 ChildIndex = 0; ChildIndex < NumChildren; ++ChildIndex)
+		if (ensure(CreationTypeHandle))
 		{
-			TSharedPtr<IPropertyHandle> ChildHandle = PropertyHandle->GetChildHandle(ChildIndex);
-			const FName PropertyName = ChildHandle->GetProperty()->GetFName();
-
-			if (PropertyName == Name_ViewModelName)
-			{
-				ensure(CastField<FNameProperty>(ChildHandle->GetProperty()));
-				if (TSharedPtr<FWidgetBlueprintEditor> SharedWidgetBlueprintEditor = WidgetBlueprintEditor.Pin())
-				{
-					IDetailPropertyRow& PropertyRow = ChildBuilder.AddProperty(ChildHandle.ToSharedRef());
-
-					TSharedPtr<SWidget> NameWidget, ValueWidget;
-					PropertyRow.GetDefaultWidgets(NameWidget, ValueWidget);
-					PropertyRow.CustomWidget()
-					.NameContent()
+			IDetailPropertyRow& PropertyRow = ChildBuilder.AddProperty(CreationTypeHandle.ToSharedRef())
+				.IsEnabled(bCanEdit);
+
+			TSharedPtr<SWidget> NameWidget, ValueWidget;
+			PropertyRow.GetDefaultWidgets(NameWidget, ValueWidget);
+			PropertyRow.CustomWidget()
+				.NameContent()
+				[
+					NameWidget.ToSharedRef()
+				]
+				.ValueContent()
+				[
+					SNew(SComboButton)
+					.ContentPadding(FMargin(4.f, 0.f))
+					.OnGetMenuContent(this, &FBlueprintViewModelContextDetailCustomization::CreateExecutionTypeMenuContent)
+					.ButtonContent()
 					[
-						NameWidget.ToSharedRef()
+						SNew(STextBlock)
+						.Text(this, &FBlueprintViewModelContextDetailCustomization::GetCreationTypeValue)
+						.ToolTipText(this, &FBlueprintViewModelContextDetailCustomization::GetExecutionTypeValueToolTip)
 					]
-					.ValueContent()
-					[
-						SNew(SHorizontalBox)
-						+ SHorizontalBox::Slot()
-						.FillWidth(1.f)
-						[
-							SNew(SEditableTextBox)
-							.Text(this, &FBlueprintViewModelContextDetailCustomization::GetViewModalNameValueAsText)
-							.Font(CustomizationUtils.GetRegularFont())
-							.SelectAllTextWhenFocused(true)
-							.ClearKeyboardFocusOnCommit(false)
-							.OnTextCommitted(this, &FBlueprintViewModelContextDetailCustomization::HandleNameTextCommitted)
-							.OnVerifyTextChanged(this, &FBlueprintViewModelContextDetailCustomization::HandleNameVerifyTextChanged)
-							.SelectAllTextOnCommit(true)
-							.IsEnabled(bCanRename && ViewModelContextId.IsValid())
-						]
-					]
-					.IsEnabled(bCanRename);
-				}
-			}
-			else if (PropertyName == Name_ViewModelPropertyPath)
+				];
+		}
+
+		TSharedPtr<IPropertyHandle> GlobalViewModelIdentifierlHandle = PropertyHandle->GetChildHandle(GET_MEMBER_NAME_CHECKED(FMVVMBlueprintViewModelContext, GlobalViewModelIdentifier), false);
+		if (ensure(GlobalViewModelIdentifierlHandle))
+		{
+			ChildBuilder.AddProperty(GlobalViewModelIdentifierlHandle.ToSharedRef())
+				.IsEnabled(bCanEdit)
+				.Visibility(MakeAttributeLambda([ContextPtr]()
+					{
+						bool bResult = ContextPtr->CreationType == EMVVMBlueprintViewModelContextCreationType::GlobalViewModelCollection;
+						return bResult ? EVisibility::Visible : EVisibility::Collapsed;
+					}));
+		}
+		
+		if (ensure(PropertyPathHandle))
+		{
+			if (TSharedPtr<FWidgetBlueprintEditor> SharedWidgetBlueprintEditor = WidgetBlueprintEditor.Pin())
 			{
-				ensure(CastField<FStrProperty>(ChildHandle->GetProperty()));
-				if (TSharedPtr<FWidgetBlueprintEditor> SharedWidgetBlueprintEditor = WidgetBlueprintEditor.Pin())
-				{
-					IDetailPropertyRow& PropertyRow = ChildBuilder.AddProperty(ChildHandle.ToSharedRef());
+				IDetailPropertyRow& PropertyRow = ChildBuilder.AddProperty(PropertyPathHandle.ToSharedRef())
+					.IsEnabled(bCanEdit)
+					.Visibility(MakeAttributeLambda([ContextPtr]()
+						{
+							bool bResult = ContextPtr->CreationType == EMVVMBlueprintViewModelContextCreationType::PropertyPath;
+							return bResult ? EVisibility::Visible : EVisibility::Collapsed;
+						}));
 
-					TSharedPtr<SWidget> NameWidget, ValueWidget;
-					PropertyRow.GetDefaultWidgets(NameWidget, ValueWidget);
-					PropertyRow.CustomWidget()
+				TSharedPtr<SWidget> NameWidget, ValueWidget;
+				PropertyRow.GetDefaultWidgets(NameWidget, ValueWidget);
+				PropertyRow.CustomWidget()
 					.NameContent()
 					[
 						NameWidget.ToSharedRef()
@@ -269,47 +279,57 @@ void FBlueprintViewModelContextDetailCustomization::CustomizeChildren(TSharedRef
 							PropertyAccessEditor.MakePropertyBindingWidget(SharedWidgetBlueprintEditor.ToSharedRef(), NotifyFieldValueClassHandle->GetProperty(), PropertyPathHandle.ToSharedRef(), ViewModelPropertyName)
 						]
 					];
-				}
 			}
-			else if (PropertyName == Name_CreationType)
-			{
-				IDetailPropertyRow& PropertyRow = ChildBuilder.AddProperty(ChildHandle.ToSharedRef());
+		}
 
-				TSharedPtr<SWidget> NameWidget, ValueWidget;
-				PropertyRow.GetDefaultWidgets(NameWidget, ValueWidget);
-				PropertyRow.CustomWidget()
-					.NameContent()
-					[
-						NameWidget.ToSharedRef()
-					]
-					.ValueContent()
-					[
-						SNew(SComboButton)
-						.ContentPadding(FMargin(4.f, 0.f))
-						.OnGetMenuContent(this, &FBlueprintViewModelContextDetailCustomization::CreateExecutionTypeMenuContent)
-						.ButtonContent()
-						[
-							SNew(STextBlock)
-							.Text(this, &FBlueprintViewModelContextDetailCustomization::GetCreationTypeValue)
-							.ToolTipText(this, &FBlueprintViewModelContextDetailCustomization::GetExecutionTypeValueToolTip)
-						]
-					]
-					.IsEnabled(bCanEdit);
-			}
-			else if (PropertyName == Name_CreateSetterFunction)
-			{
-				if (GetDefault<UMVVMDeveloperProjectSettings>()->bAllowGeneratedViewModelSetter)
-				{
-					ChildBuilder.AddProperty(ChildHandle.ToSharedRef())
-						.IsEnabled(bCanEdit);
-				}
-			}
-			else
+		TSharedPtr<IPropertyHandle> ResolverHandle = PropertyHandle->GetChildHandle(GET_MEMBER_NAME_CHECKED(FMVVMBlueprintViewModelContext, Resolver), false);
+		if (ensure(ResolverHandle))
+		{
+			ChildBuilder.AddProperty(ResolverHandle.ToSharedRef())
+				.IsEnabled(bCanEdit)
+				.Visibility(MakeAttributeLambda([ContextPtr]()
+					{
+						bool bResult = ContextPtr->CreationType == EMVVMBlueprintViewModelContextCreationType::Resolver;
+						return bResult ? EVisibility::Visible : EVisibility::Collapsed;
+					}));
+		}
+
+		if (GetDefault<UMVVMDeveloperProjectSettings>()->bAllowGeneratedViewModelSetter)
+		{
+			TSharedPtr<IPropertyHandle> CreateSetterFunctionHandle = PropertyHandle->GetChildHandle(GET_MEMBER_NAME_CHECKED(FMVVMBlueprintViewModelContext, bCreateSetterFunction), false);
+			if (ensure(CreateSetterFunctionHandle))
 			{
-				ChildBuilder.AddProperty(ChildHandle.ToSharedRef())
-					.IsEnabled(bCanEdit);;
+				ChildBuilder.AddProperty(CreateSetterFunctionHandle.ToSharedRef())
+					.IsEnabled(bCanEdit)
+					.Visibility(MakeAttributeLambda([ContextPtr]()
+						{
+							bool bResult = ContextPtr->CreationType != EMVVMBlueprintViewModelContextCreationType::Manual;
+							return bResult ? EVisibility::Visible : EVisibility::Collapsed;
+						}));
 			}
 		}
+
+		TSharedPtr<IPropertyHandle> OptionalHandle = PropertyHandle->GetChildHandle(GET_MEMBER_NAME_CHECKED(FMVVMBlueprintViewModelContext, bOptional), false);
+		if (ensure(OptionalHandle))
+		{
+			ChildBuilder.AddProperty(OptionalHandle.ToSharedRef())
+				.IsEnabled(bCanEdit)
+				.Visibility(MakeAttributeLambda([ContextPtr]()
+					{
+						bool bResult = ContextPtr->CreationType == EMVVMBlueprintViewModelContextCreationType::GlobalViewModelCollection
+						 || ContextPtr->CreationType == EMVVMBlueprintViewModelContextCreationType::PropertyPath;
+						return bResult ? EVisibility::Visible : EVisibility::Collapsed;
+					}));
+		}
+	}
+	else
+	{
+		TSharedPtr<IPropertyHandle> InstancedViewModelHandle = PropertyHandle->GetChildHandle(GET_MEMBER_NAME_CHECKED(FMVVMBlueprintViewModelContext, InstancedViewModel), false);
+		if (ensure(InstancedViewModelHandle))
+		{
+			ChildBuilder.AddProperty(InstancedViewModelHandle.ToSharedRef())
+				.IsEnabled(bCanEdit);
+		}
 	}
 }
 
diff --git a/Engine/Plugins/Runtime/ModelViewViewModel/Source/ModelViewViewModelEditor/Private/MVVMEditorSubsystem.cpp b/Engine/Plugins/Runtime/ModelViewViewModel/Source/ModelViewViewModelEditor/Private/MVVMEditorSubsystem.cpp
--- a/Engine/Plugins/Runtime/ModelViewViewModel/Source/ModelViewViewModelEditor/Private/MVVMEditorSubsystem.cpp
+++ b/Engine/Plugins/Runtime/ModelViewViewModel/Source/ModelViewViewModelEditor/Private/MVVMEditorSubsystem.cpp
@@ -16,6 +16,7 @@
 #include "K2Node_FunctionEntry.h"
 #include "K2Node_FunctionResult.h"
 #include "K2Node_VariableGet.h"
+#include "MVVMBlueprintInstancedViewModel.h"
 #include "MVVMBlueprintView.h"
 #include "MVVMBlueprintViewConversionFunction.h"
 #include "MVVMBlueprintViewEvent.h"
@@ -152,9 +153,9 @@ UMVVMBlueprintView* UMVVMEditorSubsystem::GetView(const UWidgetBlueprint* Widget
 	return nullptr;
 }
 
-FName UMVVMEditorSubsystem::AddViewModel(UWidgetBlueprint* WidgetBlueprint, const UClass* ViewModelClass)
+FGuid UMVVMEditorSubsystem::AddViewModel(UWidgetBlueprint* WidgetBlueprint, const UClass* ViewModelClass)
 {
-	FName Result;
+	FGuid Result;
 	if (ViewModelClass)
 	{
 		if (UMVVMBlueprintView* View = GetView(WidgetBlueprint))
@@ -176,30 +177,69 @@ FName UMVVMEditorSubsystem::AddViewModel(UWidgetBlueprint* WidgetBlueprint, cons
 				++Index;
 			}
 
-			Result = *ViewModelName;
-			FMVVMBlueprintViewModelContext Context = FMVVMBlueprintViewModelContext(ViewModelClass, Result);
+			FMVVMBlueprintViewModelContext Context = FMVVMBlueprintViewModelContext(ViewModelClass, *ViewModelName);
 			if (Context.IsValid())
 			{
 				const FScopedTransaction Transaction(LOCTEXT("AddViewModel", "Add viewmodel"));
 				View->Modify();
 				View->AddViewModel(Context);
-			}
-			else
-			{
-				Result = FName();
+				Result = Context.GetViewModelId();
 			}
 		}
 	}
 	return Result;
 }
 
+FGuid UMVVMEditorSubsystem::AddInstancedViewModel(UWidgetBlueprint* WidgetBlueprint)
+{
+	UMVVMWidgetBlueprintExtension_View* ExtensionView = UMVVMWidgetBlueprintExtension_View::GetExtension<UMVVMWidgetBlueprintExtension_View>(WidgetBlueprint);
+	UMVVMBlueprintView* View = ExtensionView ? ExtensionView->GetBlueprintView() : nullptr;
+
+	if (!View)
+	{
+		return FGuid();
+	}
+
+	const FScopedTransaction Transaction(LOCTEXT("AddInstancedViewModel", "Add instanced viewmodel"));
+
+	FGuid Result;
+
+	FName UniqueName = MakeUniqueObjectName(View, UMVVMBlueprintInstancedViewModel::StaticClass(), "InstancedViewmodel");
+	UMVVMBlueprintInstancedViewModel* NewInstancedViewModel = NewObject<UMVVMBlueprintInstancedViewModel>(View, UniqueName);
+	NewInstancedViewModel->GenerateClass();
+	FMVVMBlueprintViewModelContext Context = FMVVMBlueprintViewModelContext(NewInstancedViewModel->GetGeneratedClass(), UniqueName);
+	if (Context.IsValid())
+	{
+		Context.InstancedViewModel = NewInstancedViewModel;
+		Context.CreationType = EMVVMBlueprintViewModelContextCreationType::CreateInstance;
+		View->Modify();
+		View->AddViewModel(Context);
+		Result = Context.GetViewModelId();
+	}
+	else
+	{
+		auto RenameToTransient = [](UObject* ObjectToRename)
+		{
+			FName TrashName = MakeUniqueObjectName(GetTransientPackage(), ObjectToRename->GetClass(), *FString::Printf(TEXT("TRASH_%s"), *ObjectToRename->GetName()));
+			ObjectToRename->Rename(*TrashName.ToString(), GetTransientPackage());
+		};
+		if (NewInstancedViewModel->GetGeneratedClass())
+		{
+			RenameToTransient(NewInstancedViewModel->GetGeneratedClass());
+		}
+		RenameToTransient(NewInstancedViewModel);
+	}
+
+	return Result;
+}
+
 void UMVVMEditorSubsystem::RemoveViewModel(UWidgetBlueprint* WidgetBlueprint, FName ViewModel)
 {
 	if (UMVVMBlueprintView* View = GetView(WidgetBlueprint))
 	{
 		if (const FMVVMBlueprintViewModelContext* ViewModelContext = View->FindViewModel(ViewModel))
 		{
-			if (ViewModelContext->bCanRemove)
+			if (ViewModelContext->bCanRemove && ViewModelContext->InstancedViewModel)
 			{
 				View->RemoveViewModel(ViewModelContext->GetViewModelId());
 			}
@@ -329,6 +369,7 @@ void UMVVMEditorSubsystem::SetSourceToDestinationConversionFunction(UWidgetBluep
 			}
 
 			UE::MVVM::Private::OnBindingPostEditChange(View, GET_MEMBER_NAME_CHECKED(FMVVMBlueprintViewBinding, Conversion));
+			FBlueprintEditorUtils::MarkBlueprintAsModified(WidgetBlueprint);
 		}
 	}
 }
@@ -359,6 +400,7 @@ void UMVVMEditorSubsystem::SetDestinationToSourceConversionFunction(UWidgetBluep
 			}
 
 			UE::MVVM::Private::OnBindingPostEditChange(View, GET_MEMBER_NAME_CHECKED(FMVVMBlueprintViewBinding, Conversion));
+			FBlueprintEditorUtils::MarkBlueprintAsModified(WidgetBlueprint);
 		}
 	}
 }
@@ -389,6 +431,7 @@ void UMVVMEditorSubsystem::SetDestinationPathForBinding(UWidgetBlueprint* Widget
 		}
 
 		UE::MVVM::Private::OnBindingPostEditChange(View, GET_MEMBER_NAME_CHECKED(FMVVMBlueprintViewBinding, DestinationPath));
+		FBlueprintEditorUtils::MarkBlueprintAsModified(WidgetBlueprint);
 	}
 }
 
@@ -403,6 +446,7 @@ void UMVVMEditorSubsystem::SetSourcePathForBinding(UWidgetBlueprint* WidgetBluep
 		Binding.SourcePath = Field;
 
 		UE::MVVM::Private::OnBindingPostEditChange(View, GET_MEMBER_NAME_CHECKED(FMVVMBlueprintViewBinding, SourcePath));
+		FBlueprintEditorUtils::MarkBlueprintAsModified(WidgetBlueprint);
 	}
 }
 
@@ -420,6 +464,7 @@ void UMVVMEditorSubsystem::OverrideExecutionModeForBinding(UWidgetBlueprint* Wid
 			Binding.OverrideExecutionMode = Mode;
 
 			UE::MVVM::Private::OnBindingPostEditChange(View, GET_MEMBER_NAME_CHECKED(FMVVMBlueprintViewBinding, OverrideExecutionMode));
+			FBlueprintEditorUtils::MarkBlueprintAsModified(WidgetBlueprint);
 		}
 	}
 }
@@ -437,6 +482,7 @@ void UMVVMEditorSubsystem::ResetExecutionModeForBinding(UWidgetBlueprint* Widget
 			Binding.bOverrideExecutionMode = false;
 
 			UE::MVVM::Private::OnBindingPostEditChange(View, GET_MEMBER_NAME_CHECKED(FMVVMBlueprintViewBinding, OverrideExecutionMode));
+			FBlueprintEditorUtils::MarkBlueprintAsModified(WidgetBlueprint);
 		}
 	}
 }
@@ -454,6 +500,7 @@ void UMVVMEditorSubsystem::SetBindingTypeForBinding(UWidgetBlueprint* WidgetBlue
 			Binding.BindingType = Type;
 
 			UE::MVVM::Private::OnBindingPostEditChange(View, GET_MEMBER_NAME_CHECKED(FMVVMBlueprintViewBinding, BindingType));
+			FBlueprintEditorUtils::MarkBlueprintAsModified(WidgetBlueprint);
 		}
 	}
 }
@@ -471,6 +518,7 @@ void UMVVMEditorSubsystem::SetEnabledForBinding(UWidgetBlueprint* WidgetBlueprin
 			Binding.bEnabled = bEnabled;
 
 			UE::MVVM::Private::OnBindingPostEditChange(View, GET_MEMBER_NAME_CHECKED(FMVVMBlueprintViewBinding, bEnabled));
+			FBlueprintEditorUtils::MarkBlueprintAsModified(WidgetBlueprint);
 		}
 	}
 }
@@ -488,6 +536,7 @@ void UMVVMEditorSubsystem::SetCompileForBinding(UWidgetBlueprint* WidgetBlueprin
 			Binding.bCompile = bCompile;
 
 			UE::MVVM::Private::OnBindingPostEditChange(View, GET_MEMBER_NAME_CHECKED(FMVVMBlueprintViewBinding, bCompile));
+			FBlueprintEditorUtils::MarkBlueprintAsModified(WidgetBlueprint);
 		}
 	}
 }
@@ -832,6 +881,7 @@ void UMVVMEditorSubsystem::SetPathForConversionFunctionArgument(UWidgetBlueprint
 		check(ConversionFunction);
 		ConversionFunction->SetGraphPin(WidgetBlueprint, ArgumentName, Path);
 	}
+	FBlueprintEditorUtils::MarkBlueprintAsModified(WidgetBlueprint);
 }
 
 TArray<UE::MVVM::FBindingSource> UMVVMEditorSubsystem::GetBindableWidgets(const UWidgetBlueprint* WidgetBlueprint) const
diff --git a/Engine/Plugins/Runtime/ModelViewViewModel/Source/ModelViewViewModelEditor/Private/Widgets/SMVVMViewModelPanel.cpp b/Engine/Plugins/Runtime/ModelViewViewModel/Source/ModelViewViewModelEditor/Private/Widgets/SMVVMViewModelPanel.cpp
--- a/Engine/Plugins/Runtime/ModelViewViewModel/Source/ModelViewViewModelEditor/Private/Widgets/SMVVMViewModelPanel.cpp
+++ b/Engine/Plugins/Runtime/ModelViewViewModel/Source/ModelViewViewModelEditor/Private/Widgets/SMVVMViewModelPanel.cpp
@@ -6,6 +6,7 @@
 #include "Dialogs/Dialogs.h"
 #include "IStructureDetailsView.h"
 #include "MVVMBlueprintView.h"
+#include "MVVMBlueprintInstancedViewModel.h"
 #include "MVVMDeveloperProjectSettings.h"
 #include "MVVMEditorSubsystem.h"
 #include "MVVMWidgetBlueprintExtension_View.h"
@@ -20,6 +21,7 @@
 #include "Framework/Commands/GenericCommands.h"
 #include "Framework/MultiBox/MultiBoxBuilder.h"
 #include "Styling/MVVMEditorStyle.h"
+#include "Styling/ToolBarStyle.h"
 
 #include "Widgets/Images/SImage.h"
 #include "Widgets/Input/SButton.h"
@@ -75,6 +77,21 @@ void SetSelectObjectsToView(TWeakPtr<FWidgetBlueprintEditor> WeakEditor)
 	}
 }
 
+void CreateViewModelInstance(TWeakPtr<FWidgetBlueprintEditor> WeakEditor)
+{
+	if (TSharedPtr<FWidgetBlueprintEditor> Editor = WeakEditor.Pin())
+	{
+		UMVVMEditorSubsystem* MVVMEditorSubsystem = GEditor->GetEditorSubsystem<UMVVMEditorSubsystem>();
+		if (!MVVMEditorSubsystem)
+		{
+			return;
+		}
+
+		UMVVMBlueprintView* BlueprintView = MVVMEditorSubsystem->RequestView(Editor->GetWidgetBlueprintObj());
+		MVVMEditorSubsystem->AddInstancedViewModel(Editor->GetWidgetBlueprintObj());
+	}
+}
+
 } //namespace UE::MVVM::Private
 
 namespace UE::MVVM
@@ -96,10 +113,29 @@ void SMVVMViewModelPanel::RegisterMenu()
 					}
 				}
 			}));
+	} {
+		UToolMenu* Menu = UToolMenus::Get()->RegisterMenu("MVVM.Viewmodels.Add");
+		FToolMenuSection& Section = Menu->FindOrAddSection("Main");
+		Section.AddDynamicEntry("Main", FNewToolMenuSectionDelegate::CreateLambda([](FToolMenuSection& InSection)
+			{
+				if (const UWidgetBlueprintToolMenuContext* Context = InSection.FindContext<UWidgetBlueprintToolMenuContext>())
+				{
+					if (GetDefault<UMVVMDeveloperProjectSettings>()->bCanCreateViewModelInView)
+					{
+						InSection.AddMenuEntry(
+							"CreateViewmodel"
+							, LOCTEXT("CreateViewmodel", "Instanced Viewmodel")
+							, LOCTEXT("CreateViewmodelTooltip", "Create a Viewmodel inside the View. The Viewmodel is not accessible from outside the View.")
+							, FSlateIcon(FMVVMEditorStyle::Get().GetStyleSetName(), "BlueprintView.TabIcon")
+							, FUIAction(FExecuteAction::CreateStatic(UE::MVVM::Private::CreateViewModelInstance, Context->WidgetBlueprintEditor))
+							, EUserInterfaceActionType::Button
+						);
+					}
+				}
+			}));
 	}
 	{
 		UToolMenu* Menu = UToolMenus::Get()->RegisterMenu("MVVM.Viewmodels.Settings");
-		//Menu->MenuType = EMultiBoxType::Menu;
 		FToolMenuSection& Section = Menu->FindOrAddSection("Main");
 		Section.AddDynamicEntry("Main", FNewToolMenuSectionDelegate::CreateLambda([](FToolMenuSection& InSection)
 			{
@@ -124,15 +160,37 @@ void SMVVMViewModelPanel::RegisterMenu()
 
 void SMVVMViewModelPanel::BuildContextMenu(FToolMenuSection& InSection)
 {
-	if (!AddMenuButton.IsValid())
-	{
-		SAssignNew(AddMenuButton, SPositiveActionButton)
-			.OnGetMenuContent(this, &SMVVMViewModelPanel::MakeAddMenu)
-			.Text(LOCTEXT("Viewmodel", "Viewmodel"))
-			.IsEnabled(this, &SMVVMViewModelPanel::HandleCanEditViewmodelList);
-	}
-	
+	SAssignNew(AddMenuButton, SPositiveActionButton)
+		.OnGetMenuContent(this, &SMVVMViewModelPanel::MakeAddMenu)
+		.Text(LOCTEXT("Viewmodel", "Viewmodel"))
+		.IsEnabled(this, &SMVVMViewModelPanel::HandleCanEditViewmodelList);
+
 	InSection.AddEntry(FToolMenuEntry::InitWidget("AddViewmodel", AddMenuButton.ToSharedRef(), FText()));
+
+	const FToolBarStyle& ToolBarStyle = FCoreStyle::Get().GetWidgetStyle<FToolBarStyle>("CalloutToolbar");
+	const FVector2f IconSize = ToolBarStyle.IconSize;
+	const FComboButtonStyle* ComboStyle = &ToolBarStyle.SettingsComboButton;
+	const FButtonStyle* ButtonStyle = &ComboStyle->ButtonStyle;
+	FSlateColor OpenForegroundColor = ButtonStyle->HoveredForeground;
+
+	TWeakPtr<SComboButton> WeakComboBox;
+	TSharedRef<SWidget> ComboBottom = SAssignNew(WeakComboBox, SComboButton)
+		.ContentPadding(0.f)
+		.ComboButtonStyle(ComboStyle)
+		.ButtonStyle(ButtonStyle)
+		.ForegroundColor_Lambda([WeakComboBox, OpenForegroundColor]()
+		{
+			TSharedPtr<SComboButton> LocalComboButton = WeakComboBox.Pin();
+			return LocalComboButton ? OpenForegroundColor : FSlateColor::UseStyle();
+		})
+		// Route the content generator event
+		.OnGetMenuContent(this, &SMVVMViewModelPanel::HandleAddViewModelContextMenu)
+		.IsEnabled(this, &SMVVMViewModelPanel::HandleCanEditViewmodelList)
+		.ButtonContent()
+		[
+			SNullWidget::NullWidget
+		];
+	InSection.AddEntry(FToolMenuEntry::InitWidget("AddViewmodelContext", ComboBottom, FText()));
 }
 
 
@@ -410,6 +468,23 @@ bool SMVVMViewModelPanel::HandleCanEditViewmodelList() const
 }
 
 
+TSharedRef<SWidget> SMVVMViewModelPanel::HandleAddViewModelContextMenu()
+{
+	FToolMenuContext GenerateWidgetContext;
+	{
+		UWidgetBlueprintToolMenuContext* WidgetBlueprintMenuContext = NewObject<UWidgetBlueprintToolMenuContext>();
+		WidgetBlueprintMenuContext->WidgetBlueprintEditor = WeakBlueprintEditor;
+		GenerateWidgetContext.AddObject(WidgetBlueprintMenuContext);
+
+		UMVVMViewModelPanelToolMenuContext* ViewModelPanelToolMenuContext = NewObject<UMVVMViewModelPanelToolMenuContext>();
+		ViewModelPanelToolMenuContext->ViewModelPanel = SharedThis(this);
+		GenerateWidgetContext.AddObject(ViewModelPanelToolMenuContext);
+	}
+
+	return UToolMenus::Get()->GenerateWidget("MVVM.Viewmodels.Add", GenerateWidgetContext);
+}
+
+
 TSharedPtr<SWidget> SMVVMViewModelPanel::HandleGetPreSlot(UE::PropertyViewer::SPropertyViewer::FHandle Handle, TArrayView<const FFieldVariant> FieldPath)
 {
 	if (FieldPath.Num() > 0)
@@ -440,7 +515,7 @@ TSharedRef<SWidget> SMVVMViewModelPanel::HandleGenerateContainer(UE::PropertyVie
 				{
 					TSharedRef<SInlineEditableTextBlock> EditableTextBlock = SNew(SInlineEditableTextBlock)
 						.Text(ViewModelContext->GetDisplayName())
-						.IsReadOnly(this, &SMVVMViewModelPanel::HandleCanRename, VMGuid)
+						.IsReadOnly_Lambda([this, VMGuid]() { return !HandleCanRename(VMGuid); })
 						.OnVerifyTextChanged(this, &SMVVMViewModelPanel::HandleVerifyNameTextChanged, VMGuid)
 						.OnTextCommitted(this, &SMVVMViewModelPanel::HandleNameTextCommited, VMGuid);
 					EditableTextBlocks.Add(VMGuid, EditableTextBlock);
@@ -620,9 +695,9 @@ void SMVVMViewModelPanel::HandleDeleteViewModel()
 							return;
 						}
 					}
-				}
 
-				BlueprintView->RemoveViewModel(*VMGuidPtr);
+					GEditor->GetEditorSubsystem<UMVVMEditorSubsystem>()->RemoveViewModel(WidgetBP, ViewModelContext->GetViewModelName());
+				}
 			}
 		}
 	}
diff --git a/Engine/Plugins/Runtime/ModelViewViewModel/Source/ModelViewViewModelEditor/Private/Widgets/SMVVMViewModelPanel.h b/Engine/Plugins/Runtime/ModelViewViewModel/Source/ModelViewViewModelEditor/Private/Widgets/SMVVMViewModelPanel.h
--- a/Engine/Plugins/Runtime/ModelViewViewModel/Source/ModelViewViewModelEditor/Private/Widgets/SMVVMViewModelPanel.h
+++ b/Engine/Plugins/Runtime/ModelViewViewModel/Source/ModelViewViewModelEditor/Private/Widgets/SMVVMViewModelPanel.h
@@ -95,6 +95,7 @@ private:
 	bool HandleCanDeleteViewModel() const;
 	void HandleRenameViewModel();
 	bool HandleCanRenameViewModel() const;
+	TSharedRef<SWidget> HandleAddViewModelContextMenu();
 
 	EVisibility GetWarningPanelVisibility() const;
 	FReply HandleDisableWarningPanel();
diff --git a/Engine/Plugins/Runtime/ModelViewViewModel/Source/ModelViewViewModelEditor/Public/MVVMEditorSubsystem.h b/Engine/Plugins/Runtime/ModelViewViewModel/Source/ModelViewViewModelEditor/Public/MVVMEditorSubsystem.h
--- a/Engine/Plugins/Runtime/ModelViewViewModel/Source/ModelViewViewModelEditor/Public/MVVMEditorSubsystem.h
+++ b/Engine/Plugins/Runtime/ModelViewViewModel/Source/ModelViewViewModelEditor/Public/MVVMEditorSubsystem.h
@@ -37,7 +37,10 @@ public:
 	UMVVMBlueprintView* GetView(const UWidgetBlueprint* WidgetBlueprint) const;
 
 	UFUNCTION(BlueprintCallable, Category = "Viewmodel")
-	FName AddViewModel(UWidgetBlueprint* WidgetBlueprint, const UClass* ViewModel);
+	FGuid AddViewModel(UWidgetBlueprint* WidgetBlueprint, const UClass* ViewModel);
+
+	UFUNCTION(BlueprintCallable, Category = "Viewmodel")
+	FGuid AddInstancedViewModel(UWidgetBlueprint* WidgetBlueprint);
 	
 	UFUNCTION(BlueprintCallable, Category = "Viewmodel")
 	void RemoveViewModel(UWidgetBlueprint* WidgetBlueprint, FName ViewModel);
