MVVM: Fix bad compilation unit on linux.

diff --git a/Engine/Plugins/Runtime/ModelViewViewModel/Source/ModelViewViewModelBlueprint/Public/MVVMBlueprintInstancedViewModel.h b/Engine/Plugins/Runtime/ModelViewViewModel/Source/ModelViewViewModelBlueprint/Public/MVVMBlueprintInstancedViewModel.h
--- a/Engine/Plugins/Runtime/ModelViewViewModel/Source/ModelViewViewModelBlueprint/Public/MVVMBlueprintInstancedViewModel.h
+++ b/Engine/Plugins/Runtime/ModelViewViewModel/Source/ModelViewViewModelBlueprint/Public/MVVMBlueprintInstancedViewModel.h
@@ -27,15 +27,8 @@ public:
 		return GeneratedClass;
 	}
 
-	virtual const UStruct* GetSourceStruct() const
-	{
-		PURE_VIRTUAL(GetSourceStruct, return nullptr;);
-	}
-
-	virtual const uint8* GetSourceDefaults() const
-	{
-		PURE_VIRTUAL(GetSourceDefaults, return nullptr;);
-	}
+	virtual const UStruct* GetSourceStruct() const PURE_VIRTUAL(GetSourceStruct, return nullptr;);
+	virtual const uint8* GetSourceDefaults() const PURE_VIRTUAL(GetSourceDefaults, return nullptr;);
 
 protected:
 	virtual void PreAddProperties();
