MVVM: Fix bad test when removing a viewmodel in the editor.

diff --git a/Engine/Plugins/Runtime/ModelViewViewModel/Source/ModelViewViewModelEditor/Private/MVVMEditorSubsystem.cpp b/Engine/Plugins/Runtime/ModelViewViewModel/Source/ModelViewViewModelEditor/Private/MVVMEditorSubsystem.cpp
--- a/Engine/Plugins/Runtime/ModelViewViewModel/Source/ModelViewViewModelEditor/Private/MVVMEditorSubsystem.cpp
+++ b/Engine/Plugins/Runtime/ModelViewViewModel/Source/ModelViewViewModelEditor/Private/MVVMEditorSubsystem.cpp
@@ -239,7 +239,7 @@ void UMVVMEditorSubsystem::RemoveViewModel(UWidgetBlueprint* WidgetBlueprint, FN
 	{
 		if (const FMVVMBlueprintViewModelContext* ViewModelContext = View->FindViewModel(ViewModel))
 		{
-			if (ViewModelContext->bCanRemove && ViewModelContext->InstancedViewModel)
+			if (ViewModelContext->bCanRemove)
 			{
 				View->RemoveViewModel(ViewModelContext->GetViewModelId());
 			}
