MVVM:
- Remove FPropertyEditorPermissionList call that was massively slowing down launching Frontend in PIE

diff --git a/Engine/Plugins/Runtime/ModelViewViewModel/Source/ModelViewViewModelBlueprint/Private/MVVMViewBlueprintCompiler.cpp b/Engine/Plugins/Runtime/ModelViewViewModel/Source/ModelViewViewModelBlueprint/Private/MVVMViewBlueprintCompiler.cpp
--- a/Engine/Plugins/Runtime/ModelViewViewModel/Source/ModelViewViewModelBlueprint/Private/MVVMViewBlueprintCompiler.cpp
+++ b/Engine/Plugins/Runtime/ModelViewViewModel/Source/ModelViewViewModelBlueprint/Private/MVVMViewBlueprintCompiler.cpp
@@ -931,8 +931,6 @@ bool FMVVMViewBlueprintCompiler::PreCompile(UWidgetBlueprintGeneratedClass* Clas
 	ComplexConversionFunctionContexts.Reset();
 	EventDelegateContexts.Reset();
 
-	FPropertyEditorPermissionList::Get().AddPermissionList(Class, FNamePermissionList(), EPropertyPermissionListRules::AllowListAllProperties);
-
 	PreCompileBindingSources(Class, BlueprintView);
 	PreCompileSourceCreators(Class, BlueprintView);
 	PreCompileBindings(Class, BlueprintView);
