Fix compilation unit.
#rnx

diff --git a/Engine/Plugins/Runtime/ModelViewViewModel/Source/ModelViewViewModel/Public/Debugging/MVVMDebugging.h b/Engine/Plugins/Runtime/ModelViewViewModel/Source/ModelViewViewModel/Public/Debugging/MVVMDebugging.h
--- a/Engine/Plugins/Runtime/ModelViewViewModel/Source/ModelViewViewModel/Public/Debugging/MVVMDebugging.h
+++ b/Engine/Plugins/Runtime/ModelViewViewModel/Source/ModelViewViewModel/Public/Debugging/MVVMDebugging.h
@@ -24,7 +24,7 @@ namespace UE::MVVM
 class MODELVIEWVIEWMODEL_API FDebugging
 {
 public:
-	struct FView
+	struct MODELVIEWVIEWMODEL_API FView
 	{
 	public:
 		FView() = default;
