MVVM: Remove cached functions that was affecting GC cluster.

diff --git a/Engine/Plugins/Runtime/ModelViewViewModel/Source/ModelViewViewModel/Private/Bindings/MVVMCompiledBindingLibrary.cpp b/Engine/Plugins/Runtime/ModelViewViewModel/Source/ModelViewViewModel/Private/Bindings/MVVMCompiledBindingLibrary.cpp
--- a/Engine/Plugins/Runtime/ModelViewViewModel/Source/ModelViewViewModel/Private/Bindings/MVVMCompiledBindingLibrary.cpp
+++ b/Engine/Plugins/Runtime/ModelViewViewModel/Source/ModelViewViewModel/Private/Bindings/MVVMCompiledBindingLibrary.cpp
@@ -20,7 +20,7 @@ DECLARE_CYCLE_STAT(TEXT("Load Library"), STAT_UMG_Viewmodel_LoadLibrary, STATGRO
 
 namespace UE::MVVM::Private
 {
-	bool IsFunctionVirtual(UFunction* Function)
+	bool IsFunctionVirtual(const UFunction* Function)
 	{
 		return !Function->HasAnyFunctionFlags(FUNC_Static | FUNC_Final);
 	}
@@ -53,6 +53,44 @@ UE::FieldNotification::FFieldId FMVVMVCompiledFields::GetFieldId(FName FieldName
 }
 
 
+/**
+ *
+ */
+FMVVMCompiledBindingLibrary::FLoadedFunction::FLoadedFunction(const UFunction* Function)
+{
+	if (Function)
+	{
+		ClassOwner = Function->GetOwnerClass();
+		FunctionName = Function->GetFName();
+		bIsFunctionVirtual = UE::MVVM::Private::IsFunctionVirtual(Function);
+	}
+}
+
+UFunction* FMVVMCompiledBindingLibrary::FLoadedFunction::GetFunction() const
+{
+	const UClass* ClassPtr = ClassOwner.Get();
+	if (ClassPtr)
+	{
+		check(!FunctionName.IsNone());
+		return ClassPtr->FindFunctionByName(FunctionName);
+	}
+	return nullptr;
+}
+
+UFunction* FMVVMCompiledBindingLibrary::FLoadedFunction::GetFunction(const UObject* CallingContext) const
+{
+	if (bIsFunctionVirtual)
+	{
+		check(!FunctionName.IsNone());
+		if (CallingContext)
+		{
+			return CallingContext->GetClass()->FindFunctionByName(FunctionName);
+		}
+		return nullptr;
+	}
+	return GetFunction();
+}
+
 /**
  *
  */
@@ -118,8 +156,9 @@ void FMVVMCompiledBindingLibrary::Load()
 			{
 				FName FieldName = Field.GetFunctionName(CompiledFieldNames, Index);
 				UFunction* LoadedFunction = Field.GetFunction(FieldName);
-				LoadedFunctions.Add(LoadedFunction); // add it even if none to keep the index valid
 				ensureAlwaysMsgf(LoadedFunction != nullptr, TEXT("The function '%s:%s' could not be loaded."), (Field.GetStruct() ? *Field.GetStruct()->GetName() : TEXT("None")), *FieldName.ToString());
+
+				LoadedFunctions.Emplace(LoadedFunction); // add it even if none to keep the index valid
 			}
 		}
 		{
@@ -378,12 +417,8 @@ TValueOrError<UE::MVVM::FFieldContext, void> FMVVMCompiledBindingLibrary::Evalua
 		else
 		{
 			check(LoadedFunctions.IsValidIndex(PathIndex.Index));
-			UFunction* Function = LoadedFunctions[PathIndex.Index];
-			if (Function && UE::MVVM::Private::IsFunctionVirtual(Function) && CurrentContainer.IsUObject())
-			{
-				Function = CurrentContainer.GetUObject()->GetClass()->FindFunctionByName(Function->GetFName());
-			}
-
+			const FLoadedFunction& LoadedFunction = LoadedFunctions[PathIndex.Index];
+			UFunction* Function = CurrentContainer.IsUObject() ? LoadedFunction.GetFunction(CurrentContainer.GetUObject()) : LoadedFunction.GetFunction();
 			if (!Function)
 			{
 				return MakeError();
@@ -448,12 +483,8 @@ TValueOrError<UE::MVVM::FMVVMFieldVariant, void> FMVVMCompiledBindingLibrary::Ge
 	else
 	{
 		check(LoadedFunctions.IsValidIndex(FinalPathIndex.Index));
-		UFunction* Function = LoadedFunctions[FinalPathIndex.Index];
-		if (Function && UE::MVVM::Private::IsFunctionVirtual(Function) && CurrentContainer.IsUObject())
-		{
-			Function = CurrentContainer.GetUObject()->GetClass()->FindFunctionByName(Function->GetFName());
-		}
-
+		const FLoadedFunction& LoadedFunction = LoadedFunctions[FinalPathIndex.Index];
+		UFunction* Function = CurrentContainer.IsUObject() ? LoadedFunction.GetFunction(CurrentContainer.GetUObject()) : LoadedFunction.GetFunction();
 		if (!Function)
 		{
 			return MakeError();
@@ -523,7 +554,8 @@ TValueOrError<FString, FString> FMVVMCompiledBindingLibrary::FieldPathToString(F
 		{
 			if (LoadedFunctions.IsValidIndex(PathIndex.Index))
 			{
-				const UFunction* Function = LoadedFunctions[PathIndex.Index];
+				const FLoadedFunction& LoadedFunction = LoadedFunctions[PathIndex.Index];
+				const UFunction* Function = LoadedFunction.GetFunction();
 				if (!Function)
 				{
 					StringBuilder << TEXT("<Invalid>");
diff --git a/Engine/Plugins/Runtime/ModelViewViewModel/Source/ModelViewViewModel/Public/Bindings/MVVMCompiledBindingLibrary.h b/Engine/Plugins/Runtime/ModelViewViewModel/Source/ModelViewViewModel/Public/Bindings/MVVMCompiledBindingLibrary.h
--- a/Engine/Plugins/Runtime/ModelViewViewModel/Source/ModelViewViewModel/Public/Bindings/MVVMCompiledBindingLibrary.h
+++ b/Engine/Plugins/Runtime/ModelViewViewModel/Source/ModelViewViewModel/Public/Bindings/MVVMCompiledBindingLibrary.h
@@ -339,11 +339,21 @@ private:
 	TValueOrError<UE::MVVM::FMVVMFieldVariant, void> GetFinalFieldFromPathImpl(UE::MVVM::FObjectVariant CurrentContainer, const FMVVMVCompiledFieldPath& FieldPath) const;
 
 private:
-	TArray<FProperty*> LoadedProperties;
+	struct FLoadedFunction
+	{
+		FLoadedFunction() = default;
+		MODELVIEWVIEWMODEL_API FLoadedFunction(const UFunction* Function);
 
-	UPROPERTY(Transient)
-	TArray<TObjectPtr<UFunction>> LoadedFunctions;
+		TWeakObjectPtr<const UClass> ClassOwner;
+		FName FunctionName;
+		bool bIsFunctionVirtual = false;
 
+		UFunction* GetFunction() const;
+		UFunction* GetFunction(const UObject* CallingContext) const;
+	};
+
+	TArray<FProperty*> LoadedProperties;
+	TArray<FLoadedFunction> LoadedFunctions;
 	TArray<UE::FieldNotification::FFieldId> LoadedFieldIds;
 
 	UPROPERTY()
diff --git a/Engine/Plugins/Runtime/ModelViewViewModel/Source/ModelViewViewModelBlueprint/Private/Bindings/MVVMCompiledBindingLibraryCompiler.cpp b/Engine/Plugins/Runtime/ModelViewViewModel/Source/ModelViewViewModelBlueprint/Private/Bindings/MVVMCompiledBindingLibraryCompiler.cpp
--- a/Engine/Plugins/Runtime/ModelViewViewModel/Source/ModelViewViewModelBlueprint/Private/Bindings/MVVMCompiledBindingLibraryCompiler.cpp
+++ b/Engine/Plugins/Runtime/ModelViewViewModel/Source/ModelViewViewModelBlueprint/Private/Bindings/MVVMCompiledBindingLibraryCompiler.cpp
@@ -751,7 +751,7 @@ TValueOrError<FCompiledBindingLibraryCompiler::FCompileResult, FText> FCompiledB
 			{
 				check(Field.IsFunction());
 				// N.B. no need to translate from skeletal to generated, we only use the the name or the index.
-				Result.Library.LoadedFunctions.Add(const_cast<UFunction*>(Field.GetFunction()));
+				Result.Library.LoadedFunctions.Emplace(Field.GetFunction());
 				FunctionNames.Add(Field.GetName());
 				RawField.LoadedPropertyOrFunctionIndex = TotalNumberOfFunctions;
 				++TotalNumberOfFunctions;
