MVVM: Clean up the code after CL 27409229 code review

diff --git a/Engine/Plugins/Runtime/ModelViewViewModel/Source/ModelViewViewModel/Private/ViewModel/MVVMInstancedViewModelGeneratedClass.cpp b/Engine/Plugins/Runtime/ModelViewViewModel/Source/ModelViewViewModel/Private/ViewModel/MVVMInstancedViewModelGeneratedClass.cpp
--- a/Engine/Plugins/Runtime/ModelViewViewModel/Source/ModelViewViewModel/Private/ViewModel/MVVMInstancedViewModelGeneratedClass.cpp
+++ b/Engine/Plugins/Runtime/ModelViewViewModel/Source/ModelViewViewModel/Private/ViewModel/MVVMInstancedViewModelGeneratedClass.cpp
@@ -39,8 +39,7 @@ DEFINE_FUNCTION(UMVVMInstancedViewModelGeneratedClass::K2_CallNativeOnRep)
 		return;
 	}
 	UMVVMInstancedViewModelGeneratedClass* GeneratedClass = Cast<UMVVMInstancedViewModelGeneratedClass>(CallingObject->GetClass());
-	ensure(GeneratedClass);
-	if (!GeneratedClass)
+	if (!ensure(GeneratedClass))
 	{
 		return;
 	}
diff --git a/Engine/Plugins/Runtime/ModelViewViewModel/Source/ModelViewViewModelBlueprint/Private/MVVMBlueprintInstancedViewModel.cpp b/Engine/Plugins/Runtime/ModelViewViewModel/Source/ModelViewViewModelBlueprint/Private/MVVMBlueprintInstancedViewModel.cpp
--- a/Engine/Plugins/Runtime/ModelViewViewModel/Source/ModelViewViewModelBlueprint/Private/MVVMBlueprintInstancedViewModel.cpp
+++ b/Engine/Plugins/Runtime/ModelViewViewModel/Source/ModelViewViewModelBlueprint/Private/MVVMBlueprintInstancedViewModel.cpp
@@ -152,6 +152,7 @@ void UMVVMBlueprintInstancedViewModelBase::GenerateClass()
 		FromPropertyToCreatedProperty.Empty();
 	}
 
+	// Find a way to call DestroyPropertiesPendingDestruction but after everyone had the time to update.
 	//GeneratedClass->DestroyPropertiesPendingDestruction();
 }
 
@@ -234,7 +235,6 @@ void UMVVMBlueprintInstancedViewModelBase::SetDefaultValue(const FProperty* Prop
 
 	if (NewProperty)
 	{
-		//void* DestinationPtr = reinterpret_cast<uint8*>(GeneratedClass->GetDefaultObject()) + NewProperty->GetOffset_ForInternal();
 		void* DestinationPtr = NewProperty->ContainerPtrToValuePtr<void>(GeneratedClass->GetDefaultObject());
 		NewProperty->CopyCompleteValue(DestinationPtr, ValuePtr);
 	}
diff --git a/Engine/Plugins/Runtime/ModelViewViewModel/Source/ModelViewViewModelEditor/Private/Customizations/MVVMBlueprintViewModelContextCustomization.cpp b/Engine/Plugins/Runtime/ModelViewViewModel/Source/ModelViewViewModelEditor/Private/Customizations/MVVMBlueprintViewModelContextCustomization.cpp
--- a/Engine/Plugins/Runtime/ModelViewViewModel/Source/ModelViewViewModelEditor/Private/Customizations/MVVMBlueprintViewModelContextCustomization.cpp
+++ b/Engine/Plugins/Runtime/ModelViewViewModel/Source/ModelViewViewModelEditor/Private/Customizations/MVVMBlueprintViewModelContextCustomization.cpp
@@ -317,7 +317,8 @@ void FBlueprintViewModelContextDetailCustomization::CustomizeChildren(TSharedRef
 				.Visibility(MakeAttributeLambda([ContextPtr]()
 					{
 						bool bResult = ContextPtr->CreationType == EMVVMBlueprintViewModelContextCreationType::GlobalViewModelCollection
-						 || ContextPtr->CreationType == EMVVMBlueprintViewModelContextCreationType::PropertyPath;
+						 || ContextPtr->CreationType == EMVVMBlueprintViewModelContextCreationType::PropertyPath
+						 || ContextPtr->CreationType == EMVVMBlueprintViewModelContextCreationType::Resolver;
 						return bResult ? EVisibility::Visible : EVisibility::Collapsed;
 					}));
 		}
diff --git a/Engine/Plugins/Runtime/ModelViewViewModel/Source/ModelViewViewModelEditor/Private/Widgets/SMVVMViewBindingPanel.cpp b/Engine/Plugins/Runtime/ModelViewViewModel/Source/ModelViewViewModelEditor/Private/Widgets/SMVVMViewBindingPanel.cpp
--- a/Engine/Plugins/Runtime/ModelViewViewModel/Source/ModelViewViewModelEditor/Private/Widgets/SMVVMViewBindingPanel.cpp
+++ b/Engine/Plugins/Runtime/ModelViewViewModel/Source/ModelViewViewModelEditor/Private/Widgets/SMVVMViewBindingPanel.cpp
@@ -351,24 +351,6 @@ TSharedRef<SWidget> SBindingsPanel::HandleAddDefaultBindingContextMenu()
 			FCanExecuteAction::CreateSP(this, &SBindingsPanel::CanAddBinding)
 		));
 
-	MenuBuilder.AddMenuEntry(
-		LOCTEXT("AddSelectedWidgetEvent", "Add Selected Widget(s) event"),
-		FText::GetEmpty(),
-		FSlateIcon(),
-		FUIAction(
-			FExecuteAction::CreateSP(this, &SBindingsPanel::HandleAddDefaultEventButtonClick, EAddBindingMode::Selected),
-			FCanExecuteAction::CreateSP(this, &SBindingsPanel::CanAddBinding)
-		));
-
-	MenuBuilder.AddMenuEntry(
-		LOCTEXT("AddEmptyWidgetEvent", "Add Empty event"),
-		FText::GetEmpty(),
-		FSlateIcon(),
-		FUIAction(
-			FExecuteAction::CreateSP(this, &SBindingsPanel::HandleAddDefaultEventButtonClick, EAddBindingMode::Empty),
-			FCanExecuteAction::CreateSP(this, &SBindingsPanel::CanAddBinding)
-		));
-
 	return MenuBuilder.MakeWidget();
 }
 
@@ -382,41 +364,6 @@ void SBindingsPanel::HandleAddDefaultBindingButtonClick(EAddBindingMode NewMode)
 	AddDefaultBinding();
 }
 
-void SBindingsPanel::HandleAddDefaultEventButtonClick(EAddBindingMode NewMode)
-{
-	if (UMVVMWidgetBlueprintExtension_View* MVVMExtensionPtr = MVVMExtension.Get())
-	{
-		UMVVMEditorSubsystem* EditorSubsystem = GEditor->GetEditorSubsystem<UMVVMEditorSubsystem>();
-		if (TSharedPtr<FWidgetBlueprintEditor> BlueprintEditor = WeakBlueprintEditor.Pin())
-		{
-			UMVVMBlueprintViewEvent* NewEvent = nullptr;
-			if (AddBindingMode == EAddBindingMode::Selected)
-			{
-				for (const FWidgetReference& WidgetReference : BlueprintEditor->GetSelectedWidgets())
-				{
-					if (WidgetReference.IsValid() && WidgetReference.GetTemplate())
-					{
-						NewEvent = EditorSubsystem->AddEvent(MVVMExtensionPtr->GetWidgetBlueprint());
-						FMVVMBlueprintPropertyPath Path;
-						Path.SetWidgetName(WidgetReference.GetTemplate()->GetFName());
-						EditorSubsystem->SetEventPath(NewEvent, Path);
-					}
-				}
-			}
-
-			if (!NewEvent)
-			{
-				NewEvent = EditorSubsystem->AddEvent(MVVMExtensionPtr->GetWidgetBlueprint());
-			}
-
-			if (NewEvent && BindingsList)
-			{
-				BindingsList->RequestNavigateToEvent(NewEvent);
-			}
-		}
-	}
-}
-
 TSharedRef<SWidget> SBindingsPanel::CreateDrawerDockButton()
 {
 	if (bIsDrawerTab)
diff --git a/Engine/Plugins/Runtime/ModelViewViewModel/Source/ModelViewViewModelEditor/Private/Widgets/SMVVMViewBindingPanel.h b/Engine/Plugins/Runtime/ModelViewViewModel/Source/ModelViewViewModelEditor/Private/Widgets/SMVVMViewBindingPanel.h
--- a/Engine/Plugins/Runtime/ModelViewViewModel/Source/ModelViewViewModelEditor/Private/Widgets/SMVVMViewBindingPanel.h
+++ b/Engine/Plugins/Runtime/ModelViewViewModel/Source/ModelViewViewModelEditor/Private/Widgets/SMVVMViewBindingPanel.h
@@ -63,7 +63,6 @@ private:
 	FText GetAddBindingToolTip() const;
 	TSharedRef<SWidget> HandleAddDefaultBindingContextMenu();
 	void HandleAddDefaultBindingButtonClick(EAddBindingMode NewMode);
-	void HandleAddDefaultEventButtonClick(EAddBindingMode NewMode);
 
 	TSharedRef<SWidget> GenerateSettingsMenu();
 
