[MVVM] Add helper macro to set fields that cannot use templates functions (i.e. bitfields) by inlining the setter functionality.

#test Used locally to add bitfields to a VM

diff --git a/Engine/Plugins/Runtime/ModelViewViewModel/Source/ModelViewViewModel/Public/MVVMViewModelBase.h b/Engine/Plugins/Runtime/ModelViewViewModel/Source/ModelViewViewModel/Public/MVVMViewModelBase.h
--- a/Engine/Plugins/Runtime/ModelViewViewModel/Source/ModelViewViewModel/Public/MVVMViewModelBase.h
+++ b/Engine/Plugins/Runtime/ModelViewViewModel/Source/ModelViewViewModel/Public/MVVMViewModelBase.h
@@ -17,6 +17,10 @@
 #define UE_MVVM_SET_PROPERTY_VALUE(MemberName, NewValue) \
 	SetPropertyValue(MemberName, NewValue, ThisClass::FFieldNotificationClassDescriptor::MemberName)
 
+/** Use this version to set property values that can't be captured as a function arguments (i.e. bitfields). */
+#define UE_MVVM_SET_PROPERTY_VALUE_INLINE(MemberName, NewValue) \
+	[this, InNewValue = (NewValue)]() { if (MemberName == InNewValue) { return false; } MemberName = InNewValue; BroadcastFieldValueChanged(ThisClass::FFieldNotificationClassDescriptor::MemberName); return true; }()
+
 namespace UMMVVMViewModelBase
 {
 	// valid keywords for the UCLASS macro
