Temporary remove some code to fix the build. Will reenabled later when we figure out what is causing the build issue.
#rnx

diff --git a/Engine/Plugins/Runtime/ModelViewViewModel/Source/ModelViewViewModelEditor/Private/Widgets/SMVVMPreviewSourcePanel.cpp b/Engine/Plugins/Runtime/ModelViewViewModel/Source/ModelViewViewModelEditor/Private/Widgets/SMVVMPreviewSourcePanel.cpp
--- a/Engine/Plugins/Runtime/ModelViewViewModel/Source/ModelViewViewModelEditor/Private/Widgets/SMVVMPreviewSourcePanel.cpp
+++ b/Engine/Plugins/Runtime/ModelViewViewModel/Source/ModelViewViewModelEditor/Private/Widgets/SMVVMPreviewSourcePanel.cpp
@@ -69,7 +69,7 @@ void SPreviewSourcePanel::Construct(const FArguments& InArgs, TSharedPtr<FWidget
 		HandlePreviewWidgetChanged();
 		Context->OnPreviewWidgetChanged().AddSP(this, &SPreviewSourcePanel::HandlePreviewWidgetChanged);
 		Context->OnSelectedObjectChanged().AddSP(this, &SPreviewSourcePanel::HandleSelectedObjectChanged);
-		FDebugging::OnViewSourceValueChanged.AddSP(this, &SPreviewSourcePanel::HandleViewChanged);
+		//FDebugging::OnViewSourceValueChanged.AddSP(this, &SPreviewSourcePanel::HandleViewChanged);
 	}
 
 	ChildSlot
@@ -172,16 +172,16 @@ void SPreviewSourcePanel::HandleSourceSelectionChanged(TSharedPtr<Private::SPrev
 }
 
 
-void SPreviewSourcePanel::HandleViewChanged(const FDebugging::FView& View, const FDebugging::FViewSourceValueArgs& Args)
-{
-	if (SourceListView)
-	{
-		if (View.GetView() == WeakView.Get())
-		{
-			SourceListView->RebuildList(); // to prevent access to invalid class, rebuild everything.
-		}
-	}
-}
+//void SPreviewSourcePanel::HandleViewChanged(const FDebugging::FView& View, const FDebugging::FViewSourceValueArgs& Args)
+//{
+//	if (SourceListView)
+//	{
+//		if (View.GetView() == WeakView.Get())
+//		{
+//			SourceListView->RebuildList(); // to prevent access to invalid class, rebuild everything.
+//		}
+//	}
+//}
 
 
 TSharedRef<ITableRow> SPreviewSourcePanel::GenerateWidget(TSharedPtr<Private::SPreviewSourceEntry> Entry, const TSharedRef<STableViewBase>& OwnerTable) const
diff --git a/Engine/Plugins/Runtime/ModelViewViewModel/Source/ModelViewViewModelEditor/Private/Widgets/SMVVMPreviewSourcePanel.h b/Engine/Plugins/Runtime/ModelViewViewModel/Source/ModelViewViewModelEditor/Private/Widgets/SMVVMPreviewSourcePanel.h
--- a/Engine/Plugins/Runtime/ModelViewViewModel/Source/ModelViewViewModelEditor/Private/Widgets/SMVVMPreviewSourcePanel.h
+++ b/Engine/Plugins/Runtime/ModelViewViewModel/Source/ModelViewViewModelEditor/Private/Widgets/SMVVMPreviewSourcePanel.h
@@ -36,7 +36,7 @@ private:
 	void HandlePreviewWidgetChanged();
 	void HandleSelectedObjectChanged();
 	void HandleSourceSelectionChanged(TSharedPtr<Private::SPreviewSourceEntry> Entry, ESelectInfo::Type SelectionType) const;
-	void HandleViewChanged(const FDebugging::FView&, const FDebugging::FViewSourceValueArgs&);
+	//void HandleViewChanged(const FDebugging::FView&, const FDebugging::FViewSourceValueArgs&);
 	TSharedRef<ITableRow> GenerateWidget(TSharedPtr<Private::SPreviewSourceEntry> Entry, const TSharedRef<STableViewBase>& OwnerTable) const;
 
 	TWeakObjectPtr<UMVVMView> WeakView;
