MVVM: Change the key for the TMap for delayed binding.

diff --git a/Engine/Plugins/Runtime/ModelViewViewModel/Source/ModelViewViewModel/Private/View/MVVMBindingSubsystem.cpp b/Engine/Plugins/Runtime/ModelViewViewModel/Source/ModelViewViewModel/Private/View/MVVMBindingSubsystem.cpp
--- a/Engine/Plugins/Runtime/ModelViewViewModel/Source/ModelViewViewModel/Private/View/MVVMBindingSubsystem.cpp
+++ b/Engine/Plugins/Runtime/ModelViewViewModel/Source/ModelViewViewModel/Private/View/MVVMBindingSubsystem.cpp
@@ -32,6 +32,32 @@ void UMVVMBindingSubsystem::Deinitialize()
 	Super::Deinitialize();
 }
 
+namespace UE::MVVM::Private
+{
+struct FViewAndBinding
+{
+	FViewAndBinding(const TObjectKey<const UMVVMView>& InView, FMVVMViewDelayedBinding InBinding)
+		: View(InView)
+		, Binding(InBinding)
+	{
+	}
+	TObjectKey<const UMVVMView> View;
+	FMVVMViewDelayedBinding Binding;
+
+	bool operator== (const FViewAndBinding& Other) const
+	{
+		return View == Other.View && Binding == Other.Binding;
+	}
+
+	friend uint32 GetTypeHash(const FViewAndBinding& Key)
+	{
+		uint32 Value1 = GetTypeHash(Key.View);
+		uint32 Value2 = GetTypeHash(Key.Binding.GetCompiledBindingIndex());
+		return HashCombine(Value1, Value2);
+	}
+};
+}
+
 void UMVVMBindingSubsystem::HandlePreTick(float DeltaTime)
 {
 	SCOPE_CYCLE_COUNTER(STAT_MVVMBindingTick);
@@ -57,7 +83,7 @@ void UMVVMBindingSubsystem::HandlePreTick(float DeltaTime)
 
 	if (DelayedBindings.Num() > 0)
 	{
-		FDelayedMap AllDelayedBindingsExecutedThisFrame;
+		TSet<UE::MVVM::Private::FViewAndBinding> AllDelayedBindingsExecutedThisFrame;
 		AllDelayedBindingsExecutedThisFrame.Reserve(DelayedBindings.Num());
 
 		FDelayedMap DelayedBindingsWhileTicking = MoveTemp(DelayedBindings);
@@ -67,38 +93,50 @@ void UMVVMBindingSubsystem::HandlePreTick(float DeltaTime)
 		{
 			for (const auto& DelayedBindingsPair : DelayedBindingsWhileTicking)
 			{
-				if (const UMVVMView* View = DelayedBindingsPair.Key.Get())
+				if (const UMVVMView* View = DelayedBindingsPair.Key.ResolveObjectPtr())
 				{
 					ensure(DelayedBindingsPair.Value.Num() > 0);
-					FDelayedBindingList& AllBindingList = AllDelayedBindingsExecutedThisFrame.FindOrAdd(DelayedBindingsPair.Key);
 					for (const FMVVMViewDelayedBinding& DelayedBinding : DelayedBindingsPair.Value)
 					{
 						View->ExecuteDelayedBinding(DelayedBinding);
-						AllBindingList.AddUnique(DelayedBinding);
+						UE::MVVM::Private::FViewAndBinding ViewAndBinding = UE::MVVM::Private::FViewAndBinding(DelayedBindingsPair.Key, DelayedBinding);
+						AllDelayedBindingsExecutedThisFrame.Add(ViewAndBinding);
 					}
 				}
 			}
 
 			DelayedBindingsWhileTicking.Reset();
 
-			// Test new binding added while executing the current binding list
-			for (const auto& DelayedBindingPair : DelayedBindings)
+			// Test new bindings added while executing the latest binding list.
+			//If it's a new  binding (not already executed this frame), execute it this frame. Else, execute it next frame.
+			for (auto DelayedBindingItt = DelayedBindings.CreateIterator(); DelayedBindingItt; ++DelayedBindingItt)
 			{
-				if (DelayedBindingPair.Key.Get())
+				FDelayedBindingList* FoundDelayedBindingListPtr = nullptr;
+				for (int32 DelayIndex = DelayedBindingItt.Value().Num() - 1; DelayIndex >= 0; --DelayIndex)
 				{
-					if (FDelayedBindingList* AllDelayedBindingsListPtr = AllDelayedBindingsExecutedThisFrame.Find(DelayedBindingPair.Key))
+					// Was it executed this frame
+					const FMVVMViewDelayedBinding& DelayedBinding = DelayedBindingItt.Value()[DelayIndex];
+					UE::MVVM::Private::FViewAndBinding ViewAndBinding = UE::MVVM::Private::FViewAndBinding(DelayedBindingItt.Key(), DelayedBinding);
+					if (!AllDelayedBindingsExecutedThisFrame.Find(ViewAndBinding))
 					{
-						for (const FMVVMViewDelayedBinding& NewCompiledBinding : DelayedBindingPair.Value)
+						if (!FoundDelayedBindingListPtr)
 						{
-							if (!AllDelayedBindingsListPtr->Find(NewCompiledBinding))
-							{
-								// It was not executed. Add it to be executed this frame.
-								DelayedBindingsWhileTicking.FindOrAdd(DelayedBindingPair.Key).AddUnique(NewCompiledBinding);
-							}
+							FoundDelayedBindingListPtr = &DelayedBindingsWhileTicking.FindOrAdd(DelayedBindingItt.Key());
 						}
+						FoundDelayedBindingListPtr->AddUnique(DelayedBinding);
+						DelayedBindingItt.Value().RemoveAtSwap(DelayIndex);
+					}
+
+					// If they were all executed, remove the key from the next frame list.
+					if (DelayedBindingItt.Value().Num() == 0)
+					{
+						DelayedBindingItt.RemoveCurrent();
 					}
 				}
 			}
+
+			//DelayedBindings = FDelayedMap(); // do not reset the array. Bindings could be executed this frame and needs to be re executed next frame
+
 		} while (DelayedBindingsWhileTicking.Num() > 0);
 	}
 }
diff --git a/Engine/Plugins/Runtime/ModelViewViewModel/Source/ModelViewViewModel/Private/View/MVVMBindingSubsystem.h b/Engine/Plugins/Runtime/ModelViewViewModel/Source/ModelViewViewModel/Private/View/MVVMBindingSubsystem.h
--- a/Engine/Plugins/Runtime/ModelViewViewModel/Source/ModelViewViewModel/Private/View/MVVMBindingSubsystem.h
+++ b/Engine/Plugins/Runtime/ModelViewViewModel/Source/ModelViewViewModel/Private/View/MVVMBindingSubsystem.h
@@ -3,6 +3,7 @@
 #pragma once
 
 #include "Subsystems/EngineSubsystem.h"
+#include "UObject/ObjectKey.h"
 
 #include "MVVMBindingSubsystem.generated.h"
 
@@ -54,7 +55,7 @@ private:
 	void HandlePreTick(float DeltaTIme);
 
 	using FDelayedBindingList = TArray<FMVVMViewDelayedBinding, TInlineAllocator<8>>;
-	using FDelayedMap = TMap<TWeakObjectPtr<const UMVVMView>, FDelayedBindingList>;
+	using FDelayedMap = TMap<TObjectKey<const UMVVMView>, FDelayedBindingList>;
 	FDelayedMap DelayedBindings;
 	TArray<TWeakObjectPtr<const UMVVMView>> EveryTickBindings;
 };
