MVVM: Change the display name of Conv_BoolToSlateVisibility
#rnx

diff --git a/Engine/Plugins/Runtime/ModelViewViewModel/Source/ModelViewViewModel/Public/Bindings/MVVMConversionLibrary.h b/Engine/Plugins/Runtime/ModelViewViewModel/Source/ModelViewViewModel/Public/Bindings/MVVMConversionLibrary.h
--- a/Engine/Plugins/Runtime/ModelViewViewModel/Source/ModelViewViewModel/Public/Bindings/MVVMConversionLibrary.h
+++ b/Engine/Plugins/Runtime/ModelViewViewModel/Source/ModelViewViewModel/Public/Bindings/MVVMConversionLibrary.h
@@ -15,6 +15,6 @@ class MODELVIEWVIEWMODEL_API UMVVMConversionLibrary : public UBlueprintFunctionL
 	/**
 	 * Converts a bool to a slate visibility.
 	 */
-	UFUNCTION(BlueprintPure, Category = "Widget", meta = (DisplayName = "To Bool (Visibility)"))
+	UFUNCTION(BlueprintPure, Category = "Widget", meta = (DisplayName = "To Visibility (Boolean)"))
 	static ESlateVisibility Conv_BoolToSlateVisibility(bool bIsVisible, ESlateVisibility TrueVisibility = ESlateVisibility::Visible, ESlateVisibility FalseVisibility = ESlateVisibility::Collapsed);
 };
\ No newline at end of file
