MVVM: Update the InstanceViewModel to be more generic. Force a compilation when the widget loads, in case something changed in the cpp. Only rebuild the class when needed. Change the default name of the generated viewmodel.

diff --git a/Engine/Plugins/Runtime/ModelViewViewModel/Source/ModelViewViewModelBlueprint/Private/MVVMBlueprintInstancedViewModel.cpp b/Engine/Plugins/Runtime/ModelViewViewModel/Source/ModelViewViewModelBlueprint/Private/MVVMBlueprintInstancedViewModel.cpp
--- a/Engine/Plugins/Runtime/ModelViewViewModel/Source/ModelViewViewModelBlueprint/Private/MVVMBlueprintInstancedViewModel.cpp
+++ b/Engine/Plugins/Runtime/ModelViewViewModel/Source/ModelViewViewModelBlueprint/Private/MVVMBlueprintInstancedViewModel.cpp
@@ -11,6 +11,7 @@
 #include "Engine/Engine.h"
 #include "Kismet2/BlueprintEditorUtils.h"
 #include "Kismet2/KismetEditorUtilities.h"
+#include "Hash/CityHash.h"
 
 #include UE_INLINE_GENERATED_CPP_BY_NAME(MVVMBlueprintInstancedViewModel)
 
@@ -22,7 +23,7 @@ UMVVMBlueprintInstancedViewModelBase::UMVVMBlueprintInstancedViewModelBase()
 	ParentClass = UMVVMViewModelBase::StaticClass();
 }
 
-void UMVVMBlueprintInstancedViewModelBase::GenerateClass()
+void UMVVMBlueprintInstancedViewModelBase::GenerateClass(bool bForceGeneration)
 {
 	if (ParentClass.Get() == nullptr)
 	{
@@ -37,91 +38,43 @@ void UMVVMBlueprintInstancedViewModelBase::GenerateClass()
 		GeneratedClassType = UMVVMInstancedViewModelGeneratedClass::StaticClass();
 	}
 
-	auto SafeRename = [](UObject* Object)
-	{
-		ERenameFlags RenameFlags = REN_ForceNoResetLoaders | REN_NonTransactional | REN_DoNotDirty | REN_DontCreateRedirectors;
-		FName TrashName = MakeUniqueObjectName(GetTransientPackage(), Object->GetClass(), *FString::Printf(TEXT("TRASH_%s"), *Object->GetName()));
-		Object->Rename(*TrashName.ToString(), GetTransientPackage(), RenameFlags);
-	};
+	PreloadObjectsForCompilation();
 
-	bool bGeneratedNewClass = false;
+	UObject* PreviousClass = nullptr;
 	if (GeneratedClass == nullptr)
 	{
-		bGeneratedNewClass = true;
-		GeneratedClass = NewObject<UMVVMInstancedViewModelGeneratedClass>(GetOutermost(), GeneratedClassType.Get());
+		UPackage* Outermost = GetOutermost();
+		FName NewClassName = *FString::Printf(TEXT("InstanceViewmodel%d_IC"), GetFName().GetNumber());
+		PreviousClass = StaticFindObjectFastInternal(nullptr, Outermost, NewClassName, true);
+		if (PreviousClass)
+		{
+			SafeRename(PreviousClass);
+		}
+		GeneratedClass = NewObject<UMVVMInstancedViewModelGeneratedClass>(Outermost, GeneratedClassType.Get(), NewClassName);
 	}
 
-	UObject* PreviousDefaultObject = GeneratedClass->GetDefaultObject(false);
-	if (PreviousDefaultObject)
+	if (bForceGeneration == false && !IsClassDirty())
 	{
-		SafeRename(PreviousDefaultObject);
-	}
-	for (TFieldIterator<UFunction> FunctionIter(GeneratedClass, EFieldIteratorFlags::ExcludeSuper); FunctionIter; ++FunctionIter)
-	{
-		FunctionIter->FunctionFlags &= ~FUNC_Native;
-		SafeRename(*FunctionIter);
+		return;
 	}
 
-	// Clean class and reset basic properties
-	GeneratedClass->PurgeClass(false);
-	GeneratedClass->PropertyLink = ParentClass->PropertyLink;
-	GeneratedClass->SetSuperStruct(ParentClass);
-	GeneratedClass->ClassWithin = UObject::StaticClass();
-	GeneratedClass->ClassConfigName = ParentClass->ClassConfigName;
-	GeneratedClass->ClassFlags |= CLASS_NotPlaceable;
+	UObject* PreviousDefaultObject = GeneratedClass->GetDefaultObject(false);
+	CleanClass();
+	AddProperties();
+	ConstructClass();
 
-	// Clean up temporary generate variables
+	// Relink the new default object
+	if (PreviousClass)
 	{
-		FromPropertyToCreatedProperty.Empty();
+		FLinkerLoad::PRIVATE_PatchNewObjectIntoExport(PreviousClass, GeneratedClass);
 	}
-
-	// Create Properties
-	{
-		PreAddProperties();
-		for (TPropertyValueIterator<FProperty> It(GetSourceStruct(), GetSourceDefaults(), EPropertyValueIteratorFlags::NoRecursion, EFieldIteratorFlags::ExcludeDeprecated); It; ++It)
-		{
-			const FProperty* Property = It.Key();
-			void const* ValuePtr = It.Value();
-
-			AddProperty(Property);
-		}
-		PostAddProperties();
-	}
-
-	// Update the class
-	GeneratedClass->Bind();
-	GeneratedClass->StaticLink(true);
-	GeneratedClass->AssembleReferenceTokenStream();
-	ensure(GeneratedClass->ClassGeneratedBy == nullptr);
-	ensure(GeneratedClass->GetDefaultObject(false) == nullptr);
-	GeneratedClass->GetDefaultObject(true);
-	GeneratedClass->UpdateCustomPropertyListForPostConstruction();
-
-	// The class is not public and can only be access inside the UMG. What about inherited UMG???
-	GeneratedClass->ClearFlags(RF_Public | RF_Transactional);
-	GeneratedClass->GetDefaultObject(true)->ClearFlags(RF_Public);
-
-	// Relink the new default object
 	if (PreviousDefaultObject)
 	{
 		FLinkerLoad::PRIVATE_PatchNewObjectIntoExport(PreviousDefaultObject, GeneratedClass->GetDefaultObject(true));
 	}
 
 	// Initialize the default object value
-	{
-		PreSetDefaultValues();
-		for (TPropertyValueIterator<FProperty> It(GetSourceStruct(), GetSourceDefaults(), EPropertyValueIteratorFlags::NoRecursion, EFieldIteratorFlags::ExcludeDeprecated); It; ++It)
-		{
-			const FProperty* Property = It.Key();
-			void const* ValuePtr = It.Value();
-
-			SetDefaultValue(Property, ValuePtr);
-		}
-		PostSetDefaultValues();
-	}
-
-	// Initialize the fields for MVVM
-	GeneratedClass->InitializeFieldNotifies();
+	SetDefaultValues();
 
 	// Inform that the class changed
 	{
@@ -130,10 +83,11 @@ void UMVVMBlueprintInstancedViewModelBase::GenerateClass()
 		{
 			OldToNew.Emplace(PreviousDefaultObject, GeneratedClass->GetDefaultObject(true));
 		}
-		if (bGeneratedNewClass)
+		if (PreviousClass)
 		{
-			OldToNew.Emplace(GeneratedClass, GeneratedClass);
+			OldToNew.Emplace(PreviousClass, GeneratedClass);
 		}
+		OldToNew.Emplace(GeneratedClass, GeneratedClass);
 		if (GEngine && OldToNew.Num() > 0)
 		{
 			GEngine->NotifyToolsOfObjectReplacement(OldToNew);
@@ -147,97 +101,132 @@ void UMVVMBlueprintInstancedViewModelBase::GenerateClass()
 		ActionDB->RefreshClassActions(GeneratedClass);
 	}
 
-	// Clean up temporary generate variables
-	{
-		FromPropertyToCreatedProperty.Empty();
-	}
+	ClassGenerated();
 
 	// Find a way to call DestroyPropertiesPendingDestruction but after everyone had the time to update.
-	//GeneratedClass->DestroyPropertiesPendingDestruction();
+	GeneratedClass->DestroyPropertiesPendingDestruction();
+	GeneratedClass->Modify();
 }
 
-bool UMVVMBlueprintInstancedViewModelBase::IsValidFieldName(const FName NewPropertyName) const
+void UMVVMBlueprintInstancedViewModelBase::PreloadObjectsForCompilation()
 {
-	if (!FName::IsValidXName(NewPropertyName, INVALID_NAME_CHARACTERS))
+
+}
+
+bool UMVVMBlueprintInstancedViewModelBase::IsClassDirty() const
+{
+	return true;
+}
+
+void UMVVMBlueprintInstancedViewModelBase::CleanClass()
+{
+	UObject* PreviousDefaultObject = GeneratedClass->GetDefaultObject(false);
+	if (PreviousDefaultObject)
 	{
-		return false;
+		SafeRename(PreviousDefaultObject);
 	}
-
-	// Check if the name already exist. If it does, do not add the property.
-	for (TFieldIterator<FField> PropertyIter(GeneratedClass, EFieldIteratorFlags::IncludeSuper); PropertyIter; ++PropertyIter)
+	for (TFieldIterator<UFunction> FunctionIter(GeneratedClass, EFieldIteratorFlags::ExcludeSuper); FunctionIter; ++FunctionIter)
 	{
-		if (PropertyIter->GetFName() == NewPropertyName)
-		{
-			return false;
-		}
+		FunctionIter->FunctionFlags &= ~FUNC_Native;
+		SafeRename(*FunctionIter);
 	}
-	for (TFieldIterator<UField> FunctionIter(GeneratedClass, EFieldIteratorFlags::IncludeSuper); FunctionIter; ++FunctionIter)
+
+	for (TFieldIterator<FField> FieldIter(GeneratedClass, EFieldIteratorFlags::ExcludeSuper); FieldIter; ++FieldIter)
 	{
-		if (FunctionIter->GetFName() == NewPropertyName)
-		{
-			return false;
-		}
+		FName NewName = *FString::Printf(TEXT("TRASH_%s"), *FieldIter->GetName());
+		FieldIter->Rename(NewName);
 	}
 
-	return true;
+	// Clean class and reset basic properties
+	GeneratedClass->PurgeClass(false);
+	GeneratedClass->PropertyLink = ParentClass->PropertyLink;
+	GeneratedClass->SetSuperStruct(ParentClass);
+	GeneratedClass->ClassWithin = UObject::StaticClass();
+	GeneratedClass->ClassConfigName = ParentClass->ClassConfigName;
+	GeneratedClass->ClassFlags |= CLASS_NotPlaceable;
 }
 
-void UMVVMBlueprintInstancedViewModelBase::PreAddProperties()
+void UMVVMBlueprintInstancedViewModelBase::AddProperties()
 {
-
 }
 
-void UMVVMBlueprintInstancedViewModelBase::PostAddProperties()
+void UMVVMBlueprintInstancedViewModelBase::ConstructClass()
 {
+	check(GeneratedClass);
+	GeneratedClass->Bind();
+	GeneratedClass->StaticLink(true);
+	GeneratedClass->AssembleReferenceTokenStream();
+	ensure(GeneratedClass->ClassGeneratedBy == nullptr);
+	ensure(GeneratedClass->GetDefaultObject(false) == nullptr);
+	GeneratedClass->GetDefaultObject(true);
+	GeneratedClass->UpdateCustomPropertyListForPostConstruction();
+	GeneratedClass->SetUpRuntimeReplicationData();
+	// Initialize the fields for MVVM
+	GeneratedClass->InitializeFieldNotifies();
+
+	// The class is not public and can only be access inside the UMG. What about inherited UMG???
+	GeneratedClass->ClearFlags(RF_Public | RF_Transactional);
+	GeneratedClass->GetDefaultObject(true)->ClearFlags(RF_Public);
 }
 
-void UMVVMBlueprintInstancedViewModelBase::AddProperty(const FProperty* FromProperty)
+void UMVVMBlueprintInstancedViewModelBase::SetDefaultValues()
 {
-	if (!IsValidFieldName(FromProperty->GetFName()))
-	{
-		return;
-	}
 
-	FProperty* NewProperty = CastFieldChecked<FProperty>(FField::Duplicate(FromProperty, GeneratedClass, FromProperty->GetFName()));
-#if WITH_EDITOR
-	FField::CopyMetaData(FromProperty, NewProperty);
-#endif
-	FInitializePropertyArgs Args;
-	Args.PropertyName = FromProperty->GetFName();
-	Args.DisplayName = FromProperty->GetMetaData(FBlueprintMetadata::MD_DisplayName);
-	Args.bNetwork = true;
-	InitializeProperty(NewProperty, Args);
-
-	LinkProperty(NewProperty);
-	FromPropertyToCreatedProperty.Add(FromProperty, NewProperty);
 }
 
-
-void UMVVMBlueprintInstancedViewModelBase::PreSetDefaultValues()
+void UMVVMBlueprintInstancedViewModelBase::ClassGenerated()
 {
+	GetOutermost()->Modify(true);
 }
 
-void UMVVMBlueprintInstancedViewModelBase::PostSetDefaultValues()
+bool UMVVMBlueprintInstancedViewModelBase::IsValidFieldName(const FName NewPropertyName) const
 {
+	return IsValidFieldName(NewPropertyName, GeneratedClass);
 }
 
-void UMVVMBlueprintInstancedViewModelBase::SetDefaultValue(const FProperty* Property, void const* ValuePtr)
+
+bool UMVVMBlueprintInstancedViewModelBase::IsValidFieldName(const FName NewPropertyName, UStruct* NewOwner) const
 {
-	FProperty* NewProperty = nullptr;
-	if (FProperty** NewPropertyPtr = FromPropertyToCreatedProperty.Find(Property))
+	if (!FName::IsValidXName(NewPropertyName, INVALID_NAME_CHARACTERS))
 	{
-		NewProperty = *NewPropertyPtr;
+		return false;
 	}
-	else
+
+	// Check if the name already exist. If it does, do not add the property.
+	for (TFieldIterator<FField> PropertyIter(NewOwner, EFieldIteratorFlags::IncludeSuper); PropertyIter; ++PropertyIter)
 	{
-		NewProperty = GeneratedClass->FindPropertyByName(Property->GetFName());
+		if (PropertyIter->GetFName() == NewPropertyName)
+		{
+			return false;
+		}
 	}
+	for (TFieldIterator<UField> FunctionIter(NewOwner, EFieldIteratorFlags::IncludeSuper); FunctionIter; ++FunctionIter)
+	{
+		if (FunctionIter->GetFName() == NewPropertyName)
+		{
+			return false;
+		}
+	}
+
+	return true;
+}
+
+FProperty* UMVVMBlueprintInstancedViewModelBase::CreateProperty(const FProperty* FromProperty, UStruct* NewOwner)
+{
+	return CreateProperty(FromProperty, NewOwner, FromProperty->GetFName());
+}
 
-	if (NewProperty)
+FProperty* UMVVMBlueprintInstancedViewModelBase::CreateProperty(const FProperty* FromProperty, UStruct* NewOwner, FName NewPropertyName)
+{
+	if (!IsValidFieldName(NewPropertyName))
 	{
-		void* DestinationPtr = NewProperty->ContainerPtrToValuePtr<void>(GeneratedClass->GetDefaultObject());
-		NewProperty->CopyCompleteValue(DestinationPtr, ValuePtr);
+		return nullptr;
 	}
+
+	FProperty* NewProperty = CastFieldChecked<FProperty>(FField::Duplicate(FromProperty, NewOwner, NewPropertyName));
+	FField::CopyMetaData(FromProperty, NewProperty);
+
+	return NewProperty;
 }
 
 void UMVVMBlueprintInstancedViewModelBase::InitializeProperty(FProperty* NewProperty, FInitializePropertyArgs& Args)
@@ -259,10 +248,14 @@ void UMVVMBlueprintInstancedViewModelBase::InitializeProperty(FProperty* NewProp
 	}
 }
 
-void UMVVMBlueprintInstancedViewModelBase::LinkProperty(FProperty* NewProperty)
+void UMVVMBlueprintInstancedViewModelBase::LinkProperty(FProperty* NewProperty) const
+{
+	LinkProperty(NewProperty, GeneratedClass);
+}
+
+void UMVVMBlueprintInstancedViewModelBase::LinkProperty(FProperty* NewProperty, UStruct* NewOwner) const
 {
-	NewProperty->Next = GeneratedClass->ChildProperties;
-	GeneratedClass->ChildProperties = NewProperty;
+	NewOwner->AddCppProperty(NewProperty);
 }
 
 FName UMVVMBlueprintInstancedViewModelBase::AddOnRepFunction(FName PropertyName)
@@ -286,18 +279,149 @@ FName UMVVMBlueprintInstancedViewModelBase::AddOnRepFunction(FName PropertyName)
 	return Func->GetFName();
 }
 
+void UMVVMBlueprintInstancedViewModelBase::SafeRename(UObject* Object)
+{
+	ERenameFlags RenameFlags = REN_ForceNoResetLoaders | REN_NonTransactional | REN_DoNotDirty | REN_DontCreateRedirectors;
+	FName TrashName = MakeUniqueObjectName(GetTransientPackage(), Object->GetClass(), *FString::Printf(TEXT("TRASH_%s"), *Object->GetName()));
+	Object->Rename(*TrashName.ToString(), GetTransientPackage(), RenameFlags);
+}
+
+void UMVVMBlueprintInstancedViewModelBase::SetDefaultValue(const FProperty* SourceProperty, void const* SourceValuePtr, const FProperty* DestinationProperty)
+{
+	check(DestinationProperty);
+	check(SourceProperty);
+	check(SourceValuePtr);
+	{
+		void* DestinationPtr = DestinationProperty->ContainerPtrToValuePtr<void>(GeneratedClass->GetDefaultObject());
+		DestinationProperty->CopyCompleteValue(DestinationPtr, SourceValuePtr);
+	}
+}
+
+/**
+ *
+ */
+namespace UE::MVVM::Private
+{
+uint64 GetObjectHash(const UObject* Object)
+{
+	const FString PathName = GetPathNameSafe(Object);
+	return CityHash64((const char*)GetData(PathName), PathName.Len() * sizeof(TCHAR));
+}
+
+uint64 CalcPropertyDescHash(const FPropertyBagPropertyDesc& Desc)
+{
+	UStruct* ValueTypeObject = Cast<UStruct>(Desc.ValueTypeObject);
+	FTopLevelAssetPath ValueTypeObjectPath = ValueTypeObject ? ValueTypeObject->GetStructPathName() : FTopLevelAssetPath();
+#if WITH_EDITORONLY_DATA
+	const uint32 Hashes[] = { GetTypeHash(ValueTypeObjectPath), GetTypeHash(Desc.ID), GetTypeHash(Desc.Name), GetTypeHash(Desc.ValueType), GetTypeHash(Desc.ContainerTypes), GetTypeHash(Desc.MetaData) };
+#else
+	const uint32 Hashes[] = { GetTypeHash(ValueTypeObjectPath), GetTypeHash(Desc.ID), GetTypeHash(Desc.Name), GetTypeHash(Desc.ValueType), GetTypeHash(Desc.ContainerTypes) };
+#endif
+	return CityHash64((const char*)Hashes, sizeof(Hashes));
+}
+
+uint64 CalcPropertyDescArrayHash(const TConstArrayView<FPropertyBagPropertyDesc> Descs)
+{
+	uint64 Hash = 0;
+	for (const FPropertyBagPropertyDesc& Desc : Descs)
+	{
+		Hash = CityHash128to64(Uint128_64(Hash, CalcPropertyDescHash(Desc)));
+	}
+	// Should add the default values.
+	return Hash;
+}
+} //namespace
+
+bool UMVVMBlueprintInstancedViewModel_PropertyBag::IsClassDirty() const
+{
+	const UPropertyBag* PropertyBag = Variables.GetPropertyBagStruct();
+	if (PropertyBag)
+	{
+		return UE::MVVM::Private::CalcPropertyDescArrayHash(PropertyBag->GetPropertyDescs()) != PropertiesHash;
+	}
+	return PropertiesHash != 0;
+}
+
+void UMVVMBlueprintInstancedViewModel_PropertyBag::CleanClass()
+{
+	Super::CleanClass();
+
+	FromPropertyToCreatedProperty.Empty();
+}
+
+void UMVVMBlueprintInstancedViewModel_PropertyBag::AddProperties()
+{
+	Super::AddProperties();
+
+	for (TPropertyValueIterator<FProperty> It(GetSourceStruct(), GetSourceDefaults(), EPropertyValueIteratorFlags::NoRecursion, EFieldIteratorFlags::ExcludeDeprecated); It; ++It)
+	{
+		const FProperty* Property = It.Key();
+		void const* ValuePtr = It.Value();
+
+		FProperty* NewProperty = CreateProperty(Property, GeneratedClass);
+		if (NewProperty)
+		{
+			FInitializePropertyArgs Args;
+			Args.PropertyName = NewProperty->GetFName();
+			Args.bFieldNotify = true;
+			//Args.bNetwork = true;
+			InitializeProperty(NewProperty, Args);
+
+			LinkProperty(NewProperty);
+			FromPropertyToCreatedProperty.Add(Property, NewProperty);
+		}
+	}
+}
+
+void UMVVMBlueprintInstancedViewModel_PropertyBag::SetDefaultValues()
+{
+	Super::SetDefaultValues();
+
+	for (TPropertyValueIterator<FProperty> It(GetSourceStruct(), GetSourceDefaults(), EPropertyValueIteratorFlags::NoRecursion, EFieldIteratorFlags::ExcludeDeprecated); It; ++It)
+	{
+		const FProperty* Property = It.Key();
+		void const* ValuePtr = It.Value();
+
+		FProperty* NewProperty = nullptr;
+		if (FProperty** NewPropertyPtr = FromPropertyToCreatedProperty.Find(Property))
+		{
+			NewProperty = *NewPropertyPtr;
+		}
+		else
+		{
+			NewProperty = GeneratedClass->FindPropertyByName(Property->GetFName());
+		}
+
+		if (NewProperty)
+		{
+			SetDefaultValue(Property, ValuePtr, NewProperty);
+		}
+	}
+}
+
+void UMVVMBlueprintInstancedViewModel_PropertyBag::ClassGenerated()
+{
+	Super::ClassGenerated();
+
+	FromPropertyToCreatedProperty.Empty();
+	const UPropertyBag* PropertyBag = Variables.GetPropertyBagStruct();
+	PropertiesHash = PropertyBag ? UE::MVVM::Private::CalcPropertyDescArrayHash(PropertyBag->GetPropertyDescs()) : 0;
+}
+
 #if WITH_EDITOR
-void UMVVMBlueprintInstancedViewModel::PostEditChangeChainProperty(FPropertyChangedChainEvent& PropertyChangedEvent)
+void UMVVMBlueprintInstancedViewModel_PropertyBag::PostEditChangeChainProperty(FPropertyChangedChainEvent& PropertyChangedEvent)
 {
 	Super::PostEditChangeChainProperty(PropertyChangedEvent);
 
-	FProperty* VariableProperty = GetClass()->FindPropertyByName(GET_MEMBER_NAME_CHECKED(UMVVMBlueprintInstancedViewModel, Variables));
+	FProperty* VariableProperty = GetClass()->FindPropertyByName(GET_MEMBER_NAME_CHECKED(UMVVMBlueprintInstancedViewModel_PropertyBag, Variables));
+	FProperty* ParentClassProperty = GetClass()->FindPropertyByName(GET_MEMBER_NAME_CHECKED(UMVVMBlueprintInstancedViewModel_PropertyBag, ParentClass));
 	FEditPropertyChain::TDoubleLinkedListNode* CurrentPropertyNode = PropertyChangedEvent.PropertyChain.GetActiveMemberNode();
 	while (CurrentPropertyNode)
 	{
-		if (CurrentPropertyNode->GetValue() == VariableProperty)
+		if (CurrentPropertyNode->GetValue() == VariableProperty || CurrentPropertyNode->GetValue() == ParentClassProperty)
 		{
 			FBlueprintEditorUtils::MarkBlueprintAsModified(GetOuterUMVVMBlueprintView()->GetOuterUMVVMWidgetBlueprintExtension_View()->GetWidgetBlueprint());
+			PropertiesHash = 0; // force class regeneration
 			break;
 		}
 		CurrentPropertyNode = CurrentPropertyNode->GetNextNode();
diff --git a/Engine/Plugins/Runtime/ModelViewViewModel/Source/ModelViewViewModelBlueprint/Private/MVVMWidgetBlueprintExtension_View.cpp b/Engine/Plugins/Runtime/ModelViewViewModel/Source/ModelViewViewModelBlueprint/Private/MVVMWidgetBlueprintExtension_View.cpp
--- a/Engine/Plugins/Runtime/ModelViewViewModel/Source/ModelViewViewModelBlueprint/Private/MVVMWidgetBlueprintExtension_View.cpp
+++ b/Engine/Plugins/Runtime/ModelViewViewModel/Source/ModelViewViewModelBlueprint/Private/MVVMWidgetBlueprintExtension_View.cpp
@@ -60,10 +60,7 @@ void UMVVMWidgetBlueprintExtension_View::HandlePreloadObjectsForCompilation(UBlu
 			if (AvailableViewModel.InstancedViewModel)
 			{
 				UBlueprint::ForceLoad(AvailableViewModel.InstancedViewModel);
-				if (AvailableViewModel.InstancedViewModel->GetGeneratedClass())
-				{
-					AvailableViewModel.InstancedViewModel->GetGeneratedClass()->ConditionalPostLoad();
-				}
+				AvailableViewModel.InstancedViewModel->GenerateClass(true);
 			}
 			if (AvailableViewModel.GetViewModelClass())
 			{
@@ -91,7 +88,7 @@ void UMVVMWidgetBlueprintExtension_View::HandleBeginCompilation(FWidgetBlueprint
 	{
 		if (AvailableViewModel.InstancedViewModel)
 		{
-			AvailableViewModel.InstancedViewModel->GenerateClass();
+			AvailableViewModel.InstancedViewModel->GenerateClass(false);
 		}
 	}
 
diff --git a/Engine/Plugins/Runtime/ModelViewViewModel/Source/ModelViewViewModelBlueprint/Public/MVVMBlueprintInstancedViewModel.h b/Engine/Plugins/Runtime/ModelViewViewModel/Source/ModelViewViewModelBlueprint/Public/MVVMBlueprintInstancedViewModel.h
--- a/Engine/Plugins/Runtime/ModelViewViewModel/Source/ModelViewViewModelBlueprint/Public/MVVMBlueprintInstancedViewModel.h
+++ b/Engine/Plugins/Runtime/ModelViewViewModel/Source/ModelViewViewModelBlueprint/Public/MVVMBlueprintInstancedViewModel.h
@@ -20,42 +20,42 @@ class MODELVIEWVIEWMODELBLUEPRINT_API UMVVMBlueprintInstancedViewModelBase : pub
 public:
 	UMVVMBlueprintInstancedViewModelBase();
 
-	void GenerateClass();
+	void GenerateClass(bool bForceGeneration);
 
 	UClass* GetGeneratedClass() const
 	{
 		return GeneratedClass;
 	}
 
-	virtual const UStruct* GetSourceStruct() const PURE_VIRTUAL(GetSourceStruct, return nullptr;);
-	virtual const uint8* GetSourceDefaults() const PURE_VIRTUAL(GetSourceDefaults, return nullptr;);
-
 protected:
-	virtual void PreAddProperties();
-	virtual void AddProperty(const FProperty* FromProperty);
-	virtual void PostAddProperties();
-
-	virtual void PreSetDefaultValues();
-	virtual void SetDefaultValue(const FProperty* FromProperty, void const* SourceValuePtr);
-	virtual void PostSetDefaultValues();
-
-	virtual bool IsValidFieldName(const FName NewPropertyName) const;
+	virtual void PreloadObjectsForCompilation();
+	virtual bool IsClassDirty() const;
+	virtual void CleanClass();
+	virtual void AddProperties();
+	virtual void ConstructClass();
+	virtual void SetDefaultValues();
+	virtual void ClassGenerated();
 
 protected:
 	struct FInitializePropertyArgs
 	{
 		FName PropertyName;
-		FString DisplayName;
-		bool bFieldNotify = true;
+		bool bFieldNotify = false;
 		bool bReadOnly = false;
 		bool bNetwork = false;
 		bool bPrivate = false;
 	};
+
+	bool IsValidFieldName(const FName NewPropertyName) const;
+	bool IsValidFieldName(const FName NewPropertyName, UStruct* NewOwner) const;
+	FProperty* CreateProperty(const FProperty* FromProperty, UStruct* NewOwner);
+	FProperty* CreateProperty(const FProperty* FromProperty, UStruct* NewOwner, FName NewPropertyName);
 	void InitializeProperty(FProperty* NewProperty, FInitializePropertyArgs& Args);
-	void LinkProperty(FProperty* NewProperty);
+	void LinkProperty(FProperty* NewProperty) const;
+	void LinkProperty(FProperty* NewProperty, UStruct* NewOwner) const;
 	FName AddOnRepFunction(FName PropertyName);
-
-	TMap<const FProperty*, FProperty*> FromPropertyToCreatedProperty;
+	void SafeRename(UObject* Object);
+	void SetDefaultValue(const FProperty* SourceProperty, void const* SourceValuePtr, const FProperty* DestinationProperty);
 
 public:
 	/** The base object of the generated class. */
@@ -75,25 +75,36 @@ protected:
  *
  */
 UCLASS()
-class MODELVIEWVIEWMODELBLUEPRINT_API UMVVMBlueprintInstancedViewModel : public UMVVMBlueprintInstancedViewModelBase
+class MODELVIEWVIEWMODELBLUEPRINT_API UMVVMBlueprintInstancedViewModel_PropertyBag : public UMVVMBlueprintInstancedViewModelBase
 {
 	GENERATED_BODY()
 
 public:
-	virtual const UStruct* GetSourceStruct() const
+	const UStruct* GetSourceStruct() const
 	{
 		return Variables.GetValue().GetScriptStruct();
 }
 
-	virtual const uint8* GetSourceDefaults() const
+	const uint8* GetSourceDefaults() const
 	{
 		return Variables.GetValue().GetMemory();
 	}
 
-public:
+protected:
+	virtual bool IsClassDirty() const override;
+	virtual void CleanClass() override;
+	virtual void AddProperties() override;
+	virtual void SetDefaultValues() override;
+	virtual void ClassGenerated() override;
+
+private:
 	UPROPERTY(EditAnywhere, Category = "Viewmodel")
-	FInstancedPropertyBag Variables;	// todo, this should be a base and we have 2 implementation, one for the normal with a struct builder, and one for verse with the class ptr
+	FInstancedPropertyBag Variables;
 
+	UPROPERTY()
+	uint64 PropertiesHash = 0;
+
+	TMap<const FProperty*, FProperty*> FromPropertyToCreatedProperty;
 
 public:
 #if WITH_EDITOR
diff --git a/Engine/Plugins/Runtime/ModelViewViewModel/Source/ModelViewViewModelEditor/Private/MVVMEditorSubsystem.cpp b/Engine/Plugins/Runtime/ModelViewViewModel/Source/ModelViewViewModelEditor/Private/MVVMEditorSubsystem.cpp
--- a/Engine/Plugins/Runtime/ModelViewViewModel/Source/ModelViewViewModelEditor/Private/MVVMEditorSubsystem.cpp
+++ b/Engine/Plugins/Runtime/ModelViewViewModel/Source/ModelViewViewModelEditor/Private/MVVMEditorSubsystem.cpp
@@ -204,9 +204,9 @@ FGuid UMVVMEditorSubsystem::AddInstancedViewModel(UWidgetBlueprint* WidgetBluepr
 
 	FGuid Result;
 
-	FName UniqueName = MakeUniqueObjectName(View, UMVVMBlueprintInstancedViewModel::StaticClass(), "InstancedViewmodel");
-	UMVVMBlueprintInstancedViewModel* NewInstancedViewModel = NewObject<UMVVMBlueprintInstancedViewModel>(View, UniqueName);
-	NewInstancedViewModel->GenerateClass();
+	FName UniqueName = MakeUniqueObjectName(View, UMVVMBlueprintInstancedViewModel_PropertyBag::StaticClass(), "InstancedViewmodel");
+	UMVVMBlueprintInstancedViewModel_PropertyBag* NewInstancedViewModel = NewObject<UMVVMBlueprintInstancedViewModel_PropertyBag>(View, UniqueName);
+	NewInstancedViewModel->GenerateClass(true);
 	FMVVMBlueprintViewModelContext Context = FMVVMBlueprintViewModelContext(NewInstancedViewModel->GetGeneratedClass(), UniqueName);
 	if (Context.IsValid())
 	{
