MVVM: Expose the ContextResolver::Destroy to BP. Fix the display name of K2_CreateInstance.

diff --git a/Engine/Plugins/Runtime/ModelViewViewModel/Source/ModelViewViewModel/Public/View/MVVMViewModelContextResolver.h b/Engine/Plugins/Runtime/ModelViewViewModel/Source/ModelViewViewModel/Public/View/MVVMViewModelContextResolver.h
--- a/Engine/Plugins/Runtime/ModelViewViewModel/Source/ModelViewViewModel/Public/View/MVVMViewModelContextResolver.h
+++ b/Engine/Plugins/Runtime/ModelViewViewModel/Source/ModelViewViewModel/Public/View/MVVMViewModelContextResolver.h
@@ -25,8 +25,14 @@ public:
 		return K2_CreateInstance(ExpectedType, UserWidget).GetObject();
 	}
 
-	UFUNCTION(BlueprintImplementableEvent, Category="Viewmodel", DisplayName="Create Instance")
+	UFUNCTION(BlueprintImplementableEvent, Category = "Viewmodel", meta = (DisplayName = "Create Instance"))
 	TScriptInterface<INotifyFieldValueChanged> K2_CreateInstance(const UClass* ExpectedType, const UUserWidget* UserWidget) const;
 
-	virtual void DestroyInstance(const UObject* ViewModel, const UMVVMView* View) const {}
+	virtual void DestroyInstance(const UObject* ViewModel, const UMVVMView* View) const
+	{
+		K2_DestroyInstance(ViewModel, View);
+	}
+
+	UFUNCTION(BlueprintImplementableEvent, Category = "Viewmodel", meta = (DisplayName = "Destroy Instance"))
+	void K2_DestroyInstance(const UObject* ViewModel, const UMVVMView* View) const;
 };
