Fix issue 361586 with ModelViewViewModelBlueprint.cpp introduced with CL 27701383

diff --git a/Engine/Plugins/Runtime/ModelViewViewModel/Source/ModelViewViewModelBlueprint/Private/MVVMBlueprintInstancedViewModel.cpp b/Engine/Plugins/Runtime/ModelViewViewModel/Source/ModelViewViewModelBlueprint/Private/MVVMBlueprintInstancedViewModel.cpp
--- a/Engine/Plugins/Runtime/ModelViewViewModel/Source/ModelViewViewModelBlueprint/Private/MVVMBlueprintInstancedViewModel.cpp
+++ b/Engine/Plugins/Runtime/ModelViewViewModel/Source/ModelViewViewModelBlueprint/Private/MVVMBlueprintInstancedViewModel.cpp
@@ -310,7 +310,7 @@ uint64 GetObjectHash(const UObject* Object)
 
 uint64 CalcPropertyDescHash(const FPropertyBagPropertyDesc& Desc)
 {
-	UStruct* ValueTypeObject = Cast<UStruct>(Desc.ValueTypeObject);
+	const UStruct* ValueTypeObject = Cast<const UStruct>(Desc.ValueTypeObject);
 	FTopLevelAssetPath ValueTypeObjectPath = ValueTypeObject ? ValueTypeObject->GetStructPathName() : FTopLevelAssetPath();
 #if WITH_EDITORONLY_DATA
 	const uint32 Hashes[] = { GetTypeHash(ValueTypeObjectPath), GetTypeHash(Desc.ID), GetTypeHash(Desc.Name), GetTypeHash(Desc.ValueType), GetTypeHash(Desc.ContainerTypes), GetTypeHash(Desc.MetaData) };
