MVVM: Rework how we find the property for RepNotify function in InstancedViewModel

diff --git a/Engine/Plugins/Runtime/ModelViewViewModel/Source/ModelViewViewModel/Private/ViewModel/MVVMInstancedViewModelGeneratedClass.cpp b/Engine/Plugins/Runtime/ModelViewViewModel/Source/ModelViewViewModel/Private/ViewModel/MVVMInstancedViewModelGeneratedClass.cpp
--- a/Engine/Plugins/Runtime/ModelViewViewModel/Source/ModelViewViewModel/Private/ViewModel/MVVMInstancedViewModelGeneratedClass.cpp
+++ b/Engine/Plugins/Runtime/ModelViewViewModel/Source/ModelViewViewModel/Private/ViewModel/MVVMInstancedViewModelGeneratedClass.cpp
@@ -11,22 +11,49 @@ UClass* UMVVMInstancedViewModelGeneratedClass::GetAuthoritativeClass()
 {
 	return this;
 }
+
+void UMVVMInstancedViewModelGeneratedClass::PurgeClass(bool bRecompilingOnLoad)
+{
+	Super::PurgeClass(bRecompilingOnLoad);
+	PurgeNativeRepNotifyFunctions();
+
+}
+
+void UMVVMInstancedViewModelGeneratedClass::AddNativeRepNotifyFunction(UFunction* Function, const FProperty* Property)
+{
+	if (Property && Function && Property->RepNotifyFunc == Function->GetFName())
+	{
+		OnRepFunctionToLink.Add(Function);
+		OnRepToPropertyMap.Add(Function, Property);
+	}
+}
+
+void UMVVMInstancedViewModelGeneratedClass::PurgeNativeRepNotifyFunctions()
+{
+	OnRepFunctionToLink.Empty();
+	OnRepToPropertyMap.Empty();
+}
 #endif
 
 void UMVVMInstancedViewModelGeneratedClass::Link(FArchive& Ar, bool bRelinkExistingProperties)
 {
 	Super::Link(Ar, bRelinkExistingProperties);
 
+	for (UFunction* OnRep : OnRepFunctionToLink)
+	{
+		NativeFunctionLookupTable.Emplace(OnRep->GetFName(), &UMVVMInstancedViewModelGeneratedClass::K2_CallNativeOnRep);
+	}
+	OnRepToPropertyMap.Empty();
+
 	for (TFieldIterator<FProperty> PropertyIter(this, EFieldIteratorFlags::ExcludeSuper); PropertyIter; ++PropertyIter)
 	{
 		if (!PropertyIter->RepNotifyFunc.IsNone())
 		{
-			if (!NativeFunctionLookupTable.ContainsByPredicate([ToFind = PropertyIter->RepNotifyFunc](const FNativeFunctionLookup& Other) { return Other.Name == ToFind; }))
+			TObjectPtr<UFunction>* FoundFunction = OnRepFunctionToLink.FindByPredicate([FunctionName = PropertyIter->RepNotifyFunc](const UFunction* Other){ return Other->GetFName() == FunctionName; });
+			if (FoundFunction)
 			{
-				NativeFunctionLookupTable.Emplace(PropertyIter->RepNotifyFunc, &UMVVMInstancedViewModelGeneratedClass::K2_CallNativeOnRep);
+				OnRepToPropertyMap.Add(FoundFunction->Get(), *PropertyIter);
 			}
-			//NewFunction->RPCId = Params.RPCId;
-			//NewFunction->RPCResponseId = Params.RPCResponseId;
 		}
 	}
 }
@@ -34,6 +61,9 @@ void UMVVMInstancedViewModelGeneratedClass::Link(FArchive& Ar, bool bRelinkExist
 DEFINE_FUNCTION(UMVVMInstancedViewModelGeneratedClass::K2_CallNativeOnRep)
 {
 	UObject* CallingObject = P_THIS_OBJECT;
+
+	P_NATIVE_BEGIN;
+
 	if (CallingObject == nullptr || !CallingObject->GetClass()->ImplementsInterface(UNotifyFieldValueChanged::StaticClass()))
 	{
 		return;
@@ -46,24 +76,18 @@ DEFINE_FUNCTION(UMVVMInstancedViewModelGeneratedClass::K2_CallNativeOnRep)
 
 	FName PropertyName;
 	{
-		const UFunction* Func = Stack.CurrentNativeFunction;
-		FString FunctionName = Func->GetName();
-		if (FunctionName.RemoveFromStart(TEXT("__OnRep_")))
+		const UFunction* Function = Stack.CurrentNativeFunction;
+		const FProperty* const* FoundProperty = GeneratedClass->OnRepToPropertyMap.Find(TObjectKey<UFunction>(Function));
+		if (FoundProperty)
 		{
-			PropertyName = FName(*FunctionName, EFindName::FNAME_Find);
-			if (!PropertyName.IsNone())
-			{
-				FProperty* FoundProperty = CallingObject->GetClass()->FindPropertyByName(PropertyName);
-				if (FoundProperty)
-				{
-					GeneratedClass->OnPropertyReplicated(CallingObject, FoundProperty);
-				}
-			}
+			GeneratedClass->OnPropertyReplicated(CallingObject, *FoundProperty);
 		}
 	}
+
+	P_NATIVE_END;
 }
 
-void UMVVMInstancedViewModelGeneratedClass::OnPropertyReplicated(UObject* Object, const FProperty* Property)
+void UMVVMInstancedViewModelGeneratedClass::BroadcastFieldValueChanged(UObject* Object, const FProperty* Property)
 {
 	TScriptInterface<INotifyFieldValueChanged> CallingInterface = Object;
 	UE::FieldNotification::FFieldId FieldId = CallingInterface->GetFieldNotificationDescriptor().GetField(Object->GetClass(), Property->GetFName());
@@ -72,3 +96,8 @@ void UMVVMInstancedViewModelGeneratedClass::OnPropertyReplicated(UObject* Object
 		CallingInterface->BroadcastFieldValueChanged(FieldId);
 	}
 }
+
+void UMVVMInstancedViewModelGeneratedClass::OnPropertyReplicated(UObject* Object, const FProperty* Property)
+{
+	BroadcastFieldValueChanged(Object, Property);
+}
\ No newline at end of file
diff --git a/Engine/Plugins/Runtime/ModelViewViewModel/Source/ModelViewViewModel/Public/ViewModel/MVVMInstancedViewModelGeneratedClass.h b/Engine/Plugins/Runtime/ModelViewViewModel/Source/ModelViewViewModel/Public/ViewModel/MVVMInstancedViewModelGeneratedClass.h
--- a/Engine/Plugins/Runtime/ModelViewViewModel/Source/ModelViewViewModel/Public/ViewModel/MVVMInstancedViewModelGeneratedClass.h
+++ b/Engine/Plugins/Runtime/ModelViewViewModel/Source/ModelViewViewModel/Public/ViewModel/MVVMInstancedViewModelGeneratedClass.h
@@ -3,6 +3,7 @@
 #pragma once
 
 #include "Engine/BlueprintGeneratedClass.h"
+#include "UObject/ObjectKey.h"
 #include "UObject/Package.h"
 #include "MVVMInstancedViewModelGeneratedClass.generated.h"
 
@@ -19,10 +20,20 @@ public:
 	
 #if WITH_EDITOR
 	virtual UClass* GetAuthoritativeClass() override;
+	virtual void PurgeClass(bool bRecompilingOnLoad) override;
+	void AddNativeRepNotifyFunction(UFunction* Function, const FProperty* Property);
+	void PurgeNativeRepNotifyFunctions();
 #endif
 
 public:
 	DECLARE_FUNCTION(K2_CallNativeOnRep);
 
+	void BroadcastFieldValueChanged(UObject* Object, const FProperty* Property);
 	virtual void OnPropertyReplicated(UObject* Object, const FProperty* Property);
+
+private:
+	UPROPERTY()
+	TArray<TObjectPtr<UFunction>> OnRepFunctionToLink;
+
+	TMap<TObjectKey<UFunction>, const FProperty*> OnRepToPropertyMap;
 };
diff --git a/Engine/Plugins/Runtime/ModelViewViewModel/Source/ModelViewViewModelBlueprint/Private/MVVMBlueprintInstancedViewModel.cpp b/Engine/Plugins/Runtime/ModelViewViewModel/Source/ModelViewViewModelBlueprint/Private/MVVMBlueprintInstancedViewModel.cpp
--- a/Engine/Plugins/Runtime/ModelViewViewModel/Source/ModelViewViewModelBlueprint/Private/MVVMBlueprintInstancedViewModel.cpp
+++ b/Engine/Plugins/Runtime/ModelViewViewModel/Source/ModelViewViewModelBlueprint/Private/MVVMBlueprintInstancedViewModel.cpp
@@ -233,7 +233,7 @@ void UMVVMBlueprintInstancedViewModelBase::InitializeProperty(FProperty* NewProp
 {
 	if (Args.bNetwork)
 	{
-		NewProperty->RepNotifyFunc = AddOnRepFunction(Args.PropertyName);
+		AddOnRepFunction(NewProperty);
 	}
 
 	EPropertyFlags NewFlags = NewProperty->GetPropertyFlags() | EPropertyFlags::CPF_Edit;
@@ -244,7 +244,7 @@ void UMVVMBlueprintInstancedViewModelBase::InitializeProperty(FProperty* NewProp
 	if (Args.bFieldNotify)
 	{
 		NewProperty->SetMetaData(FBlueprintMetadata::MD_FieldNotify, TEXT(""));
-		GeneratedClass->FieldNotifies.Add(FFieldNotificationId(Args.PropertyName));
+		GeneratedClass->FieldNotifies.Add(FFieldNotificationId(NewProperty->GetFName()));
 	}
 }
 
@@ -258,25 +258,30 @@ void UMVVMBlueprintInstancedViewModelBase::LinkProperty(FProperty* NewProperty,
 	NewOwner->AddCppProperty(NewProperty);
 }
 
-FName UMVVMBlueprintInstancedViewModelBase::AddOnRepFunction(FName PropertyName)
+void UMVVMBlueprintInstancedViewModelBase::AddOnRepFunction(FProperty* NewProperty)
 {
-	FString OnRepCallFunctionName = FString::Printf(TEXT("__OnRep_%s"), *PropertyName.ToString());
+	FString OnRepCallFunctionName = FString::Printf(TEXT("__OnRep_%s"), *NewProperty->GetName());
 	FName Name_OnRepCallFunctionName = *OnRepCallFunctionName;
 	UObject* PreviousObj = StaticFindObjectFastInternal(nullptr, GeneratedClass, Name_OnRepCallFunctionName, true);
-	if (PreviousObj)
+	if (ensure(PreviousObj))
 	{
 		// The function or property already exist. Something is wrong.
-		return FName();
+		return;
 	}
+
 	UFunction* Func = NewObject<UFunction>(GeneratedClass, Name_OnRepCallFunctionName);
 	Func->FunctionFlags |= FUNC_Native | FUNC_Event | FUNC_BlueprintEvent | FUNC_BlueprintCallable;
 	GeneratedClass->AddNativeFunction(*OnRepCallFunctionName, &UMVVMInstancedViewModelGeneratedClass::K2_CallNativeOnRep);
 	GeneratedClass->AddFunctionToFunctionMap(Func, Func->GetFName());
 	Func->Bind();
 	Func->StaticLink(true);
+
 	Func->Next = GeneratedClass->Children;
 	GeneratedClass->Children = Func;
-	return Func->GetFName();
+
+	NewProperty->RepNotifyFunc = Func->GetFName();
+
+	GeneratedClass->AddNativeRepNotifyFunction(Func, NewProperty);
 }
 
 void UMVVMBlueprintInstancedViewModelBase::SafeRename(UObject* Object)
@@ -362,9 +367,7 @@ void UMVVMBlueprintInstancedViewModel_PropertyBag::AddProperties()
 		if (NewProperty)
 		{
 			FInitializePropertyArgs Args;
-			Args.PropertyName = NewProperty->GetFName();
 			Args.bFieldNotify = true;
-			//Args.bNetwork = true;
 			InitializeProperty(NewProperty, Args);
 
 			LinkProperty(NewProperty);
diff --git a/Engine/Plugins/Runtime/ModelViewViewModel/Source/ModelViewViewModelBlueprint/Public/MVVMBlueprintInstancedViewModel.h b/Engine/Plugins/Runtime/ModelViewViewModel/Source/ModelViewViewModelBlueprint/Public/MVVMBlueprintInstancedViewModel.h
--- a/Engine/Plugins/Runtime/ModelViewViewModel/Source/ModelViewViewModelBlueprint/Public/MVVMBlueprintInstancedViewModel.h
+++ b/Engine/Plugins/Runtime/ModelViewViewModel/Source/ModelViewViewModelBlueprint/Public/MVVMBlueprintInstancedViewModel.h
@@ -39,7 +39,6 @@ protected:
 protected:
 	struct FInitializePropertyArgs
 	{
-		FName PropertyName;
 		bool bFieldNotify = false;
 		bool bReadOnly = false;
 		bool bNetwork = false;
@@ -53,7 +52,7 @@ protected:
 	void InitializeProperty(FProperty* NewProperty, FInitializePropertyArgs& Args);
 	void LinkProperty(FProperty* NewProperty) const;
 	void LinkProperty(FProperty* NewProperty, UStruct* NewOwner) const;
-	FName AddOnRepFunction(FName PropertyName);
+	void AddOnRepFunction(FProperty* NewProperty);
 	void SafeRename(UObject* Object);
 	void SetDefaultValue(const FProperty* SourceProperty, void const* SourceValuePtr, const FProperty* DestinationProperty);
 
