Removing blueprint Callable and Blueprint Event flag from OnRepFunction. Clears Class children while handling compilation finish.

diff --git a/Engine/Plugins/Runtime/ModelViewViewModel/Source/ModelViewViewModelBlueprint/Private/MVVMBlueprintInstancedViewModel.cpp b/Engine/Plugins/Runtime/ModelViewViewModel/Source/ModelViewViewModelBlueprint/Private/MVVMBlueprintInstancedViewModel.cpp
--- a/Engine/Plugins/Runtime/ModelViewViewModel/Source/ModelViewViewModelBlueprint/Private/MVVMBlueprintInstancedViewModel.cpp
+++ b/Engine/Plugins/Runtime/ModelViewViewModel/Source/ModelViewViewModelBlueprint/Private/MVVMBlueprintInstancedViewModel.cpp
@@ -144,6 +144,7 @@ void UMVVMBlueprintInstancedViewModelBase::CleanClass()
 	GeneratedClass->ClassWithin = UObject::StaticClass();
 	GeneratedClass->ClassConfigName = ParentClass->ClassConfigName;
 	GeneratedClass->ClassFlags |= CLASS_NotPlaceable;
+	GeneratedClass->FieldNotifies.Empty();
 }
 
 void UMVVMBlueprintInstancedViewModelBase::AddProperties()
@@ -263,14 +264,15 @@ void UMVVMBlueprintInstancedViewModelBase::AddOnRepFunction(FProperty* NewProper
 	FString OnRepCallFunctionName = FString::Printf(TEXT("__OnRep_%s"), *NewProperty->GetName());
 	FName Name_OnRepCallFunctionName = *OnRepCallFunctionName;
 	UObject* PreviousObj = StaticFindObjectFastInternal(nullptr, GeneratedClass, Name_OnRepCallFunctionName, true);
-	if (ensure(PreviousObj))
+	ensure(PreviousObj == nullptr);
+	if (PreviousObj)
 	{
 		// The function or property already exist. Something is wrong.
 		return;
 	}
 
 	UFunction* Func = NewObject<UFunction>(GeneratedClass, Name_OnRepCallFunctionName);
-	Func->FunctionFlags |= FUNC_Native | FUNC_Event | FUNC_BlueprintEvent | FUNC_BlueprintCallable;
+	Func->FunctionFlags |= FUNC_Final | FUNC_Native | FUNC_Event;
 	GeneratedClass->AddNativeFunction(*OnRepCallFunctionName, &UMVVMInstancedViewModelGeneratedClass::K2_CallNativeOnRep);
 	GeneratedClass->AddFunctionToFunctionMap(Func, Func->GetFName());
 	Func->Bind();
diff --git a/Engine/Plugins/Runtime/ModelViewViewModel/Source/ModelViewViewModelBlueprint/Private/MVVMWidgetBlueprintExtension_View.cpp b/Engine/Plugins/Runtime/ModelViewViewModel/Source/ModelViewViewModelBlueprint/Private/MVVMWidgetBlueprintExtension_View.cpp
--- a/Engine/Plugins/Runtime/ModelViewViewModel/Source/ModelViewViewModelBlueprint/Private/MVVMWidgetBlueprintExtension_View.cpp
+++ b/Engine/Plugins/Runtime/ModelViewViewModel/Source/ModelViewViewModelBlueprint/Private/MVVMWidgetBlueprintExtension_View.cpp
@@ -174,6 +174,7 @@ void UMVVMWidgetBlueprintExtension_View::HandleFinishCompilingClass(UWidgetBluep
 					UFunction* Function = *FunctionIter;
 					Function->SetFlags(RF_Transient);
 				}
+				Class->Children = nullptr; // clear all functions (including those that the user generated).
 			}
 
 			// Does it have any bindings
