[MVVM] Fix inverted hierarchy check
* Specialized some error messages for more clarity

#tests Refactored some code so property is returning a derived class but binding is refering to data from base, abstract, class.
#rnx

diff --git a/Engine/Plugins/Runtime/ModelViewViewModel/Source/ModelViewViewModelBlueprint/Private/Bindings/MVVMCompiledBindingLibraryCompiler.cpp b/Engine/Plugins/Runtime/ModelViewViewModel/Source/ModelViewViewModelBlueprint/Private/Bindings/MVVMCompiledBindingLibraryCompiler.cpp
--- a/Engine/Plugins/Runtime/ModelViewViewModel/Source/ModelViewViewModelBlueprint/Private/Bindings/MVVMCompiledBindingLibraryCompiler.cpp
+++ b/Engine/Plugins/Runtime/ModelViewViewModel/Source/ModelViewViewModelBlueprint/Private/Bindings/MVVMCompiledBindingLibraryCompiler.cpp
@@ -415,7 +415,7 @@ TValueOrError<FCompiledBindingLibraryCompiler::FFieldPathHandle, FText> FCompile
 
 	if (FieldPath.Num() == 0)
 	{
-		return MakeError(FText::Format(LOCTEXT("FieldDoesNotReturnType", "The field does not return a '{0}'."), ExpectedType->GetDisplayNameText()));
+		return MakeError(FText::Format(LOCTEXT("FieldDoesNotReturnTypeEmptyPath", "The field does not return a '{0}', Field path was empty."), ExpectedType->GetDisplayNameText()));
 	}
 
 	UE::MVVM::FMVVMConstFieldVariant Last = FieldPath.Last();
@@ -432,9 +432,9 @@ TValueOrError<FCompiledBindingLibraryCompiler::FFieldPathHandle, FText> FCompile
 
 	if (ObjectPropertyBase == nullptr)
 	{
-		return MakeError(FText::Format(LOCTEXT("FieldDoesNotReturnType", "The field does not return a '{0}'."), ExpectedType->GetDisplayNameText()));
+		return MakeError(FText::Format(LOCTEXT("FieldDoesNotReturnTypeNoProperty", "The field does not return a '{0}', Object Property not found."), ExpectedType->GetDisplayNameText()));
 	}
-	if (ObjectPropertyBase->PropertyClass == nullptr || !ExpectedType->IsChildOf(ObjectPropertyBase->PropertyClass))
+	if (ObjectPropertyBase->PropertyClass == nullptr || !ObjectPropertyBase->PropertyClass->IsChildOf(ExpectedType))
 	{
 		return MakeError(FText::Format(LOCTEXT("FieldDoesNotReturnType", "The field does not return a '{0}'."), ExpectedType->GetDisplayNameText()));
 	}
