MVVM: Sort the view bindings and the view sources to make sure we have a reliable order of execution.
Initialize each bindings source one at the time. That allow a binding to set a variable that will then be read by another binding. That other binding can also change another source (by calling SetViewmodel internally). The bindings of the new source need to execute only if that source was already initialized. Delay bindings executed that way if they are complex.

diff --git a/Engine/Plugins/Runtime/ModelViewViewModel/Source/ModelViewViewModel/Private/View/MVVMView.cpp b/Engine/Plugins/Runtime/ModelViewViewModel/Source/ModelViewViewModel/Private/View/MVVMView.cpp
--- a/Engine/Plugins/Runtime/ModelViewViewModel/Source/ModelViewViewModel/Private/View/MVVMView.cpp
+++ b/Engine/Plugins/Runtime/ModelViewViewModel/Source/ModelViewViewModel/Private/View/MVVMView.cpp
@@ -167,15 +167,37 @@ void UMVVMView::InitializeBindings()
 	ensure(RegisteredLibraryBindings.IsEmpty());
 	RegisteredLibraryBindings.Reset(CompiledBindings.Num());
 	RegisteredLibraryBindings.AddDefaulted(CompiledBindings.Num());
-	for (int32 Index = 0; Index < CompiledBindings.Num(); ++Index)
+
+	if (CompiledBindings.Num())
 	{
-		const FMVVMViewClass_CompiledBinding& Binding = CompiledBindings[Index];
-		if (Binding.IsEnabledByDefault())
+		FName PreviousSourceName = CompiledBindings[0].GetSourceName();
+		for (int32 Index = 0; Index < CompiledBindings.Num(); ++Index)
 		{
-			EnableLibraryBinding(Binding, Index);
+			const FMVVMViewClass_CompiledBinding& Binding = CompiledBindings[Index];
+
+			// A binding execution might call SetViewModel.
+			//The bBindingsInitialized for that source needs to be true to re execute all the bindings.
+			if (PreviousSourceName != Binding.GetSourceName())
+			{
+				if (FMVVMViewSource* FoundSource = FindViewSource(PreviousSourceName))
+				{
+					FoundSource->bBindingsInitialized = true;
+				}
+				PreviousSourceName = Binding.GetSourceName();
+			}
+
+			if (Binding.IsEnabledByDefault())
+			{
+				EnableLibraryBinding(Binding, Index, false, false);
+			}
 		}
 	}
 
+	for (FMVVMViewSource& Source : Sources)
+	{
+		Source.bBindingsInitialized = true;
+	}
+
 	bBindingsInitialized = true;
 
 	if (bHasEveryTickBinding)
@@ -215,6 +237,7 @@ void UMVVMView::UninitializeInternal(bool bUninitializeSources)
 	{
 		const FMVVMViewClass_SourceCreator& Item = AllViewModelCreators[Index];
 		FMVVMViewSource& ViewSource = Sources[Index];
+		ViewSource.bBindingsInitialized = false;
 
 		if (ViewSource.RegisteredCount > 0 && ViewSource.Source && bUninitializeBindings)
 		{
@@ -403,7 +426,7 @@ bool UMVVMView::SetSourceInternal(FName ViewModelName, TScriptInterface<INotifyF
 		const TArrayView<const FMVVMViewClass_CompiledEvent> CompiledEvents = ClassExtension->GetCompiledEvents();
 
 		// Unregister any bindings from that source
-		if (bBindingsInitialized)
+		if (Sources[AllCreatedSourcesIndex].bBindingsInitialized)
 		{
 			for (int32 Index = 0; Index < CompiledBindings.Num(); ++Index)
 			{
@@ -415,7 +438,7 @@ bool UMVVMView::SetSourceInternal(FName ViewModelName, TScriptInterface<INotifyF
 			}
 		}
 
-		// Unregisterer any events from that source
+		// Unregister any events from that source
 		if (bConstructed)
 		{
 			for (int32 Index = 0; Index < CompiledEvents.Num(); ++Index)
@@ -449,7 +472,7 @@ bool UMVVMView::SetSourceInternal(FName ViewModelName, TScriptInterface<INotifyF
 		bHasEveryTickBinding = false;
 
 		// Register back any bindings that was previously enabled
-		if (bBindingsInitialized && NewValue.GetObject())
+		if (Sources[AllCreatedSourcesIndex].bBindingsInitialized && NewValue.GetObject())
 		{
 			// Enabled the default bindings
 			for (int32 Index = 0; Index < CompiledBindings.Num(); ++Index)
@@ -460,7 +483,12 @@ bool UMVVMView::SetSourceInternal(FName ViewModelName, TScriptInterface<INotifyF
 					// Binding on this viewmodel
 					if (Binding.GetSourceName() == ViewModelName)
 					{
-						EnableLibraryBinding(Binding, Index);
+						// The binding needs execution, but maybe it's a complex binding and not all the sources from that binding are initialized.
+						//This occurs only at initialization when all the bindings from source A exectuted and a binding from source B call SetViewModel on A.
+						// A should reexecute its bindings but only if all the sources from the complex conversion initialized. We do not know that. Wait for the next frame and executed it. 
+						const bool bForceDelayed = Binding.IsConversionFunctionComplex() && !bBindingsInitialized;
+						EnableLibraryBinding(Binding, Index, true, bForceDelayed);
+
 						// Bindings that depends on this binding
 						if (Binding.IsEvaluateSourceCreatorBinding())
 						{
@@ -741,7 +769,7 @@ bool UMVVMView::IsLibraryBindingEnabled(int32 InBindindIndex) const
 }
 
 
-void UMVVMView::EnableLibraryBinding(const FMVVMViewClass_CompiledBinding& Binding, int32 BindingIndex)
+void UMVVMView::EnableLibraryBinding(const FMVVMViewClass_CompiledBinding& Binding, int32 BindingIndex, bool bSourceChanges, bool bForceDelayed)
 {
 	check(!RegisteredLibraryBindings[BindingIndex].IsValid());
 
@@ -767,9 +795,20 @@ void UMVVMView::EnableLibraryBinding(const FMVVMViewClass_CompiledBinding& Bindi
 		));
 	}
 
-	if (bCanExecute && Binding.NeedsExecutionAtInitialization())
+	if (bCanExecute)
 	{
-		ExecuteLibraryBinding(Binding, BindingIndex);
+		// @TODO: Asil.Karatas - Check if bSourceChanges should be guarded with Binding.NeedsExecuteAtInitializationWhenSourceChanges()
+		if (Binding.NeedsExecutionAtInitialization() || bSourceChanges) 
+		{
+			if (!bForceDelayed)
+			{
+				ExecuteLibraryBinding(Binding, BindingIndex);
+			}
+			else
+			{
+				GEngine->GetEngineSubsystem<UMVVMBindingSubsystem>()->AddDelayedBinding(this, BindingIndex);
+			}
+		}
 	}
 
 	bHasEveryTickBinding = bHasEveryTickBinding || ExecutionMode == EMVVMExecutionMode::Tick;
diff --git a/Engine/Plugins/Runtime/ModelViewViewModel/Source/ModelViewViewModel/Public/View/MVVMView.h b/Engine/Plugins/Runtime/ModelViewViewModel/Source/ModelViewViewModel/Public/View/MVVMView.h
--- a/Engine/Plugins/Runtime/ModelViewViewModel/Source/ModelViewViewModel/Public/View/MVVMView.h
+++ b/Engine/Plugins/Runtime/ModelViewViewModel/Source/ModelViewViewModel/Public/View/MVVMView.h
@@ -42,6 +42,10 @@ struct FMVVMViewSource
 	UPROPERTY(VisibleAnywhere, Category = "Viewmodel")
 	bool bCreatedSource = false;
 
+	// The source bindings are initialized
+	UPROPERTY(VisibleAnywhere, Category = "Viewmodel")
+	bool bBindingsInitialized = false;
+
 	// The source was set manually via SetViewModel.
 	UPROPERTY(VisibleAnywhere, Category = "Viewmodel")
 	bool bSetManually = false;
@@ -148,7 +152,7 @@ private:
 
 	void ExecuteLibraryBinding(const FMVVMViewClass_CompiledBinding& Binding, int32 BindingIndex) const;
 
-	void EnableLibraryBinding(const FMVVMViewClass_CompiledBinding& Item, int32 BindingIndex);
+	void EnableLibraryBinding(const FMVVMViewClass_CompiledBinding& Item, int32 BindingIndex, bool bSourceChanges, bool bForceDelayed);
 	void DisableLibraryBinding(const FMVVMViewClass_CompiledBinding& Item, int32 BindingIndex);
 	bool IsLibraryBindingEnabled(int32 InBindindIndex) const;
 
diff --git a/Engine/Plugins/Runtime/ModelViewViewModel/Source/ModelViewViewModelBlueprint/Private/MVVMViewBlueprintCompiler.cpp b/Engine/Plugins/Runtime/ModelViewViewModel/Source/ModelViewViewModelBlueprint/Private/MVVMViewBlueprintCompiler.cpp
--- a/Engine/Plugins/Runtime/ModelViewViewModel/Source/ModelViewViewModelBlueprint/Private/MVVMViewBlueprintCompiler.cpp
+++ b/Engine/Plugins/Runtime/ModelViewViewModel/Source/ModelViewViewModelBlueprint/Private/MVVMViewBlueprintCompiler.cpp
@@ -358,7 +358,7 @@ void FMVVMViewBlueprintCompiler::CreateVariables(const FWidgetBlueprintCompilerC
 			}
 
 			const FProperty* Property = SourceContext.Field.IsProperty() ? SourceContext.Field.GetProperty() : BindingHelper::GetReturnProperty(SourceContext.Field.GetFunction());
-			const FObjectPropertyBase* ObjectProperty = CastField<const FObjectPropertyBase>(Property);
+			const FObjectProperty* ObjectProperty = CastField<const FObjectProperty>(Property);
 			const bool bIsCompatible = ObjectProperty && SourceContext.Class->IsChildOf(ObjectProperty->PropertyClass);
 			if (!bIsCompatible)
 			{
@@ -1335,6 +1335,7 @@ bool FMVVMViewBlueprintCompiler::CompileSourceCreators(const FCompiledBindingLib
 		return false;
 	}
 
+	TArray<FMVVMViewClass_SourceCreator> UnsortedSourceCreators;
 	for (const FCompilerSourceCreatorContext& SourceCreatorContext : CompilerSourceCreatorContexts)
 	{
 		const FMVVMBlueprintViewModelContext& ViewModelContext = SourceCreatorContext.ViewModelContext;
@@ -1441,7 +1442,70 @@ bool FMVVMViewBlueprintCompiler::CompileSourceCreators(const FCompiledBindingLib
 		CompiledSourceCreator.Flags |= bCanBeSet ? (uint8)FMVVMViewClass_SourceCreator::ESourceFlags::CanBeSet : 0;
 		CompiledSourceCreator.Flags |= CanBeEvaluated ? (uint8)FMVVMViewClass_SourceCreator::ESourceFlags::CanBeEvaluated : 0;
 
-		ViewExtension->SourceCreators.Add(MoveTemp(CompiledSourceCreator));
+		UnsortedSourceCreators.Add(MoveTemp(CompiledSourceCreator));
+	}
+
+	// sort the source creators by priority and then by name
+	if (UnsortedSourceCreators.Num() > 1)
+	{
+		struct FSortData
+		{
+			FSortData() = default;
+			FName SourceName;
+			FName ParentSourceName;
+			int32 SortIndex = -1;
+
+			void CalculateSortIndex(TMap<FName, FSortData>& Map)
+			{
+				if (SortIndex < 0)
+				{
+					if (ParentSourceName.IsNone())
+					{
+						SortIndex = 0;
+					}
+					else
+					{
+						Map[ParentSourceName].CalculateSortIndex(Map);
+						// calculate the Depth recursively
+						SortIndex = Map[ParentSourceName].SortIndex + 1;
+					}
+				}
+			}
+		};
+		TMap<FName, FSortData> SortDatas;
+		SortDatas.Reserve(UnsortedSourceCreators.Num());
+		for (int32 Index = 0; Index < UnsortedSourceCreators.Num(); ++Index)
+		{
+			FSortData SortData;
+			SortData.SourceName = UnsortedSourceCreators[Index].GetSourceName();
+			SortData.ParentSourceName =UnsortedSourceCreators[Index].GetParentSourceName();
+			SortDatas.Emplace(SortData.SourceName, SortData);
+		}
+		for (auto& SortDataPair : SortDatas)
+		{
+			SortDataPair.Value.CalculateSortIndex(SortDatas);
+		}
+
+		UnsortedSourceCreators.Sort([&SortDatas](const FMVVMViewClass_SourceCreator& A, const FMVVMViewClass_SourceCreator& B)
+			{
+				int32 ASortIndex = SortDatas[A.GetSourceName()].SortIndex;
+				int32 BSortIndex = SortDatas[B.GetSourceName()].SortIndex;
+				if (ASortIndex == BSortIndex)
+				{
+					return A.GetSourceName().LexicalLess(B.GetSourceName());
+				}
+				return ASortIndex < BSortIndex;
+			});
+	}
+
+	// Add sorted array to the ViewExtension
+	{
+		ViewExtension->SourceCreators.Reset(UnsortedSourceCreators.Num());
+		for (FMVVMViewClass_SourceCreator& SourceCreator : UnsortedSourceCreators)
+		{
+			ViewExtension->SourceCreators.Add(MoveTemp(SourceCreator));
+		}
+		UnsortedSourceCreators.Reset();
 	}
 
 	return bAreSourcesCreatorValid;
@@ -1747,6 +1811,11 @@ bool FMVVMViewBlueprintCompiler::CompileBindings(const FCompiledBindingLibraryCo
 		return false;
 	}
 
+	// Store bindings with corresponding ComplexConversionFunctionContextIndex to be sorted and mark the last binding 
+	// in the complex conversion as bExecuteAtInitialization == true
+	TArray<TPair<int32, FMVVMViewClass_CompiledBinding>> TempBindingPairs;
+	TempBindingPairs.Reserve(CompilerBindings.Num());
+
 	for (const FCompilerBinding& CompileBinding : CompilerBindings)
 	{
 		// PropertyBinding needs a valid CompileBinding.BindingIndex.
@@ -1855,6 +1924,9 @@ bool FMVVMViewBlueprintCompiler::CompileBindings(const FCompiledBindingLibraryCo
 		NewBinding.FieldId = CompiledFieldId ? *CompiledFieldId : FMVVMVCompiledFieldId();
 		NewBinding.Binding = CompiledBinding ? *CompiledBinding : FMVVMVCompiledBinding();
 		NewBinding.Flags = 0;
+
+		int32 ComplexConversionFunctionContextIndex = INDEX_NONE;
+
 		if (ViewBinding)
 		{
 			if (CompileBinding.Type != ECompilerBindingType::PropertyBinding)
@@ -1864,11 +1936,14 @@ bool FMVVMViewBlueprintCompiler::CompileBindings(const FCompiledBindingLibraryCo
 				continue;
 			}
 
-			bool bExecuteAtInitialization = CompileBinding.bIsForwardBinding;
-			if (bExecuteAtInitialization && CompileBinding.ComplexConversionFunctionContextIndex != INDEX_NONE)
+			// If it is forward binding and simple conversion, mark it as bExecuteAtInitialization == true
+			bool bExecuteAtInitialization = CompileBinding.bIsForwardBinding && CompileBinding.ComplexConversionFunctionContextIndex == INDEX_NONE;
+
+			// If it is forward binding and complex conversion, store ComplexConversionFunctionContextIndex to be handled later in this function.
+			if (CompileBinding.bIsForwardBinding && CompileBinding.ComplexConversionFunctionContextIndex != INDEX_NONE)
 			{
-				bExecuteAtInitialization = !ComplexConversionFunctionContexts[CompileBinding.ComplexConversionFunctionContextIndex].bExecAtInitGenerated;
-				ComplexConversionFunctionContexts[CompileBinding.ComplexConversionFunctionContextIndex].bExecAtInitGenerated = true;
+				// Only forward binding with complex conversion will have a valid ComplexConversionFunctionContextIndex
+				ComplexConversionFunctionContextIndex = CompileBinding.ComplexConversionFunctionContextIndex;
 			}
 
 			NewBinding.ExecutionMode = ViewBinding->bOverrideExecutionMode ? ViewBinding->OverrideExecutionMode : (EMVVMExecutionMode)CVarDefaultExecutionMode->GetInt();;
@@ -1898,13 +1973,71 @@ bool FMVVMViewBlueprintCompiler::CompileBindings(const FCompiledBindingLibraryCo
 			NewBinding.Flags |= (bIsSourceSelf) ? FMVVMViewClass_CompiledBinding::EBindingFlags::SourceObjectIsSelf : 0;
 		}
 
-		ViewExtension->CompiledBindings.Emplace(MoveTemp(NewBinding));
+		TempBindingPairs.Emplace(ComplexConversionFunctionContextIndex, MoveTemp(NewBinding));
+	}
+
+	if (ViewExtension->SourceCreators.Num() > 1)
+	{
+		TMap<FName, int32> SourceCreatorSortOrder;
+		SourceCreatorSortOrder.Reserve(ViewExtension->SourceCreators.Num());
+		for (int32 Index = 0; Index < ViewExtension->SourceCreators.Num(); ++Index)
+		{
+			SourceCreatorSortOrder.Add(ViewExtension->SourceCreators[Index].GetSourceName(), Index);
+		}
+
+		TempBindingPairs.StableSort([&SourceCreatorSortOrder](const TPair<int32, FMVVMViewClass_CompiledBinding>& PairA, const TPair<int32, FMVVMViewClass_CompiledBinding>& PairB)
+			{
+				const FMVVMViewClass_CompiledBinding& A = PairA.Value;
+				const FMVVMViewClass_CompiledBinding& B = PairB.Value;
+
+				int32* FoundA = SourceCreatorSortOrder.Find(A.GetSourceName());
+				int32* FoundB = SourceCreatorSortOrder.Find(B.GetSourceName());
+				int32 AIndex = FoundA ? *FoundA : -1;
+				int32 BIndex = FoundB ? *FoundB : -1;
+
+				if (AIndex == BIndex)
+				{
+					if (A.GetEvaluateSourceCreatorBindingIndex() == B.GetEvaluateSourceCreatorBindingIndex())
+					{
+						return A.GetSourceName().LexicalLess(B.GetSourceName());
+					}
+
+					return A.GetEvaluateSourceCreatorBindingIndex() < B.GetEvaluateSourceCreatorBindingIndex();
+				}
+
+				return AIndex < BIndex;
+			});
+	}
+
+	// Go backward to be able to mark only the last binding in the complex conversion.
+	for (int32 Index = TempBindingPairs.Num() - 1; Index >= 0; --Index)
+	{
+		FMVVMViewClass_CompiledBinding& Binding = TempBindingPairs[Index].Value;
+		const int32 ComplexConversionFunctionContextIndex = TempBindingPairs[Index].Key;
+
+		// Only handle bindings with complex conversions.
+		if (ComplexConversionFunctionContextIndex != INDEX_NONE)
+		{
+			// Mark the complex conversion as bExecAtInitGenerated.
+			if (!ComplexConversionFunctionContexts[ComplexConversionFunctionContextIndex].bExecAtInitGenerated)
+			{
+				ComplexConversionFunctionContexts[ComplexConversionFunctionContextIndex].bExecAtInitGenerated = true;
+
+				// Mark the last binding with the complex conversion.
+				Binding.Flags |= FMVVMViewClass_CompiledBinding::EBindingFlags::ExecuteAtInitialization;
+			}
+		}
+	}
+
+	ViewExtension->CompiledBindings.Reserve(TempBindingPairs.Num());
+	for (TPair<int32, FMVVMViewClass_CompiledBinding>& BindingPair : TempBindingPairs)
+	{
+		ViewExtension->CompiledBindings.Emplace(MoveTemp(BindingPair.Value));
 	}
 
 	return bAreBindingsValid;
 }
 
-
 bool FMVVMViewBlueprintCompiler::PreCompileEvents(UWidgetBlueprintGeneratedClass* Class, UMVVMBlueprintView* BlueprintView)
 {
 	if (!bAreSourceContextsValid || !bAreEventsValid)
