Fix MVVM failing validation due to checking if object nodes were allowed within themselves
Add PropertyHasFiltering method to complete MVVM allowlist API.
Make LocalizableMessage API exposed so that it can be allowlisted.

diff --git a/Engine/Plugins/Experimental/LocalizableMessage/Source/LocalizableMessage/Public/LocalizableMessage.h b/Engine/Plugins/Experimental/LocalizableMessage/Source/LocalizableMessage/Public/LocalizableMessage.h
--- a/Engine/Plugins/Experimental/LocalizableMessage/Source/LocalizableMessage/Public/LocalizableMessage.h
+++ b/Engine/Plugins/Experimental/LocalizableMessage/Source/LocalizableMessage/Public/LocalizableMessage.h
@@ -9,15 +9,15 @@
 #include "LocalizableMessage.generated.h"
 
 USTRUCT()
-struct FLocalizableMessageParameterEntry
+struct LOCALIZABLEMESSAGE_API FLocalizableMessageParameterEntry
 {
 	GENERATED_BODY()
 
 public:
 
 	FLocalizableMessageParameterEntry();
-	LOCALIZABLEMESSAGE_API FLocalizableMessageParameterEntry(const FString& InKey, const FInstancedStruct& InValue);
-	LOCALIZABLEMESSAGE_API ~FLocalizableMessageParameterEntry();
+	FLocalizableMessageParameterEntry(const FString& InKey, const FInstancedStruct& InValue);
+	~FLocalizableMessageParameterEntry();
 
 	bool operator==(const FLocalizableMessageParameterEntry& Other) const;
 	bool operator!=(const FLocalizableMessageParameterEntry& Other) const
@@ -33,13 +33,13 @@ public:
 };
 
 USTRUCT(BlueprintType)
-struct FLocalizableMessage
+struct LOCALIZABLEMESSAGE_API FLocalizableMessage
 {
 	GENERATED_BODY();
 
-	LOCALIZABLEMESSAGE_API FLocalizableMessage();
-	LOCALIZABLEMESSAGE_API ~FLocalizableMessage();
-	LOCALIZABLEMESSAGE_API bool operator==(const FLocalizableMessage& Other) const;
+	FLocalizableMessage();
+	~FLocalizableMessage();
+	bool operator==(const FLocalizableMessage& Other) const;
 
 	void Reset()
 	{
diff --git a/Engine/Plugins/Runtime/ModelViewViewModel/Source/ModelViewViewModelBlueprint/Private/MVVMDeveloperProjectSettings.cpp b/Engine/Plugins/Runtime/ModelViewViewModel/Source/ModelViewViewModelBlueprint/Private/MVVMDeveloperProjectSettings.cpp
--- a/Engine/Plugins/Runtime/ModelViewViewModel/Source/ModelViewViewModelBlueprint/Private/MVVMDeveloperProjectSettings.cpp
+++ b/Engine/Plugins/Runtime/ModelViewViewModel/Source/ModelViewViewModelBlueprint/Private/MVVMDeveloperProjectSettings.cpp
@@ -7,6 +7,7 @@
 #include "MVVMBlueprintViewModelContext.h"
 #include "PropertyPermissionList.h"
 #include "Types/MVVMExecutionMode.h"
+#include "UObject/UnrealType.h"
 
 #define LOCTEXT_NAMESPACE "MVVMDeveloperProjectSettings"
 
@@ -35,6 +36,36 @@ FText UMVVMDeveloperProjectSettings::GetSectionText() const
 	return LOCTEXT("MVVMProjectSettings", "Model View Viewmodel");
 }
 
+bool UMVVMDeveloperProjectSettings::PropertyHasFiltering(const FProperty* Property) const
+{
+	check(Property);
+	const UStruct* ObjectStruct = nullptr;
+
+	if (const FObjectProperty* ObjectProperty = CastField<FObjectProperty>(Property))
+	{
+		ObjectStruct = ObjectProperty->PropertyClass;
+	}
+	else
+	{
+		ObjectStruct = Property->GetOwnerStruct();
+	}
+
+	if (!FPropertyEditorPermissionList::Get().HasFiltering(ObjectStruct))
+	{
+		return false;
+	}
+
+	TStringBuilder<512> StringBuilder;
+	Property->GetOwnerClass()->GetPathName(nullptr, StringBuilder);
+	FSoftClassPath StructPath;
+	StructPath.SetPath(StringBuilder);
+	if (const FMVVMDeveloperProjectWidgetSettings* Settings = FieldSelectorPermissions.Find(StructPath))
+	{
+		return !Settings->DisallowedFieldNames.Find(Property->GetFName());
+	}
+	return true;
+}
+
 bool UMVVMDeveloperProjectSettings::IsPropertyAllowed(const FProperty* Property) const
 {
 	check(Property);
diff --git a/Engine/Plugins/Runtime/ModelViewViewModel/Source/ModelViewViewModelBlueprint/Private/MVVMViewBlueprintCompiler.cpp b/Engine/Plugins/Runtime/ModelViewViewModel/Source/ModelViewViewModelBlueprint/Private/MVVMViewBlueprintCompiler.cpp
--- a/Engine/Plugins/Runtime/ModelViewViewModel/Source/ModelViewViewModelBlueprint/Private/MVVMViewBlueprintCompiler.cpp
+++ b/Engine/Plugins/Runtime/ModelViewViewModel/Source/ModelViewViewModelBlueprint/Private/MVVMViewBlueprintCompiler.cpp
@@ -17,6 +17,7 @@
 #include "Templates/ValueOrError.h"
 #include "Types/MVVMBindingName.h"
 #include "UObject/LinkerLoad.h"
+#include "UObject/UnrealType.h"
 #include "View/MVVMViewClass.h"
 #include "View/MVVMViewModelContextResolver.h"
 
@@ -2570,8 +2571,11 @@ TArray<FMVVMConstFieldVariant> FMVVMViewBlueprintCompiler::CreatePropertyPath(co
 
 bool FMVVMViewBlueprintCompiler::IsPropertyPathValid(TArrayView<const FMVVMConstFieldVariant> PropertyPath)
 {
-	for (const FMVVMConstFieldVariant& Field : PropertyPath)
+	int32 PathLength = PropertyPath.Num();
+	for (int32 Index = 0; Index < PathLength; Index++)
 	{
+		const FMVVMConstFieldVariant& Field = PropertyPath[Index];
+
 		if (Field.IsEmpty())
 		{
 			return false;
@@ -2582,7 +2586,16 @@ bool FMVVMViewBlueprintCompiler::IsPropertyPathValid(TArrayView<const FMVVMConst
 			{
 				return false;
 			}
-			if (!GetDefault<UMVVMDeveloperProjectSettings>()->IsPropertyAllowed(Field.GetProperty()))
+
+			// Don't filter objects in path if properties remain in path to be filtered (by length). 
+			// Child properties will acquire their parent struct / these objects for filtering.
+			// @TODO: Consider same for 'FCompiledBindingLibraryCompiler::AddFieldPathImpl' & 'UE::MVVM::Private::PassFilter'.
+			const FObjectProperty* ObjectProperty = CastField<FObjectProperty>(Field.GetProperty());
+			if (ObjectProperty && Index != PathLength - 1)
+			{
+				// Do nothing
+			}
+			else if (!GetDefault<UMVVMDeveloperProjectSettings>()->IsPropertyAllowed(Field.GetProperty()))
 			{
 				return false;
 			}
diff --git a/Engine/Plugins/Runtime/ModelViewViewModel/Source/ModelViewViewModelBlueprint/Public/MVVMDeveloperProjectSettings.h b/Engine/Plugins/Runtime/ModelViewViewModel/Source/ModelViewViewModelBlueprint/Public/MVVMDeveloperProjectSettings.h
--- a/Engine/Plugins/Runtime/ModelViewViewModel/Source/ModelViewViewModelBlueprint/Public/MVVMDeveloperProjectSettings.h
+++ b/Engine/Plugins/Runtime/ModelViewViewModel/Source/ModelViewViewModelBlueprint/Public/MVVMDeveloperProjectSettings.h
@@ -69,6 +69,8 @@ public:
 	virtual FName GetCategoryName() const override;
 	virtual FText GetSectionText() const override;
 
+
+	bool PropertyHasFiltering(const FProperty* Property) const;
 	bool IsPropertyAllowed(const FProperty* Property) const;
 	bool IsFunctionAllowed(const UFunction* Function) const;
 	bool IsConversionFunctionAllowed(const UFunction* Function) const;
