MVVM: Make sure the different object ptr type matches before copying them. Evaluated path use the correct container address when evaluating a function.

diff --git a/Engine/Plugins/Runtime/ModelViewViewModel/Source/ModelViewViewModel/Private/Bindings/MVVMBindingHelper.cpp b/Engine/Plugins/Runtime/ModelViewViewModel/Source/ModelViewViewModel/Private/Bindings/MVVMBindingHelper.cpp
--- a/Engine/Plugins/Runtime/ModelViewViewModel/Source/ModelViewViewModel/Private/Bindings/MVVMBindingHelper.cpp
+++ b/Engine/Plugins/Runtime/ModelViewViewModel/Source/ModelViewViewModel/Private/Bindings/MVVMBindingHelper.cpp
@@ -461,7 +461,8 @@ namespace UE::MVVM::BindingHelper
 				&& DestinationObjectProperty
 				&& SourceObjectProperty->PropertyClass
 				&& DestinationObjectProperty->PropertyClass
-				&& SourceObjectProperty->PropertyClass->IsChildOf(DestinationObjectProperty->PropertyClass);
+				&& SourceObjectProperty->PropertyClass->IsChildOf(DestinationObjectProperty->PropertyClass)
+				&& SourceObjectProperty->GetClass() == DestinationObjectProperty->GetClass(); // FSoftObjectPtr != TObjectPtr != FWeakObjectPtr != FLazyObjectPtr
 		}
 
 		bool IsObjectArrayPropertyCompatible(const FProperty* Source, const FProperty* Destination)
diff --git a/Engine/Plugins/Runtime/ModelViewViewModel/Source/ModelViewViewModel/Private/Bindings/MVVMCompiledBindingLibrary.cpp b/Engine/Plugins/Runtime/ModelViewViewModel/Source/ModelViewViewModel/Private/Bindings/MVVMCompiledBindingLibrary.cpp
--- a/Engine/Plugins/Runtime/ModelViewViewModel/Source/ModelViewViewModel/Private/Bindings/MVVMCompiledBindingLibrary.cpp
+++ b/Engine/Plugins/Runtime/ModelViewViewModel/Source/ModelViewViewModel/Private/Bindings/MVVMCompiledBindingLibrary.cpp
@@ -343,7 +343,7 @@ namespace UE::MVVM::Private
 		void* DataPtr = FMemory_Alloca_Aligned(Function->ParmsSize, Function->GetMinAlignment());
 		ReturnObjectProperty->InitializeValue(DataPtr);
 		Container->ProcessEvent(Function, DataPtr);
-		UObject* NewContainer = ReturnObjectProperty->GetObjectPropertyValue(DataPtr);
+		UObject* NewContainer = ReturnObjectProperty->GetObjectPropertyValue_InContainer(DataPtr);
 		ReturnObjectProperty->DestroyValue(DataPtr);
 		return NewContainer;
 	}
@@ -395,6 +395,7 @@ TValueOrError<UE::MVVM::FFieldContext, void> FMVVMCompiledBindingLibrary::Evalua
 			if (PathIndex.bIsObjectProperty)
 			{
 				const FObjectPropertyBase* ObjectProperty = CastField<const FObjectPropertyBase>(Property);
+				check(ObjectProperty);
 				UObject* NewContainer = ObjectProperty->GetObjectPropertyValue_InContainer(CurrentContainer.GetData()); // this skip any Getter
 				CurrentContainer.SetUObject(NewContainer);
 			}
