MVVM: Only create properties when needed. Viewmodel properties should always be created (no inheritance). Widget should be created if they are not already. Support BindWidget.
Rename FCreateVariableContext::GetSkeletonGeneratedClass to GetGeneratedClass because it returns the current class it generates and not the skeletal. It will be the skeletal if it's generating the skeletal.

diff --git a/Engine/Plugins/Runtime/ModelViewViewModel/Source/ModelViewViewModelBlueprint/Private/MVVMViewBlueprintCompiler.cpp b/Engine/Plugins/Runtime/ModelViewViewModel/Source/ModelViewViewModelBlueprint/Private/MVVMViewBlueprintCompiler.cpp
--- a/Engine/Plugins/Runtime/ModelViewViewModel/Source/ModelViewViewModelBlueprint/Private/MVVMViewBlueprintCompiler.cpp
+++ b/Engine/Plugins/Runtime/ModelViewViewModel/Source/ModelViewViewModelBlueprint/Private/MVVMViewBlueprintCompiler.cpp
@@ -20,6 +20,7 @@
 #include "UObject/UnrealType.h"
 #include "View/MVVMViewClass.h"
 #include "View/MVVMViewModelContextResolver.h"
+#include "WidgetBlueprintEditorUtils.h"
 
 #define LOCTEXT_NAMESPACE "MVVMViewBlueprintCompiler"
 
@@ -324,31 +325,64 @@ void FMVVMViewBlueprintCompiler::CreateVariables(const FWidgetBlueprintCompilerC
 
 	for (FCompilerUserWidgetPropertyContext& SourceContext : CompilerUserWidgetPropertyContexts)
 	{
-		SourceContext.Field = BindingHelper::FindFieldByName(Context.GetSkeletonGeneratedClass(), FMVVMBindingName(SourceContext.PropertyName));
+		SourceContext.Field = BindingHelper::FindFieldByName(Context.GetGeneratedClass(), FMVVMBindingName(SourceContext.PropertyName));
 
 		// The class is not linked yet. It may not be available yet.
 		if (SourceContext.Field.IsEmpty())
 		{
-			for (FField* Field = Context.GetSkeletonGeneratedClass()->ChildProperties; Field != nullptr; Field = Field->Next)
+			for (FField* Field = Context.GetGeneratedClass()->ChildProperties; Field != nullptr; Field = Field->Next)
 			{
 				if (Field->GetFName() == SourceContext.PropertyName)
 				{
-					if (FProperty* Property = CastField<FProperty>(Field))
+					if (CastField<FProperty>(Field))
 					{
-						SourceContext.Field = FMVVMFieldVariant(Property);
-						break;
+						SourceContext.Field = FMVVMFieldVariant(CastField<FProperty>(Field));
 					}
 					else
 					{
 						WidgetBlueprintCompilerContext.MessageLog.Error(*FText::Format(LOCTEXT("FieldIsNotProperty", "The field for source '{0}' exists but is not a property."), SourceContext.DisplayName).ToString());
 						bAreSourcesCreatorValid = false;
-						continue;
 					}
+					break;
+				}
+			}
+			for (UField* Field = Context.GetGeneratedClass()->Children; Field != nullptr; Field = Field->Next)
+			{
+				if (Field->GetFName() == SourceContext.PropertyName)
+				{
+					WidgetBlueprintCompilerContext.MessageLog.Error(*FText::Format(LOCTEXT("FieldIsNotProperty", "The field for source '{0}' exists but is not a property."), SourceContext.DisplayName).ToString());
+					bAreSourcesCreatorValid = false;
+					break;
 				}
 			}
 		}
 
-		// Reuse the property if found
+		if (SourceContext.Field.IsEmpty())
+		{
+			UClass* ParentClass = Context.GetGeneratedClass()->GetSuperClass();
+			if (const FProperty* Property = ParentClass->FindPropertyByName(SourceContext.PropertyName))
+			{
+				SourceContext.Field = FMVVMFieldVariant(Property);
+			}
+		}
+
+
+		// Will always create viewmodel properties.
+		// Will never create properties for animation or other Self.Object
+		// Will create properties for widget when they are not already created.
+
+		const bool bIsWidgetInsideWidgetTree = WidgetNameToWidgetPointerMap.Find(SourceContext.PropertyName) != nullptr;
+		const bool bIsViewmodel = SourceContext.ViewModelId.IsValid();
+
+		if (bIsViewmodel && !SourceContext.Field.IsEmpty())
+		{
+			// Viewmodel property cannot already exist. It will creates issue with initialization and with View::SetViewModel.
+			const UClass* OwnerClass = Cast<UClass>(SourceContext.Field.GetOwner());
+			WidgetBlueprintCompilerContext.MessageLog.Error(*FText::Format(LOCTEXT("ViewmodelPropertyAlreadyExistInParent", "There is already a property named '{0}' in scope '{1}' for the viewmodel."), SourceContext.DisplayName, (OwnerClass ? OwnerClass->GetDisplayNameText() : FText::GetEmpty())).ToString());
+			bAreSourceContextsValid = false;
+			continue;
+		}
+
 		if (!SourceContext.Field.IsEmpty())
 		{
 			if (!BindingHelper::IsValidForSourceBinding(SourceContext.Field))
@@ -363,13 +397,25 @@ void FMVVMViewBlueprintCompiler::CreateVariables(const FWidgetBlueprintCompilerC
 			const bool bIsCompatible = ObjectProperty && SourceContext.Class->IsChildOf(ObjectProperty->PropertyClass);
 			if (!bIsCompatible)
 			{
-				WidgetBlueprintCompilerContext.MessageLog.Error(*FText::Format(LOCTEXT("PropertyExistsAndNotCompatible","There is already a property named '{0}' that is not compatible with the source of the same name."), SourceContext.DisplayName).ToString());
+				WidgetBlueprintCompilerContext.MessageLog.Error(*FText::Format(LOCTEXT("PropertyExistsAndNotCompatible", "There is already a property named '{0}' that is not compatible with the source of the same name."), SourceContext.DisplayName).ToString());
+				bAreSourceContextsValid = false;
+				continue;
+			}
+
+			const bool bIsBindWidget = FWidgetBlueprintEditorUtils::IsBindWidgetProperty(ObjectProperty);
+			if (Context.GetGeneratedClass() != ObjectProperty->GetOwnerStruct() && !bIsBindWidget)
+			{
+				// Widget needs to be BindWidget to be reused as a property.
+				const UClass* OwnerClass = Cast<UClass>(ObjectProperty->GetOwnerStruct());
+				WidgetBlueprintCompilerContext.MessageLog.Error(*FText::Format(LOCTEXT("WidgetPropertyAlreadyExist", "There is already a property named '{0}' in scope '{1}' for the widget. Are you missing a BindWidget?."), SourceContext.DisplayName, (OwnerClass ? OwnerClass->GetDisplayNameText() : FText::GetEmpty())).ToString());
 				bAreSourceContextsValid = false;
 				continue;
 			}
 		}
 
-		if (SourceContext.Field.IsEmpty())
+		// Can we reused the property or we need to create a new one.
+		bool bCreateVariable = bIsViewmodel || (bIsWidgetInsideWidgetTree && SourceContext.Field.IsEmpty());
+		if (bCreateVariable)
 		{
 			SourceContext.Field = FMVVMConstFieldVariant(CreateVariable(SourceContext));
 		}
@@ -480,10 +526,9 @@ void FMVVMViewBlueprintCompiler::CreateSourceLists(const FWidgetBlueprintCompile
 
 	bAreSourceContextsValid = bAreSourcesCreatorValid;
 
-	UWidgetBlueprintGeneratedClass* SkeletonClass = Context.GetSkeletonGeneratedClass();
 	const FName DefaultWidgetCategory = Context.GetWidgetBlueprint()->GetFName();
 
-	auto GenerateCompilerSourceContext = [Self = this, BlueprintView, DefaultWidgetCategory, Class = SkeletonClass, &ViewModelGuids, &WidgetSources](const FMVVMBlueprintPropertyPath& PropertyPath) -> TValueOrError<void, FText>
+	auto GenerateCompilerSourceContext = [Self = this, BlueprintView, DefaultWidgetCategory, Class = Context.GetGeneratedClass(), &ViewModelGuids, &WidgetSources](const FMVVMBlueprintPropertyPath& PropertyPath) -> TValueOrError<void, FText>
 	{
 		if (PropertyPath.IsFromWidget())
 		{
diff --git a/Engine/Source/Editor/UMGEditor/Private/WidgetBlueprintCompiler.cpp b/Engine/Source/Editor/UMGEditor/Private/WidgetBlueprintCompiler.cpp
--- a/Engine/Source/Editor/UMGEditor/Private/WidgetBlueprintCompiler.cpp
+++ b/Engine/Source/Editor/UMGEditor/Private/WidgetBlueprintCompiler.cpp
@@ -59,6 +59,11 @@ UWidgetBlueprintGeneratedClass* FWidgetBlueprintCompilerContext::FCreateVariable
 	return Context.NewWidgetBlueprintClass;
 }
 
+UWidgetBlueprintGeneratedClass* FWidgetBlueprintCompilerContext::FCreateVariableContext::GetGeneratedClass() const
+{
+	return Context.NewWidgetBlueprintClass;
+}
+
 
 EKismetCompileType::Type FWidgetBlueprintCompilerContext::FCreateVariableContext::GetCompileType() const
 {
diff --git a/Engine/Source/Editor/UMGEditor/Private/WidgetBlueprintEditorUtils.cpp b/Engine/Source/Editor/UMGEditor/Private/WidgetBlueprintEditorUtils.cpp
--- a/Engine/Source/Editor/UMGEditor/Private/WidgetBlueprintEditorUtils.cpp
+++ b/Engine/Source/Editor/UMGEditor/Private/WidgetBlueprintEditorUtils.cpp
@@ -2018,13 +2018,13 @@ void FWidgetBlueprintEditorUtils::ImportPropertiesFromText(UObject* Object, cons
 	}
 }
 
-bool FWidgetBlueprintEditorUtils::IsBindWidgetProperty(FProperty* InProperty)
+bool FWidgetBlueprintEditorUtils::IsBindWidgetProperty(const FProperty* InProperty)
 {
 	bool bIsOptional;
 	return IsBindWidgetProperty(InProperty, bIsOptional);
 }
 
-bool FWidgetBlueprintEditorUtils::IsBindWidgetProperty(FProperty* InProperty, bool& bIsOptional)
+bool FWidgetBlueprintEditorUtils::IsBindWidgetProperty(const FProperty* InProperty, bool& bIsOptional)
 {
 	if ( InProperty )
 	{
@@ -2037,13 +2037,13 @@ bool FWidgetBlueprintEditorUtils::IsBindWidgetProperty(FProperty* InProperty, bo
 	return false;
 }
 
-bool FWidgetBlueprintEditorUtils::IsBindWidgetAnimProperty(FProperty* InProperty)
+bool FWidgetBlueprintEditorUtils::IsBindWidgetAnimProperty(const FProperty* InProperty)
 {
 	bool bIsOptional;
 	return IsBindWidgetAnimProperty(InProperty, bIsOptional);
 }
 
-bool FWidgetBlueprintEditorUtils::IsBindWidgetAnimProperty(FProperty* InProperty, bool& bIsOptional)
+bool FWidgetBlueprintEditorUtils::IsBindWidgetAnimProperty(const FProperty* InProperty, bool& bIsOptional)
 {
 	if (InProperty)
 	{
diff --git a/Engine/Source/Editor/UMGEditor/Public/WidgetBlueprintCompiler.h b/Engine/Source/Editor/UMGEditor/Public/WidgetBlueprintCompiler.h
--- a/Engine/Source/Editor/UMGEditor/Public/WidgetBlueprintCompiler.h
+++ b/Engine/Source/Editor/UMGEditor/Public/WidgetBlueprintCompiler.h
@@ -102,7 +102,9 @@ public:
 		FProperty* CreateVariable(const FName Name, const FEdGraphPinType& Type) const;
 		void AddGeneratedFunctionGraph(UEdGraph* Graph) const;
 		UWidgetBlueprint* GetWidgetBlueprint() const;
+		UE_DEPRECATED(5.4, "GetSkeletonGeneratedClass renamed to GetGeneratedClass")
 		UWidgetBlueprintGeneratedClass* GetSkeletonGeneratedClass() const;
+		UWidgetBlueprintGeneratedClass* GetGeneratedClass() const;
 		EKismetCompileType::Type GetCompileType() const;
 
 	private:
diff --git a/Engine/Source/Editor/UMGEditor/Public/WidgetBlueprintEditorUtils.h b/Engine/Source/Editor/UMGEditor/Public/WidgetBlueprintEditorUtils.h
--- a/Engine/Source/Editor/UMGEditor/Public/WidgetBlueprintEditorUtils.h
+++ b/Engine/Source/Editor/UMGEditor/Public/WidgetBlueprintEditorUtils.h
@@ -51,11 +51,11 @@ public:
 
 	static UWidget* GetWidgetTemplateFromDragDrop(UWidgetBlueprint* Blueprint, UWidgetTree* RootWidgetTree, TSharedPtr<FDragDropOperation>& DragDropOp);
 
-	static bool IsBindWidgetProperty(FProperty* InProperty);
-	static bool IsBindWidgetProperty(FProperty* InProperty, bool& bIsOptional);
+	static bool IsBindWidgetProperty(const FProperty* InProperty);
+	static bool IsBindWidgetProperty(const FProperty* InProperty, bool& bIsOptional);
 
-	static bool IsBindWidgetAnimProperty(FProperty* InProperty);
-	static bool IsBindWidgetAnimProperty(FProperty* InProperty, bool& bIsOptional);
+	static bool IsBindWidgetAnimProperty(const FProperty* InProperty);
+	static bool IsBindWidgetAnimProperty(const FProperty* InProperty, bool& bIsOptional);
 
 	struct FUsableWidgetClassResult
 	{
