MVVM: Allow conversion functions with UObject as argument. The bug was probably introduced to prevent static functions from BlueprintLibrary (remove the worldcontext).

diff --git a/Engine/Plugins/Runtime/ModelViewViewModel/Source/ModelViewViewModelEditor/Private/MVVMEditorSubsystem.cpp b/Engine/Plugins/Runtime/ModelViewViewModel/Source/ModelViewViewModelEditor/Private/MVVMEditorSubsystem.cpp
--- a/Engine/Plugins/Runtime/ModelViewViewModel/Source/ModelViewViewModelEditor/Private/MVVMEditorSubsystem.cpp
+++ b/Engine/Plugins/Runtime/ModelViewViewModel/Source/ModelViewViewModelEditor/Private/MVVMEditorSubsystem.cpp
@@ -679,14 +679,8 @@ bool UMVVMEditorSubsystem::IsValidConversionFunction(const UWidgetBlueprint* Wid
 	const TArray<const FProperty*>& ConversionArgProperties = ArgumentsResult.GetValue();
 	for (const FProperty* ArgumentProperty : ConversionArgProperties)
 	{
-		if (ArgumentProperty->IsA<FObjectProperty>())
-		{
-			// filter out any functions with UObject properties - they aren't valid conversion functions
-			return false;
-		}
 
-		if (SourceProperty == nullptr ||
-			UE::MVVM::BindingHelper::ArePropertiesCompatible(SourceProperty, ArgumentProperty))
+		if (SourceProperty == nullptr || UE::MVVM::BindingHelper::ArePropertiesCompatible(SourceProperty, ArgumentProperty))
 		{
 			bAnyCompatible = true;
 		}
