[Backout] - CL28715693
Original CL Desc
-----------------------------------------------------------------
MVVM: Make sure the different object ptr type matches before copying them. Evaluated path use the correct container address when evaluating a function.

diff --git a/Engine/Plugins/Runtime/ModelViewViewModel/Source/ModelViewViewModel/Private/Bindings/MVVMBindingHelper.cpp b/Engine/Plugins/Runtime/ModelViewViewModel/Source/ModelViewViewModel/Private/Bindings/MVVMBindingHelper.cpp
--- a/Engine/Plugins/Runtime/ModelViewViewModel/Source/ModelViewViewModel/Private/Bindings/MVVMBindingHelper.cpp
+++ b/Engine/Plugins/Runtime/ModelViewViewModel/Source/ModelViewViewModel/Private/Bindings/MVVMBindingHelper.cpp
@@ -461,8 +461,7 @@ namespace UE::MVVM::BindingHelper
 				&& DestinationObjectProperty
 				&& SourceObjectProperty->PropertyClass
 				&& DestinationObjectProperty->PropertyClass
-				&& SourceObjectProperty->PropertyClass->IsChildOf(DestinationObjectProperty->PropertyClass)
-				&& SourceObjectProperty->AllowObjectTypeReinterpretationTo(DestinationObjectProperty); // prevents TWeakObjectPtr to UObject*, allow TObjectPtr to UObject*
+				&& SourceObjectProperty->PropertyClass->IsChildOf(DestinationObjectProperty->PropertyClass);
 		}
 
 		bool IsObjectArrayPropertyCompatible(const FProperty* Source, const FProperty* Destination)
