MVVM: Event generated function are none const.

diff --git a/Engine/Plugins/Runtime/ModelViewViewModel/Source/ModelViewViewModelBlueprint/Private/Bindings/MVVMConversionFunctionHelper.cpp b/Engine/Plugins/Runtime/ModelViewViewModel/Source/ModelViewViewModelBlueprint/Private/Bindings/MVVMConversionFunctionHelper.cpp
--- a/Engine/Plugins/Runtime/ModelViewViewModel/Source/ModelViewViewModelBlueprint/Private/Bindings/MVVMConversionFunctionHelper.cpp
+++ b/Engine/Plugins/Runtime/ModelViewViewModel/Source/ModelViewViewModelBlueprint/Private/Bindings/MVVMConversionFunctionHelper.cpp
@@ -66,7 +66,7 @@ namespace Private
 		return nullptr;
 	}
 
-	TTuple<UEdGraph*, UK2Node_FunctionEntry*, UK2Node_FunctionResult*> CreateGraph(UBlueprint* Blueprint, FName GraphName, const UFunction* FunctionEntryDefinition, bool bIsEditable, bool bAddToBlueprint)
+	TTuple<UEdGraph*, UK2Node_FunctionEntry*, UK2Node_FunctionResult*> CreateGraph(UBlueprint* Blueprint, FName GraphName, const UFunction* FunctionEntryDefinition, bool bIsConst, bool bIsEditable, bool bAddToBlueprint)
 	{
 		FName UniqueFunctionName = FBlueprintEditorUtils::FindUniqueKismetName(Blueprint, GraphName.ToString());
 		UEdGraph* FunctionGraph = FBlueprintEditorUtils::CreateNewGraph(Blueprint, UniqueFunctionName, UEdGraph::StaticClass(), UEdGraphSchema_K2::StaticClass());
@@ -97,7 +97,14 @@ namespace Private
 		{
 			FunctionEntry->FunctionReference.SetSelfMember(FunctionGraph->GetFName());
 		}
-		FunctionEntry->AddExtraFlags(FUNC_BlueprintCallable | FUNC_BlueprintPure | FUNC_Const | FUNC_Protected | FUNC_Final);
+		if (bIsConst)
+		{
+			FunctionEntry->AddExtraFlags(FUNC_BlueprintCallable | FUNC_BlueprintPure | FUNC_Const | FUNC_Protected | FUNC_Final);
+		}
+		else
+		{
+			FunctionEntry->AddExtraFlags(FUNC_BlueprintCallable | FUNC_Protected | FUNC_Final);
+		}
 		FunctionEntry->bIsEditable = bIsEditable;
 		FunctionEntry->MetaData.Category = FText::FromName(ConversionFunctionCategory.Resolve());
 		FunctionEntry->NodePosX = -500;
@@ -141,7 +148,7 @@ FName CreateWrapperName(const FMVVMBlueprintViewBinding& Binding, bool bSourceTo
 	return FName(StringBuilder.ToString());
 }
 
-TPair<UEdGraph*, UK2Node*> CreateGraph(UBlueprint* Blueprint, FName GraphName, const UFunction* Signature, const UFunction* FunctionToWrap, bool bTransient)
+TPair<UEdGraph*, UK2Node*> CreateGraph(UBlueprint* Blueprint, FName GraphName, const UFunction* Signature, const UFunction* FunctionToWrap, bool bIsConst, bool bTransient)
 {
 	bool bIsEditable = false;
 	bool bAddToBlueprint = !bTransient;
@@ -150,7 +157,7 @@ TPair<UEdGraph*, UK2Node*> CreateGraph(UBlueprint* Blueprint, FName GraphName, c
 	UK2Node_FunctionEntry* FunctionEntry = nullptr;
 	UK2Node_FunctionResult* FunctionResult = nullptr;
 	{
-		TTuple<UEdGraph*, UK2Node_FunctionEntry*, UK2Node_FunctionResult*> NewGraph = Private::CreateGraph(Blueprint, GraphName, Signature, bIsEditable, bAddToBlueprint);
+		TTuple<UEdGraph*, UK2Node_FunctionEntry*, UK2Node_FunctionResult*> NewGraph = Private::CreateGraph(Blueprint, GraphName, Signature, bIsConst, bIsEditable, bAddToBlueprint);
 		FunctionGraph = NewGraph.Get<0>();
 		FunctionEntry = NewGraph.Get<1>();
 		FunctionResult = NewGraph.Get<2>();
@@ -213,7 +220,7 @@ TPair<UEdGraph*, UK2Node*> CreateGraph(UBlueprint* Blueprint, FName GraphName, c
 }
 
 
-TPair<UEdGraph*, UK2Node*> CreateGraph(UBlueprint* Blueprint, FName GraphName, const UFunction* Signature, TSubclassOf<UK2Node> NodeType, bool bTransient, TFunctionRef<void(UK2Node*)> InitNodeCallback)
+TPair<UEdGraph*, UK2Node*> CreateGraph(UBlueprint* Blueprint, FName GraphName, const UFunction* Signature, TSubclassOf<UK2Node> NodeType, bool bIsConst, bool bTransient, TFunctionRef<void(UK2Node*)> InitNodeCallback)
 {
 	bool bIsEditable = false;
 	bool bAddToBlueprint = !bTransient;
@@ -222,7 +229,7 @@ TPair<UEdGraph*, UK2Node*> CreateGraph(UBlueprint* Blueprint, FName GraphName, c
 	UK2Node_FunctionEntry* FunctionEntry = nullptr;
 	UK2Node_FunctionResult* FunctionResult = nullptr;
 	{
-		TTuple<UEdGraph*, UK2Node_FunctionEntry*, UK2Node_FunctionResult*> NewGraph = Private::CreateGraph(Blueprint, GraphName, Signature, bIsEditable, bAddToBlueprint);
+		TTuple<UEdGraph*, UK2Node_FunctionEntry*, UK2Node_FunctionResult*> NewGraph = Private::CreateGraph(Blueprint, GraphName, Signature, bIsConst, bIsEditable, bAddToBlueprint);
 		FunctionGraph = NewGraph.Get<0>();
 		FunctionEntry = NewGraph.Get<1>();
 		FunctionResult = NewGraph.Get<2>();
diff --git a/Engine/Plugins/Runtime/ModelViewViewModel/Source/ModelViewViewModelBlueprint/Private/MVVMBlueprintViewConversionFunction.cpp b/Engine/Plugins/Runtime/ModelViewViewModel/Source/ModelViewViewModelBlueprint/Private/MVVMBlueprintViewConversionFunction.cpp
--- a/Engine/Plugins/Runtime/ModelViewViewModel/Source/ModelViewViewModelBlueprint/Private/MVVMBlueprintViewConversionFunction.cpp
+++ b/Engine/Plugins/Runtime/ModelViewViewModel/Source/ModelViewViewModelBlueprint/Private/MVVMBlueprintViewConversionFunction.cpp
@@ -242,7 +242,8 @@ UEdGraph* UMVVMBlueprintViewConversionFunction::GetOrCreateWrapperGraphInternal(
 
 UEdGraph* UMVVMBlueprintViewConversionFunction::GetOrCreateWrapperGraphInternal(UBlueprint* Blueprint, const UFunction* Function) const
 {
-	TPair<UEdGraph*, UK2Node*> Result = UE::MVVM::ConversionFunctionHelper::CreateGraph(Blueprint, GraphName, nullptr, Function, bWrapperGraphTransient);
+	bool bConst = true;
+	TPair<UEdGraph*, UK2Node*> Result = UE::MVVM::ConversionFunctionHelper::CreateGraph(Blueprint, GraphName, nullptr, Function, bConst, bWrapperGraphTransient);
 	CachedWrapperGraph = Result.Get<0>();
 	CachedWrapperNode = Result.Get<1>();
 	LoadPinValuesInternal(Blueprint);
diff --git a/Engine/Plugins/Runtime/ModelViewViewModel/Source/ModelViewViewModelBlueprint/Private/MVVMBlueprintViewEvent.cpp b/Engine/Plugins/Runtime/ModelViewViewModel/Source/ModelViewViewModelBlueprint/Private/MVVMBlueprintViewEvent.cpp
--- a/Engine/Plugins/Runtime/ModelViewViewModel/Source/ModelViewViewModelBlueprint/Private/MVVMBlueprintViewEvent.cpp
+++ b/Engine/Plugins/Runtime/ModelViewViewModel/Source/ModelViewViewModelBlueprint/Private/MVVMBlueprintViewEvent.cpp
@@ -215,8 +215,9 @@ UEdGraph* UMVVMBlueprintViewEvent::CreateWrapperGraphInternal()
 UEdGraph* UMVVMBlueprintViewEvent::CreateWrapperGraphInternal(const UFunction* DelegateSignature, const UFunction* Function)
 {
 	UWidgetBlueprint* WidgetBlueprint = GetWidgetBlueprintInternal();
+	bool bIsConst = false;
 	bool bTransient = true;
-	TPair<UEdGraph*, UK2Node*> Result = UE::MVVM::ConversionFunctionHelper::CreateGraph(WidgetBlueprint, GraphName, DelegateSignature, Function, bTransient);
+	TPair<UEdGraph*, UK2Node*> Result = UE::MVVM::ConversionFunctionHelper::CreateGraph(WidgetBlueprint, GraphName, DelegateSignature, Function, bIsConst, bTransient);
 	CachedWrapperGraph = Result.Get<0>();
 	CachedWrapperNode = Result.Get<1>();
 
@@ -244,8 +245,9 @@ UEdGraph* UMVVMBlueprintViewEvent::CreateWrapperGraphInternal(const UFunction* D
 UEdGraph* UMVVMBlueprintViewEvent::CreateWrapperGraphInternal(const UFunction* DelegateSignature, const FProperty* Property)
 {
 	UWidgetBlueprint* WidgetBlueprint = GetWidgetBlueprintInternal();
+	bool bConst = false;
 	bool bTransient = true;
-	TPair<UEdGraph*, UK2Node*> Result = UE::MVVM::ConversionFunctionHelper::CreateGraph(WidgetBlueprint, GraphName, DelegateSignature, UK2Node_VariableSet::StaticClass(), bTransient,
+	TPair<UEdGraph*, UK2Node*> Result = UE::MVVM::ConversionFunctionHelper::CreateGraph(WidgetBlueprint, GraphName, DelegateSignature, UK2Node_VariableSet::StaticClass(), bConst, bTransient,
 		[WidgetBlueprint, Property](UK2Node* NewNode)
 		{
 			UK2Node_VariableSet* VariableNode = CastChecked<UK2Node_VariableSet>(NewNode);
diff --git a/Engine/Plugins/Runtime/ModelViewViewModel/Source/ModelViewViewModelBlueprint/Public/Bindings/MVVMConversionFunctionHelper.h b/Engine/Plugins/Runtime/ModelViewViewModel/Source/ModelViewViewModelBlueprint/Public/Bindings/MVVMConversionFunctionHelper.h
--- a/Engine/Plugins/Runtime/ModelViewViewModel/Source/ModelViewViewModelBlueprint/Public/Bindings/MVVMConversionFunctionHelper.h
+++ b/Engine/Plugins/Runtime/ModelViewViewModel/Source/ModelViewViewModelBlueprint/Public/Bindings/MVVMConversionFunctionHelper.h
@@ -42,10 +42,10 @@ namespace UE::MVVM::ConversionFunctionHelper
 	MODELVIEWVIEWMODELBLUEPRINT_API FName CreateWrapperName(const FMVVMBlueprintViewBinding& Binding, bool bSourceToDestination);
 
 	/** */
-	MODELVIEWVIEWMODELBLUEPRINT_API TPair<UEdGraph*, UK2Node*> CreateGraph(UBlueprint* WidgetBlueprint, FName GraphName, const UFunction* Signature, const UFunction* FunctionToWrap, bool bTransient);
+	MODELVIEWVIEWMODELBLUEPRINT_API TPair<UEdGraph*, UK2Node*> CreateGraph(UBlueprint* WidgetBlueprint, FName GraphName, const UFunction* Signature, const UFunction* FunctionToWrap, bool bIsConst, bool bTransient);
 	
 	/** */
-	MODELVIEWVIEWMODELBLUEPRINT_API TPair<UEdGraph*, UK2Node*> CreateGraph(UBlueprint* WidgetBlueprint, FName GraphName, const UFunction* Signature, const TSubclassOf<UK2Node> Node, bool bTransient, TFunctionRef<void(UK2Node*)> InitNodeCallback);
+	MODELVIEWVIEWMODELBLUEPRINT_API TPair<UEdGraph*, UK2Node*> CreateGraph(UBlueprint* WidgetBlueprint, FName GraphName, const UFunction* Signature, const TSubclassOf<UK2Node> Node, bool bIsConst, bool bTransient, TFunctionRef<void(UK2Node*)> InitNodeCallback);
 
 	/** Find the conversion function node from the given graph. */
 	MODELVIEWVIEWMODELBLUEPRINT_API UK2Node* GetWrapperNode(UEdGraph* Graph);
