MVVM: Allow a property or a function if it's from the blueprint we are currently compiling.

diff --git a/Engine/Plugins/Runtime/ModelViewViewModel/Source/ModelViewViewModelBlueprint/Private/Bindings/MVVMCompiledBindingLibraryCompiler.cpp b/Engine/Plugins/Runtime/ModelViewViewModel/Source/ModelViewViewModelBlueprint/Private/Bindings/MVVMCompiledBindingLibraryCompiler.cpp
--- a/Engine/Plugins/Runtime/ModelViewViewModel/Source/ModelViewViewModelBlueprint/Private/Bindings/MVVMCompiledBindingLibraryCompiler.cpp
+++ b/Engine/Plugins/Runtime/ModelViewViewModel/Source/ModelViewViewModelBlueprint/Private/Bindings/MVVMCompiledBindingLibraryCompiler.cpp
@@ -89,12 +89,18 @@ public:
 /** */
 class FCompiledBindingLibraryCompilerImpl
 {
+public:
+	FCompiledBindingLibraryCompilerImpl(const UBlueprint* Context)
+		: BlueprintContext(Context)
+	{}
+
 public:
 	TArray<FRawFieldId> FieldIds;
 	TArray<FRawField> Fields;
 	TArray<FRawFieldPath> FieldPaths;
 	TArray<FRawBinding> Bindings;
 	bool bCompiled = false;
+	TWeakObjectPtr<const UBlueprint> BlueprintContext;
 
 public:
 	int32 AddUniqueField(FMVVMConstFieldVariant InFieldVariant)
@@ -204,8 +210,8 @@ FCompiledBindingLibraryCompiler::FCompileResult::FCompileResult(FGuid LibraryId)
 /**
  *
  */
-FCompiledBindingLibraryCompiler::FCompiledBindingLibraryCompiler()
-	: Impl(MakePimpl<Private::FCompiledBindingLibraryCompilerImpl>())
+FCompiledBindingLibraryCompiler::FCompiledBindingLibraryCompiler(UBlueprint* GeneratingFor)
+	: Impl(MakePimpl<Private::FCompiledBindingLibraryCompilerImpl>(GeneratingFor))
 {
 
 }
@@ -340,7 +346,7 @@ TValueOrError<FCompiledBindingLibraryCompiler::FFieldPathHandle, FText> FCompile
 				return MakeError(ValidatedStr);
 			}
 
-			if (!GetDefault<UMVVMDeveloperProjectSettings>()->IsPropertyAllowed(FieldVariant.GetProperty()))
+			if (!GetDefault<UMVVMDeveloperProjectSettings>()->IsPropertyAllowed(Impl->BlueprintContext.Get(), FieldVariant.GetProperty()))
 			{
 				return MakeError(LOCTEXT("PropertyNotAllow", "A property is not allowed."));
 			}
@@ -361,7 +367,7 @@ TValueOrError<FCompiledBindingLibraryCompiler::FFieldPathHandle, FText> FCompile
 				return MakeError(FText::Format(LOCTEXT("FunctionNotReadableAtRuntime", "Function '{0}' is not readable at runtime."), FieldVariant.GetFunction()->GetDisplayNameText()));
 			}
 
-			if (!GetDefault<UMVVMDeveloperProjectSettings>()->IsFunctionAllowed(FieldVariant.GetFunction()))
+			if (!GetDefault<UMVVMDeveloperProjectSettings>()->IsFunctionAllowed(Impl->BlueprintContext.Get(), FieldVariant.GetFunction()))
 			{
 				return MakeError(LOCTEXT("FunctionNotAllow", "A function is not allowed."));
 			}
diff --git a/Engine/Plugins/Runtime/ModelViewViewModel/Source/ModelViewViewModelBlueprint/Private/MVVMDeveloperProjectSettings.cpp b/Engine/Plugins/Runtime/ModelViewViewModel/Source/ModelViewViewModelBlueprint/Private/MVVMDeveloperProjectSettings.cpp
--- a/Engine/Plugins/Runtime/ModelViewViewModel/Source/ModelViewViewModelBlueprint/Private/MVVMDeveloperProjectSettings.cpp
+++ b/Engine/Plugins/Runtime/ModelViewViewModel/Source/ModelViewViewModelBlueprint/Private/MVVMDeveloperProjectSettings.cpp
@@ -1,9 +1,10 @@
 // Copyright Epic Games, Inc. All Rights Reserved.
 
 #include "MVVMDeveloperProjectSettings.h"
-#include "Engine/Blueprint.h"
 
 #include "BlueprintEditorSettings.h"
+#include "Engine/Blueprint.h"
+#include "Kismet2/BlueprintEditorUtils.h"
 #include "MVVMBlueprintViewModelContext.h"
 #include "PropertyPermissionList.h"
 #include "Types/MVVMExecutionMode.h"
@@ -66,10 +67,25 @@ bool UMVVMDeveloperProjectSettings::PropertyHasFiltering(const FProperty* Proper
 	return true;
 }
 
-bool UMVVMDeveloperProjectSettings::IsPropertyAllowed(const FProperty* Property) const
+namespace UE::MVVM::Private
+{
+bool ShouldDoPropertyEditorPermission(const UBlueprint* GeneratingFor, UClass* FieldOwner)
+{
+	if (GeneratingFor && FieldOwner)
+	{
+		const UClass* UpToDateClass = FBlueprintEditorUtils::GetMostUpToDateClass(FieldOwner);
+		return GeneratingFor->SkeletonGeneratedClass != UpToDateClass;
+	}
+	return true;
+}
+}//namespace
+
+bool UMVVMDeveloperProjectSettings::IsPropertyAllowed(const UBlueprint* GeneratingFor, const FProperty* Property) const
 {
 	check(Property);
-	if (!FPropertyEditorPermissionList::Get().DoesPropertyPassFilter(Property->GetOwnerStruct(), Property->GetFName()))
+	check(GeneratingFor);
+	const bool bDoPropertyEditorPermission = UE::MVVM::Private::ShouldDoPropertyEditorPermission(GeneratingFor, Property->GetTypedOwner<UClass>());
+	if (bDoPropertyEditorPermission && !FPropertyEditorPermissionList::Get().DoesPropertyPassFilter(Property->GetOwnerStruct(), Property->GetFName()))
 	{
 		return false;
 	}
@@ -85,7 +101,7 @@ bool UMVVMDeveloperProjectSettings::IsPropertyAllowed(const FProperty* Property)
 	return true;
 }
 
-bool UMVVMDeveloperProjectSettings::IsFunctionAllowed(const UFunction* Function) const
+bool UMVVMDeveloperProjectSettings::IsFunctionAllowed(const UBlueprint* GeneratingFor, const UFunction* Function) const
 {
 	check(Function);
 
@@ -111,7 +127,7 @@ bool UMVVMDeveloperProjectSettings::IsFunctionAllowed(const UFunction* Function)
 	return true;
 }
 
-bool UMVVMDeveloperProjectSettings::IsConversionFunctionAllowed(const UFunction* Function) const
+bool UMVVMDeveloperProjectSettings::IsConversionFunctionAllowed(const UBlueprint* GeneratingFor, const UFunction* Function) const
 {
 	static FName NAME_ComplexConversionFunction = TEXT("MVVMComplexConversionFunction");
 	if (Function->HasMetaData(NAME_ComplexConversionFunction))
@@ -121,7 +137,7 @@ bool UMVVMDeveloperProjectSettings::IsConversionFunctionAllowed(const UFunction*
 
 	if (ConversionFunctionFilter == EMVVMDeveloperConversionFunctionFilterType::BlueprintActionRegistry)
 	{
-		return IsFunctionAllowed(Function);
+		return IsFunctionAllowed(GeneratingFor, Function);
 	}
 	else
 	{
@@ -147,7 +163,7 @@ bool UMVVMDeveloperProjectSettings::IsConversionFunctionAllowed(const UFunction*
 		else
 		{
 			// The function is on self and may have been filtered.
-			return IsFunctionAllowed(Function);
+			return IsFunctionAllowed(GeneratingFor, Function);
 		}
 	}
 }
diff --git a/Engine/Plugins/Runtime/ModelViewViewModel/Source/ModelViewViewModelBlueprint/Private/MVVMViewBlueprintCompiler.cpp b/Engine/Plugins/Runtime/ModelViewViewModel/Source/ModelViewViewModelBlueprint/Private/MVVMViewBlueprintCompiler.cpp
--- a/Engine/Plugins/Runtime/ModelViewViewModel/Source/ModelViewViewModelBlueprint/Private/MVVMViewBlueprintCompiler.cpp
+++ b/Engine/Plugins/Runtime/ModelViewViewModel/Source/ModelViewViewModelBlueprint/Private/MVVMViewBlueprintCompiler.cpp
@@ -42,7 +42,7 @@ FAutoConsoleCommand CVarTestGenerateSetter(
 		{
 			return;
 		}
-		FMVVMViewBlueprintCompiler::TestGenerateSetter(Args[0], Args[1], Args[2]);
+		FMVVMViewBlueprintCompiler::TestGenerateSetter(nullptr, Args[0], Args[1], Args[2]);
 	})
 );
 #endif
@@ -117,6 +117,11 @@ TValueOrError<FCompiledBindingLibraryCompiler::FFieldPathHandle, FText> AddObjec
 	return MakeValue(ReadFieldPathResult.StealValue());
 }
 
+FMVVMViewBlueprintCompiler::FMVVMViewBlueprintCompiler(FWidgetBlueprintCompilerContext& InCreationContext)
+	: WidgetBlueprintCompilerContext(InCreationContext)
+	, BindingLibraryCompiler(InCreationContext.WidgetBlueprint())
+{}
+
 void FMVVMViewBlueprintCompiler::AddMessageForBinding(FMVVMBlueprintViewBinding& Binding, UMVVMBlueprintView* BlueprintView, const FText& MessageText, EBindingMessageType MessageType, FName ArgumentName) const
 {
 	const FText BindingName = FText::FromString(Binding.GetDisplayNameString(WidgetBlueprintCompilerContext.WidgetBlueprint()));
@@ -877,7 +882,7 @@ void FMVVMViewBlueprintCompiler::CreateBindingDestinationContexts(UMVVMBlueprint
 
 			const FMVVMBlueprintPropertyPath& DestinationPath = bIsForwardBinding ? BindingPtr->DestinationPath : BindingPtr->SourcePath;
 			TArray<UE::MVVM::FMVVMConstFieldVariant> SkeletalSetterPath = CreateBindingDestinationPath(BlueprintView, NewSkeletonClass, DestinationPath);
-			if (!IsPropertyPathValid(SkeletalSetterPath))
+			if (!IsPropertyPathValid(WidgetBlueprintCompilerContext.WidgetBlueprint(), SkeletalSetterPath))
 			{
 				AddMessageForBinding(*BindingPtr, BlueprintView, FText::Format(Private::PropertyPathIsInvalidFormat,
 					PropertyPathToText(NewSkeletonClass, BlueprintView, DestinationPath)),
@@ -889,7 +894,7 @@ void FMVVMViewBlueprintCompiler::CreateBindingDestinationContexts(UMVVMBlueprint
 			
 			// Generate the path with property converted to BP function
 			TValueOrError<TArray<FMVVMConstFieldVariant>, FText> SkeletalGeneratedField = FieldPathHelper::GenerateFieldPathList(SkeletalSetterPath, false);
-			if (SkeletalGeneratedField.HasError() || !IsPropertyPathValid(SkeletalGeneratedField.GetValue()))
+			if (SkeletalGeneratedField.HasError() || !IsPropertyPathValid(WidgetBlueprintCompilerContext.WidgetBlueprint(), SkeletalGeneratedField.GetValue()))
 			{
 				AddMessageForBinding(*BindingPtr, BlueprintView, FText::Format(Private::CouldNotCreateSourceFieldPathFormat
 					, ::UE::MVVM::FieldPathHelper::ToText(SkeletalSetterPath)
@@ -1050,12 +1055,12 @@ bool FMVVMViewBlueprintCompiler::PreCompileBindingSources(UWidgetBlueprintGenera
 
 		bool bIsOneTimeBinding = IsOneTimeBinding(Binding.BindingType);
 
-		auto CreateSourceContextForPropertyPath = [this, &Binding, BlueprintView, Class, Index, bIsOneTimeBinding](const FMVVMBlueprintPropertyPath& Path, bool bForwardBinding, int32 ComplexConversionFunctionContextIndex, FName ArgumentName) -> bool
+		auto CreateSourceContextForPropertyPath = [Self=this, &Binding, BlueprintView, Class, Index, bIsOneTimeBinding](const FMVVMBlueprintPropertyPath& Path, bool bForwardBinding, int32 ComplexConversionFunctionContextIndex, FName ArgumentName) -> bool
 		{
-			const TValueOrError<FBindingSourceContext, FText> CreatedBindingSourceContext = CreateBindingSourceContext(BlueprintView, Class, Path, bIsOneTimeBinding);
+			const TValueOrError<FBindingSourceContext, FText> CreatedBindingSourceContext = Self->CreateBindingSourceContext(BlueprintView, Class, Path, bIsOneTimeBinding);
 			if (CreatedBindingSourceContext.HasError())
 			{
-				AddMessageForBinding(Binding, BlueprintView,
+				Self->AddMessageForBinding(Binding, BlueprintView,
 					FText::Format(LOCTEXT("PropertyPathInvalidWithReason", "The property path '{0}' is invalid. {1}"),
 						PropertyPathToText(Class, BlueprintView, Binding.SourcePath),
 						CreatedBindingSourceContext.GetError()
@@ -1067,9 +1072,9 @@ bool FMVVMViewBlueprintCompiler::PreCompileBindingSources(UWidgetBlueprintGenera
 			}
 
 			FBindingSourceContext BindingSourceContext = CreatedBindingSourceContext.GetValue();
-			if (!IsPropertyPathValid(BindingSourceContext.PropertyPath))
+			if (!IsPropertyPathValid(Self->WidgetBlueprintCompilerContext.WidgetBlueprint(), BindingSourceContext.PropertyPath))
 			{
-				AddMessageForBinding(Binding,
+				Self->AddMessageForBinding(Binding,
 					BlueprintView,
 					FText::Format(Private::PropertyPathIsInvalidFormat, PropertyPathToText(Class, BlueprintView, Binding.SourcePath)),
 					EBindingMessageType::Error,
@@ -1080,13 +1085,13 @@ bool FMVVMViewBlueprintCompiler::PreCompileBindingSources(UWidgetBlueprintGenera
 
 			if (BindingSourceContext.SourceClass == nullptr)
 			{
-				AddMessageForBinding(Binding, BlueprintView, LOCTEXT("BindingInvalidSourceClass", "Internal error. The binding could not find its source class."), EBindingMessageType::Error, ArgumentName);
+				Self->AddMessageForBinding(Binding, BlueprintView, LOCTEXT("BindingInvalidSourceClass", "Internal error. The binding could not find its source class."), EBindingMessageType::Error, ArgumentName);
 				return false;
 			}
 
 			if (!BindingSourceContext.bIsRootWidget && BindingSourceContext.UserWidgetPropertyContextIndex == INDEX_NONE && BindingSourceContext.SourceCreatorContextIndex == INDEX_NONE)
 			{
-				AddMessageForBinding(Binding, BlueprintView, LOCTEXT("BindingInvalidSource", "Internal error. The binding could not find its source."), EBindingMessageType::Error, ArgumentName);
+				Self->AddMessageForBinding(Binding, BlueprintView, LOCTEXT("BindingInvalidSource", "Internal error. The binding could not find its source."), EBindingMessageType::Error, ArgumentName);
 				return false;
 			}
 
@@ -1094,7 +1099,7 @@ bool FMVVMViewBlueprintCompiler::PreCompileBindingSources(UWidgetBlueprintGenera
 			BindingSourceContext.bIsForwardBinding = bForwardBinding;
 			BindingSourceContext.ComplexConversionFunctionContextIndex = ComplexConversionFunctionContextIndex;
 
-			this->BindingSourceContexts.Add(MoveTemp(BindingSourceContext));
+			Self->BindingSourceContexts.Add(MoveTemp(BindingSourceContext));
 			return true;
 		};
 
@@ -1565,7 +1570,7 @@ bool FMVVMViewBlueprintCompiler::PreCompileBindings(UWidgetBlueprintGeneratedCla
 		return false;
 	}
 
-	auto GetConversionFunction = [this, Class, BlueprintView](FMVVMBlueprintViewBinding& Binding, bool bIsForwardBinding) -> TValueOrError<const UFunction*, void>
+	auto GetConversionFunction = [Self = this, Class, BlueprintView](FMVVMBlueprintViewBinding& Binding, bool bIsForwardBinding) -> TValueOrError<const UFunction*, void>
 	{
 		const UFunction* ConversionFunction = nullptr;
 		const UMVVMBlueprintViewConversionFunction* ViewConversionFunction = Binding.Conversion.GetConversionFunction(bIsForwardBinding);
@@ -1574,60 +1579,60 @@ bool FMVVMViewBlueprintCompiler::PreCompileBindings(UWidgetBlueprintGeneratedCla
 			ConversionFunction = ViewConversionFunction->GetCompiledFunction(Class);
 			if (ConversionFunction == nullptr)
 			{
-				AddMessageForBinding(Binding, BlueprintView, FText::Format(LOCTEXT("ConversionFunctionNotFound", "The conversion function '{0}' could not be found."),
+				Self->AddMessageForBinding(Binding, BlueprintView, FText::Format(LOCTEXT("ConversionFunctionNotFound", "The conversion function '{0}' could not be found."),
 					FText::FromName(ViewConversionFunction->GetCompiledFunctionName())),
 					EBindingMessageType::Error
 				);
-				bAreBindingsValid = false;
+				Self->bAreBindingsValid = false;
 				return MakeError();
 			}
 
 			TVariant<const UFunction*, TSubclassOf<UK2Node>> FunctionOrWrapperFunction = ViewConversionFunction->GetConversionFunction(Class);
 			if (FunctionOrWrapperFunction.IsType<const UFunction*>())
 			{
-				if (!GetDefault<UMVVMDeveloperProjectSettings>()->IsConversionFunctionAllowed(FunctionOrWrapperFunction.Get<const UFunction*>()))
+				if (!GetDefault<UMVVMDeveloperProjectSettings>()->IsConversionFunctionAllowed(Self->WidgetBlueprintCompilerContext.WidgetBlueprint(), FunctionOrWrapperFunction.Get<const UFunction*>()))
 				{
-					AddMessageForBinding(Binding, BlueprintView, FText::Format(LOCTEXT("ConversionFunctionNotAllow", "The conversion function {0} is not allowed."),
+					Self->AddMessageForBinding(Binding, BlueprintView, FText::Format(LOCTEXT("ConversionFunctionNotAllow", "The conversion function {0} is not allowed."),
 						FText::FromName(FunctionOrWrapperFunction.Get<const UFunction*>()->GetFName())),
 						EBindingMessageType::Error
 					);
-					bAreBindingsValid = false;
+					Self->bAreBindingsValid = false;
 					return MakeError();
 				}
 			}
 			else
 			{
-				AddMessageForBinding(Binding, BlueprintView, LOCTEXT("ConversionFunctionNodeNotAllow", "The conversion function node is not allowed."), EBindingMessageType::Error);
-				bAreBindingsValid = false;
+				Self->AddMessageForBinding(Binding, BlueprintView, LOCTEXT("ConversionFunctionNodeNotAllow", "The conversion function node is not allowed."), EBindingMessageType::Error);
+				Self->bAreBindingsValid = false;
 				return MakeError();
 			}
 		}
 		return MakeValue(ConversionFunction);
 	};
 
-	auto TestExecutionMode = [this, BlueprintView](FMVVMBlueprintViewBinding& Binding)
+	auto TestExecutionMode = [Self = this, BlueprintView](FMVVMBlueprintViewBinding& Binding)
 	{
 		if (Binding.bOverrideExecutionMode)
 		{
 			if (!GetDefault<UMVVMDeveloperProjectSettings>()->IsExecutionModeAllowed(Binding.OverrideExecutionMode))
 			{
-				AddMessageForBinding(Binding, BlueprintView, LOCTEXT("NotAllowedExecutionMode", "The binding has a restricted execution mode."), EBindingMessageType::Error);
+				Self->AddMessageForBinding(Binding, BlueprintView, LOCTEXT("NotAllowedExecutionMode", "The binding has a restricted execution mode."), EBindingMessageType::Error);
 				return false;
 			}
 		}
 		return true;
 	};
 
-	auto GetSetterPath = [this, Class, BlueprintView](FMVVMBlueprintViewBinding& Binding, int32 BindingIndex, bool bIsForwardBinding) -> TValueOrError<TArray<UE::MVVM::FMVVMConstFieldVariant>, void>
+	auto GetSetterPath = [Self = this, Class, BlueprintView](FMVVMBlueprintViewBinding& Binding, int32 BindingIndex, bool bIsForwardBinding) -> TValueOrError<TArray<UE::MVVM::FMVVMConstFieldVariant>, void>
 	{
-		const FBindingDestinationContext* DestinationContext = BindingDestinationContexts.FindByPredicate([BindingIndex, bIsForwardBinding](const FBindingDestinationContext& Other)
+		const FBindingDestinationContext* DestinationContext = Self->BindingDestinationContexts.FindByPredicate([BindingIndex, bIsForwardBinding](const FBindingDestinationContext& Other)
 			{
 				return Other.BindingIndex == BindingIndex && Other.bIsForwardBinding == bIsForwardBinding;
 			});
 
 		if (DestinationContext == nullptr)
 		{
-			AddMessageForBinding(Binding, BlueprintView, LOCTEXT("CouldNotFindDestination", "Could not find the pre compiled destination."), EBindingMessageType::Error);
+			Self->AddMessageForBinding(Binding, BlueprintView, LOCTEXT("CouldNotFindDestination", "Could not find the pre compiled destination."), EBindingMessageType::Error);
 			return MakeError();
 		}
 
@@ -1638,24 +1643,24 @@ bool FMVVMViewBlueprintCompiler::PreCompileBindings(UWidgetBlueprintGeneratedCla
 			UFunction* FoundFunction = Class->FindFunctionByName(DestinationContext->GeneratedFunctionName);
 			if (FoundFunction == nullptr)
 			{
-				AddMessageForBinding(Binding, BlueprintView, LOCTEXT("CouldNotFindDestinationBPFunction", "Could not find the generated destination function."), EBindingMessageType::Error);
+				Self->AddMessageForBinding(Binding, BlueprintView, LOCTEXT("CouldNotFindDestinationBPFunction", "Could not find the generated destination function."), EBindingMessageType::Error);
 				return MakeError();
 			}
 			FMVVMBlueprintPropertyPath DestinationPPropertyPath;
 			DestinationPPropertyPath.SetWidgetName(Class->ClassGeneratedBy->GetFName());
 			DestinationPPropertyPath.SetPropertyPath(BlueprintView->GetOuterUMVVMWidgetBlueprintExtension_View()->GetWidgetBlueprint(), UE::MVVM::FMVVMConstFieldVariant(FoundFunction));
-			SetterPath = CreateBindingDestinationPath(BlueprintView, Class, DestinationPPropertyPath);
+			SetterPath = Self->CreateBindingDestinationPath(BlueprintView, Class, DestinationPPropertyPath);
 		}
 		else
 		{
 			const FMVVMBlueprintPropertyPath& DestinationPath = bIsForwardBinding ? Binding.DestinationPath : Binding.SourcePath;
-			SetterPath = CreateBindingDestinationPath(BlueprintView, Class, DestinationPath);
+			SetterPath = Self->CreateBindingDestinationPath(BlueprintView, Class, DestinationPath);
 		}
 
-		if (!IsPropertyPathValid(SetterPath))
+		if (!IsPropertyPathValid(Self->WidgetBlueprintCompilerContext.WidgetBlueprint(), SetterPath))
 		{
 			const FMVVMBlueprintPropertyPath& DestinationPath = bIsForwardBinding ? Binding.DestinationPath : Binding.SourcePath;
-			AddMessageForBinding(Binding, BlueprintView, FText::Format(Private::PropertyPathIsInvalidFormat, PropertyPathToText(Class, BlueprintView, DestinationPath)), EBindingMessageType::Error);
+			Self->AddMessageForBinding(Binding, BlueprintView, FText::Format(Private::PropertyPathIsInvalidFormat, PropertyPathToText(Class, BlueprintView, DestinationPath)), EBindingMessageType::Error);
 			return MakeError();
 		}
 		return MakeValue(SetterPath);
@@ -2117,7 +2122,7 @@ bool FMVVMViewBlueprintCompiler::PreCompileEvents(UWidgetBlueprintGeneratedClass
 		auto TestPath = [Self = this, BlueprintView, Class](UMVVMBlueprintViewEvent* EventPtr, const FMVVMBlueprintPropertyPath& PropertyPath, FName ArgumentName)
 		{
 			TArray<UE::MVVM::FMVVMConstFieldVariant> GeneratedPath = Self->CreateBindingDestinationPath(BlueprintView, Class, PropertyPath);
-			if (!Self->IsPropertyPathValid(GeneratedPath))
+			if (!Self->IsPropertyPathValid(Self->WidgetBlueprintCompilerContext.WidgetBlueprint(), GeneratedPath))
 			{
 				Self->AddMessageForEvent(EventPtr
 					, FText::Format(Private::PropertyPathIsInvalidFormat, PropertyPathToText(Class, BlueprintView, PropertyPath))
@@ -2614,7 +2619,7 @@ TArray<FMVVMConstFieldVariant> FMVVMViewBlueprintCompiler::CreatePropertyPath(co
 }
 
 
-bool FMVVMViewBlueprintCompiler::IsPropertyPathValid(TArrayView<const FMVVMConstFieldVariant> PropertyPath)
+bool FMVVMViewBlueprintCompiler::IsPropertyPathValid(const UBlueprint* Context, TArrayView<const FMVVMConstFieldVariant> PropertyPath)
 {
 	int32 PathLength = PropertyPath.Num();
 	for (int32 Index = 0; Index < PathLength; Index++)
@@ -2632,15 +2637,7 @@ bool FMVVMViewBlueprintCompiler::IsPropertyPathValid(TArrayView<const FMVVMConst
 				return false;
 			}
 
-			// Don't filter objects in path if properties remain in path to be filtered (by length). 
-			// Child properties will acquire their parent struct / these objects for filtering.
-			// @TODO: Consider same for 'FCompiledBindingLibraryCompiler::AddFieldPathImpl' & 'UE::MVVM::Private::PassFilter'.
-			const FObjectProperty* ObjectProperty = CastField<FObjectProperty>(Field.GetProperty());
-			if (ObjectProperty && Index != PathLength - 1)
-			{
-				// Do nothing
-			}
-			else if (!GetDefault<UMVVMDeveloperProjectSettings>()->IsPropertyAllowed(Field.GetProperty()))
+			if (!GetDefault<UMVVMDeveloperProjectSettings>()->IsPropertyAllowed(Context, Field.GetProperty()))
 			{
 				return false;
 			}
@@ -2651,7 +2648,7 @@ bool FMVVMViewBlueprintCompiler::IsPropertyPathValid(TArrayView<const FMVVMConst
 			{
 				return false;
 			}
-			if (!GetDefault<UMVVMDeveloperProjectSettings>()->IsFunctionAllowed(Field.GetFunction()))
+			if (!GetDefault<UMVVMDeveloperProjectSettings>()->IsFunctionAllowed(Context, Field.GetFunction()))
 			{
 				return false;
 			}
@@ -2696,7 +2693,7 @@ bool FMVVMViewBlueprintCompiler::CanBeSetInNative(TArrayView<const FMVVMConstFie
 	return true;
 }
 
-void FMVVMViewBlueprintCompiler::TestGenerateSetter(FStringView ObjectName, FStringView FieldPath, FStringView FunctionName)
+void FMVVMViewBlueprintCompiler::TestGenerateSetter(const UBlueprint* Context, FStringView ObjectName, FStringView FieldPath, FStringView FunctionName)
 {
 #if UE_BUILD_DEVELOPMENT || UE_BUILD_DEBUG
 	UWidgetBlueprint* WidgetBlueprint = nullptr;
@@ -2716,7 +2713,7 @@ void FMVVMViewBlueprintCompiler::TestGenerateSetter(FStringView ObjectName, FStr
 	{
 		return;
 	}
-	if (!IsPropertyPathValid(SkeletalSetterPathResult.GetValue()))
+	if (!IsPropertyPathValid(Context, SkeletalSetterPathResult.GetValue()))
 	{
 		return;
 	}
diff --git a/Engine/Plugins/Runtime/ModelViewViewModel/Source/ModelViewViewModelBlueprint/Private/MVVMViewBlueprintCompiler.h b/Engine/Plugins/Runtime/ModelViewViewModel/Source/ModelViewViewModelBlueprint/Private/MVVMViewBlueprintCompiler.h
--- a/Engine/Plugins/Runtime/ModelViewViewModel/Source/ModelViewViewModelBlueprint/Private/MVVMViewBlueprintCompiler.h
+++ b/Engine/Plugins/Runtime/ModelViewViewModel/Source/ModelViewViewModelBlueprint/Private/MVVMViewBlueprintCompiler.h
@@ -36,10 +36,7 @@ private:
 	struct FBindingDestinationContext;
 
 public:
-	FMVVMViewBlueprintCompiler(FWidgetBlueprintCompilerContext& InCreationContext)
-		: WidgetBlueprintCompilerContext(InCreationContext)
-	{}
-
+	FMVVMViewBlueprintCompiler(FWidgetBlueprintCompilerContext& InCreationContext);
 
 	FWidgetBlueprintCompilerContext& GetCompilerContext()
 	{
@@ -60,7 +57,7 @@ public:
 	bool Compile(UWidgetBlueprintGeneratedClass* Class, UMVVMBlueprintView* BlueprintView, UMVVMViewClass* ViewExtension);
 
 
-	static void TestGenerateSetter(FStringView ObjectName, FStringView FieldPath, FStringView FunctionName);
+	static void TestGenerateSetter(const UBlueprint* Context, FStringView ObjectName, FStringView FieldPath, FStringView FunctionName);
 
 private:
 	void CreateWidgetMap(const FWidgetBlueprintCompilerContext::FCreateVariableContext& Context, UMVVMBlueprintView* BlueprintView);
@@ -92,7 +89,7 @@ private:
 	TValueOrError<FCompiledBinding, FText> CreateCompiledBinding(const UWidgetBlueprintGeneratedClass* Class, TArrayView<const UE::MVVM::FMVVMConstFieldVariant> GetterFields, TArrayView<const UE::MVVM::FMVVMConstFieldVariant> SetterFields, const UFunction* ConversionFunction, bool bIsComplexBinding);
 
 	static TArray<FMVVMConstFieldVariant> CreatePropertyPath(const UClass* Class, FName PropertyName, TArray<FMVVMConstFieldVariant> Properties);
-	static bool IsPropertyPathValid(TArrayView<const FMVVMConstFieldVariant> PropertyPath);
+	static bool IsPropertyPathValid(const UBlueprint* Context, TArrayView<const FMVVMConstFieldVariant> PropertyPath);
 	static bool CanBeSetInNative(TArrayView<const FMVVMConstFieldVariant> PropertyPath);
 
 private:
diff --git a/Engine/Plugins/Runtime/ModelViewViewModel/Source/ModelViewViewModelBlueprint/Public/Bindings/MVVMCompiledBindingLibraryCompiler.h b/Engine/Plugins/Runtime/ModelViewViewModel/Source/ModelViewViewModelBlueprint/Public/Bindings/MVVMCompiledBindingLibraryCompiler.h
--- a/Engine/Plugins/Runtime/ModelViewViewModel/Source/ModelViewViewModelBlueprint/Public/Bindings/MVVMCompiledBindingLibraryCompiler.h
+++ b/Engine/Plugins/Runtime/ModelViewViewModel/Source/ModelViewViewModelBlueprint/Public/Bindings/MVVMCompiledBindingLibraryCompiler.h
@@ -63,7 +63,6 @@ public:
 		int32 Id;
 	};
 
-
 	/** */
 	struct FBindingHandle
 	{
@@ -149,7 +148,7 @@ public:
 	};
 
 public:
-	FCompiledBindingLibraryCompiler();
+	FCompiledBindingLibraryCompiler(UBlueprint* Context);
 
 public:
 	/** */
diff --git a/Engine/Plugins/Runtime/ModelViewViewModel/Source/ModelViewViewModelBlueprint/Public/MVVMDeveloperProjectSettings.h b/Engine/Plugins/Runtime/ModelViewViewModel/Source/ModelViewViewModelBlueprint/Public/MVVMDeveloperProjectSettings.h
--- a/Engine/Plugins/Runtime/ModelViewViewModel/Source/ModelViewViewModelBlueprint/Public/MVVMDeveloperProjectSettings.h
+++ b/Engine/Plugins/Runtime/ModelViewViewModel/Source/ModelViewViewModelBlueprint/Public/MVVMDeveloperProjectSettings.h
@@ -71,9 +71,9 @@ public:
 
 
 	bool PropertyHasFiltering(const FProperty* Property) const;
-	bool IsPropertyAllowed(const FProperty* Property) const;
-	bool IsFunctionAllowed(const UFunction* Function) const;
-	bool IsConversionFunctionAllowed(const UFunction* Function) const;
+	bool IsPropertyAllowed(const UBlueprint* Context, const FProperty* Property) const;
+	bool IsFunctionAllowed(const UBlueprint* Context, const UFunction* Function) const;
+	bool IsConversionFunctionAllowed(const UBlueprint* Context, const UFunction* Function) const;
 
 	bool IsExecutionModeAllowed(EMVVMExecutionMode ExecutionMode) const
 	{
diff --git a/Engine/Plugins/Runtime/ModelViewViewModel/Source/ModelViewViewModelEditor/Private/MVVMEditorSubsystem.cpp b/Engine/Plugins/Runtime/ModelViewViewModel/Source/ModelViewViewModelEditor/Private/MVVMEditorSubsystem.cpp
--- a/Engine/Plugins/Runtime/ModelViewViewModel/Source/ModelViewViewModelEditor/Private/MVVMEditorSubsystem.cpp
+++ b/Engine/Plugins/Runtime/ModelViewViewModel/Source/ModelViewViewModelEditor/Private/MVVMEditorSubsystem.cpp
@@ -698,7 +698,7 @@ bool UMVVMEditorSubsystem::IsValidConversionFunction(const UWidgetBlueprint* Wid
 		return false;
 	}
 
-	return GetDefault<UMVVMDeveloperProjectSettings>()->IsConversionFunctionAllowed(Function);
+	return GetDefault<UMVVMDeveloperProjectSettings>()->IsConversionFunctionAllowed(WidgetBlueprint, Function);
 }
 
 bool UMVVMEditorSubsystem::IsSimpleConversionFunction(const UFunction* Function) const
diff --git a/Engine/Plugins/Runtime/ModelViewViewModel/Source/ModelViewViewModelEditor/Private/Widgets/SMVVMViewModelBindingListWidget.cpp b/Engine/Plugins/Runtime/ModelViewViewModel/Source/ModelViewViewModelEditor/Private/Widgets/SMVVMViewModelBindingListWidget.cpp
--- a/Engine/Plugins/Runtime/ModelViewViewModel/Source/ModelViewViewModelEditor/Private/Widgets/SMVVMViewModelBindingListWidget.cpp
+++ b/Engine/Plugins/Runtime/ModelViewViewModel/Source/ModelViewViewModelEditor/Private/Widgets/SMVVMViewModelBindingListWidget.cpp
@@ -26,10 +26,10 @@ namespace UE::MVVM
 
 namespace Private
 {
-	TOptional<FFieldVariant> PassFilter(const FMVVMAvailableBinding& Binding, const FMVVMFieldVariant& FieldVariant, EFieldVisibility FieldVisibilityFlags, const FProperty* AssignableTo, bool bDoObjectProperty);
-	TOptional<FFieldVariant> PassFilter(const FMVVMAvailableBinding& Binding, const UStruct* Struct, EFieldVisibility FieldVisibilityFlags, const FProperty* AssignableTo, bool bDoObjectProperty);
+	TOptional<FFieldVariant> PassFilter(const UBlueprint* Blueprint, const FMVVMAvailableBinding& Binding, const FMVVMFieldVariant& FieldVariant, EFieldVisibility FieldVisibilityFlags, const FProperty* AssignableTo, bool bDoObjectProperty);
+	TOptional<FFieldVariant> PassFilter(const UBlueprint* Blueprint, const FMVVMAvailableBinding& Binding, const UStruct* Struct, EFieldVisibility FieldVisibilityFlags, const FProperty* AssignableTo, bool bDoObjectProperty);
 
-	TOptional<FFieldVariant> PassFilter(const FMVVMAvailableBinding& Binding, const UStruct* Struct, EFieldVisibility FieldVisibilityFlags, const FProperty* AssignableTo, bool bDoObjectProperty)
+	TOptional<FFieldVariant> PassFilter(const UBlueprint* Blueprint, const FMVVMAvailableBinding& Binding, const UStruct* Struct, EFieldVisibility FieldVisibilityFlags, const FProperty* AssignableTo, bool bDoObjectProperty)
 	{
 		if (EnumHasAllFlags(FieldVisibilityFlags, EFieldVisibility::Readable) && !Binding.IsReadable())
 		{
@@ -42,10 +42,10 @@ namespace Private
 		}
 
 		FMVVMFieldVariant FieldVariant = BindingHelper::FindFieldByName(Struct, Binding.GetBindingName());
-		return PassFilter(Binding, FieldVariant, FieldVisibilityFlags, AssignableTo, bDoObjectProperty);
+		return PassFilter(Blueprint, Binding, FieldVariant, FieldVisibilityFlags, AssignableTo, bDoObjectProperty);
 	}
 
-	TOptional<FFieldVariant> PassFilter(const FMVVMAvailableBinding& Binding, const FMVVMFieldVariant& FieldVariant, EFieldVisibility FieldVisibilityFlags, const FProperty* AssignableTo, bool bDoObjectProperty)
+	TOptional<FFieldVariant> PassFilter(const UBlueprint* Blueprint, const FMVVMAvailableBinding& Binding, const FMVVMFieldVariant& FieldVariant, EFieldVisibility FieldVisibilityFlags, const FProperty* AssignableTo, bool bDoObjectProperty)
 	{
 		if (ensure(!FieldVariant.IsEmpty()))
 		{
@@ -93,7 +93,7 @@ namespace Private
 					}
 				}
 
-				if (!GetDefault<UMVVMDeveloperProjectSettings>()->IsFunctionAllowed(Function))
+				if (!GetDefault<UMVVMDeveloperProjectSettings>()->IsFunctionAllowed(Blueprint, Function))
 				{
 					return TOptional<FFieldVariant>();
 				}
@@ -144,7 +144,7 @@ namespace Private
 					}
 				}
 
-				if (!GetDefault<UMVVMDeveloperProjectSettings>()->IsPropertyAllowed(Property))
+				if (!GetDefault<UMVVMDeveloperProjectSettings>()->IsPropertyAllowed(Blueprint, Property))
 				{
 					return TOptional<FFieldVariant>();
 				}
@@ -187,11 +187,11 @@ TArray<FFieldVariant> FFieldIterator_Bindable::GetFields(const UStruct* Struct)
 			TOptional<FFieldVariant> PassResult;
 			if (FilterFlags == EFilterFlag::All)
 			{
-				PassResult = Private::PassFilter(Value, Struct, FieldVisibilityFlags, AssignableTo, false);
+				PassResult = Private::PassFilter(WidgetBlueprint.Get(), Value, Struct, FieldVisibilityFlags, AssignableTo, false);
 			}
 			else
 			{
-				PassResult = Private::PassFilter(Value, Struct, EFieldVisibility::None, nullptr, false);
+				PassResult = Private::PassFilter(WidgetBlueprint.Get(), Value, Struct, EFieldVisibility::None, nullptr, false);
 			}
 			if (PassResult.IsSet())
 			{
@@ -539,11 +539,11 @@ FMVVMBlueprintPropertyPath SSourceBindingList::CreateBlueprintPropertyPath(SProp
 
 					if (FilterFlags == EFilterFlag::All)
 					{
-						bPassFilter = Private::PassFilter(Binding, OwnerClass, FieldIterator->GetFieldVisibilityFlags(), FieldIterator->GetAssignableTo(), true).IsSet();
+						bPassFilter = Private::PassFilter(WidgetBlueprint.Get(), Binding, OwnerClass, FieldIterator->GetFieldVisibilityFlags(), FieldIterator->GetAssignableTo(), true).IsSet();
 					}
 					else
 					{
-						bPassFilter = Private::PassFilter(Binding, OwnerClass, EFieldVisibility::None, nullptr, true).IsSet();
+						bPassFilter = Private::PassFilter(WidgetBlueprint.Get(), Binding, OwnerClass, EFieldVisibility::None, nullptr, true).IsSet();
 					}
 				}
 				break;
@@ -587,11 +587,11 @@ FMVVMBlueprintPropertyPath SSourceBindingList::CreateBlueprintPropertyPath(SProp
 
 			if (FilterFlags == EFilterFlag::All)
 			{
-				bPassFilter = Private::PassFilter(Binding, AccessorClass, FieldIterator->GetFieldVisibilityFlags(), FieldIterator->GetAssignableTo(), true).IsSet();
+				bPassFilter = Private::PassFilter(WidgetBlueprint.Get(), Binding, AccessorClass, FieldIterator->GetFieldVisibilityFlags(), FieldIterator->GetAssignableTo(), true).IsSet();
 			}
 			else
 			{
-				bPassFilter = Private::PassFilter(Binding, AccessorClass, EFieldVisibility::None, nullptr, true).IsSet();
+				bPassFilter = Private::PassFilter(WidgetBlueprint.Get(), Binding, AccessorClass, EFieldVisibility::None, nullptr, true).IsSet();
 			}
 			if (bPassFilter)
 			{
