MVVM: Instead of saving the name of the root widget for self reference, use a different flag. When the Blueprint asset is renamed the root widget is also renamed and self reference do not work anymore. Redirect the previously saved reference.
IsEmpty is renamed to IsValid. It used to be that we wanted to know if there was a source and a path but now the IsEmpty is confusing.
Fix Widget rename for conversion function pins and event pins.

diff --git a/Engine/Plugins/Runtime/ModelViewViewModel/Source/ModelViewViewModelBlueprint/Private/Bindings/MVVMConversionFunctionHelper.cpp b/Engine/Plugins/Runtime/ModelViewViewModel/Source/ModelViewViewModelBlueprint/Private/Bindings/MVVMConversionFunctionHelper.cpp
--- a/Engine/Plugins/Runtime/ModelViewViewModel/Source/ModelViewViewModelBlueprint/Private/Bindings/MVVMConversionFunctionHelper.cpp
+++ b/Engine/Plugins/Runtime/ModelViewViewModel/Source/ModelViewViewModelBlueprint/Private/Bindings/MVVMConversionFunctionHelper.cpp
@@ -369,6 +369,7 @@ namespace Private
 				else
 				{
 					ResultPath.SetWidgetName(Member.GetMemberName());
+					ensure(Member.GetMemberName() != Blueprint->GetFName());
 				}
 			}
 			else
@@ -378,6 +379,7 @@ namespace Private
 					if (Property->PropertyClass->IsChildOf<UWidget>() || Property->PropertyClass->IsChildOf<UBlueprint>())
 					{
 						ResultPath.SetWidgetName(Property->GetFName());
+						ensure(Property->GetFName() != Blueprint->GetFName());
 					}
 					else if (Property->PropertyClass->ImplementsInterface(UNotifyFieldValueChanged::StaticClass()))
 					{
@@ -465,7 +467,7 @@ namespace Private
 			{
 				if (bFirst)
 				{
-					ResultPath.SetWidgetName(Blueprint->GetFName());
+					ResultPath.SetSelfContext();
 				}
 				else
 				{
@@ -519,7 +521,7 @@ void SetPropertyPathForPin(const UBlueprint* Blueprint, const FMVVMBlueprintProp
 	}
 
 	// Add new nodes
-	if (!Path.IsEmpty())
+	if (Path.IsValid())
 	{
 		const int32 ArgumentIndex = ConversionNode->Pins.IndexOfByPredicate([PathPin](const UEdGraphPin* Other){ return Other == PathPin; });
 		if (!ensure(ConversionNode->Pins.IsValidIndex(ArgumentIndex)))
@@ -539,16 +541,23 @@ void SetPropertyPathForPin(const UBlueprint* Blueprint, const FMVVMBlueprintProp
 			const FProperty* RootProperty = nullptr;
 			bool bCreateSelfNodeForRootProperty = false;
 			{
-				if (Path.IsFromWidget())
+				switch (Path.GetSource(Blueprint))
 				{
-					RootProperty = Blueprint->SkeletonGeneratedClass->FindPropertyByName(Path.GetWidgetName());
-					bCreateSelfNodeForRootProperty = Blueprint->GetFName() == Path.GetWidgetName();
-				}
-				else if (Path.IsFromViewModel())
+				case EMVVMBlueprintFieldPathSource::SelfContext:
+					bCreateSelfNodeForRootProperty = true;
+					break;
+				case EMVVMBlueprintFieldPathSource::ViewModel:
 				{
 					UMVVMBlueprintView* View = Private::GetView(Blueprint);
-					const FMVVMBlueprintViewModelContext* Context = View->FindViewModel(Path.GetViewModelId());
-					RootProperty = Blueprint->SkeletonGeneratedClass->FindPropertyByName(Context->GetViewModelName());
+					const FMVVMBlueprintViewModelContext* Context = View ? View->FindViewModel(Path.GetViewModelId()) : nullptr;
+					RootProperty = Context ? Blueprint->SkeletonGeneratedClass->FindPropertyByName(Context->GetViewModelName()) : nullptr;
+					break;
+				}
+				case EMVVMBlueprintFieldPathSource::Widget:
+					RootProperty = Blueprint->SkeletonGeneratedClass->FindPropertyByName(Path.GetWidgetName());
+					break;
+				default:
+					check(false);
 				}
 			}
 
@@ -791,7 +800,7 @@ TMap<FName, FMVVMBlueprintPropertyPath> GetAllArgumentPropertyPaths(const UBluep
 	for (const UEdGraphPin* Pin : FunctionNode->GetAllPins())
 	{
 		FMVVMBlueprintPropertyPath Path = Private::GetPropertyPathForPin(Blueprint, Pin, bSkipResolve);
-		if (!Path.IsEmpty())
+		if (Path.IsValid())
 		{
 			Paths.Add(Pin->PinName, Path);
 		}
diff --git a/Engine/Plugins/Runtime/ModelViewViewModel/Source/ModelViewViewModelBlueprint/Private/MVVMBlueprintPin.cpp b/Engine/Plugins/Runtime/ModelViewViewModel/Source/ModelViewViewModelBlueprint/Private/MVVMBlueprintPin.cpp
--- a/Engine/Plugins/Runtime/ModelViewViewModel/Source/ModelViewViewModelBlueprint/Private/MVVMBlueprintPin.cpp
+++ b/Engine/Plugins/Runtime/ModelViewViewModel/Source/ModelViewViewModelBlueprint/Private/MVVMBlueprintPin.cpp
@@ -7,6 +7,11 @@
 
 #include UE_INLINE_GENERATED_CPP_BY_NAME(MVVMBlueprintPin)
 
+FMVVMBlueprintPin::FMVVMBlueprintPin(FName InPinName)
+	: PinName(InPinName)
+{
+}
+
 FMVVMBlueprintPin FMVVMBlueprintPin::CreateFromPin(const UBlueprint* Blueprint, const UEdGraphPin* Pin)
 {
 	FMVVMBlueprintPin Result;
@@ -44,7 +49,7 @@ void FMVVMBlueprintPin::SetPath(const FMVVMBlueprintPropertyPath& Value)
 
 FString FMVVMBlueprintPin::GetValueAsString(const UClass* SelfContext) const
 {
-	if (!Path.IsEmpty())
+	if (Path.IsValid())
 	{
 		return Path.GetPropertyPath(SelfContext);
 	}
diff --git a/Engine/Plugins/Runtime/ModelViewViewModel/Source/ModelViewViewModelBlueprint/Private/MVVMBlueprintView.cpp b/Engine/Plugins/Runtime/ModelViewViewModel/Source/ModelViewViewModelBlueprint/Private/MVVMBlueprintView.cpp
--- a/Engine/Plugins/Runtime/ModelViewViewModel/Source/ModelViewViewModelBlueprint/Private/MVVMBlueprintView.cpp
+++ b/Engine/Plugins/Runtime/ModelViewViewModel/Source/ModelViewViewModelBlueprint/Private/MVVMBlueprintView.cpp
@@ -178,21 +178,6 @@ void UMVVMBlueprintView::RemoveBinding(const FMVVMBlueprintViewBinding* Binding)
 	RemoveBindingAt(Index);
 }
 
-FMVVMBlueprintViewBinding& UMVVMBlueprintView::AddBinding(const UWidget* Widget, const FProperty* Property)
-{
-	FMVVMBlueprintViewBinding& NewBinding = Bindings.AddDefaulted_GetRef();
-	NewBinding.DestinationPath.SetWidgetName(Widget->GetFName());
-	NewBinding.DestinationPath.SetPropertyPath(GetOuterUMVVMWidgetBlueprintExtension_View()->GetWidgetBlueprint(), UE::MVVM::FMVVMConstFieldVariant(Property));
-	NewBinding.BindingId = FGuid::NewGuid();
-
-	OnBindingsAdded.Broadcast();
-	OnBindingsUpdated.Broadcast();
-
-	FBlueprintEditorUtils::MarkBlueprintAsModified(GetOuterUMVVMWidgetBlueprintExtension_View()->GetWidgetBlueprint());
-
-	return NewBinding;
-}
-
 FMVVMBlueprintViewBinding& UMVVMBlueprintView::AddDefaultBinding()
 {
 	FMVVMBlueprintViewBinding& NewBinding = Bindings.AddDefaulted_GetRef();
@@ -305,6 +290,123 @@ void UMVVMBlueprintView::Serialize(FArchive& Ar)
 }
 
 #if WITH_EDITOR
+namespace UE::MVVM::Private
+{
+	template<typename Predicate>
+	void ForEachPropertyPath_Update(UMVVMBlueprintView* BlueprintView, Predicate Pred, bool bGenerateGraph)
+	{
+		UWidgetBlueprint* WidgetBlueprint = BlueprintView->GetOuterUMVVMWidgetBlueprintExtension_View()->GetWidgetBlueprint();
+		auto PredPin = [&Pred](const FMVVMBlueprintPin& Pin) -> TOptional<FMVVMBlueprintPropertyPath>
+		{
+			if (Pin.UsedPathAsValue())
+			{
+				FMVVMBlueprintPropertyPath NewPinPath = Pin.GetPath();
+				if (Pred(NewPinPath))
+				{
+					return NewPinPath;
+				}
+			}
+			return TOptional<FMVVMBlueprintPropertyPath>();
+		};
+
+		for (FMVVMBlueprintViewBinding& Binding : BlueprintView->GetBindings())
+		{
+			Pred(Binding.SourcePath);
+			Pred(Binding.DestinationPath);
+			if (Binding.Conversion.DestinationToSourceConversion)
+			{
+				Binding.Conversion.DestinationToSourceConversion->ConditionalPostLoad();
+				for (const FMVVMBlueprintPin& Pin : Binding.Conversion.DestinationToSourceConversion->GetPins())
+				{
+					TOptional<FMVVMBlueprintPropertyPath> NewPath = PredPin(Pin);
+					if (NewPath.IsSet())
+					{
+						if (bGenerateGraph)
+						{
+							Binding.Conversion.DestinationToSourceConversion->SetGraphPin(WidgetBlueprint, Pin.GetName(), NewPath.GetValue());
+						}
+						else
+						{
+							const_cast<FMVVMBlueprintPin&>(Pin).SetPath(NewPath.GetValue());
+						}
+					}
+				}
+			}
+			if (Binding.Conversion.SourceToDestinationConversion)
+			{
+				Binding.Conversion.SourceToDestinationConversion->ConditionalPostLoad();
+				for (const FMVVMBlueprintPin& Pin : Binding.Conversion.SourceToDestinationConversion->GetPins())
+				{
+					TOptional<FMVVMBlueprintPropertyPath> NewPath = PredPin(Pin);
+					if (NewPath.IsSet())
+					{
+						if (bGenerateGraph)
+						{
+							Binding.Conversion.SourceToDestinationConversion->SetGraphPin(WidgetBlueprint, Pin.GetName(), NewPath.GetValue());
+						}
+						else
+						{
+							const_cast<FMVVMBlueprintPin&>(Pin).SetPath(NewPath.GetValue());
+						}
+					}
+				}
+			}
+		}
+
+		for (UMVVMBlueprintViewEvent* Event : BlueprintView->GetEvents())
+		{
+			if (Event)
+			{
+				auto PredEventPath = [Event, &Pred, bGenerateGraph](const FMVVMBlueprintPropertyPath& PropertyPath, bool bEventPath)
+				{
+					FMVVMBlueprintPropertyPath NewPropertyPath = PropertyPath;
+					if (Pred(NewPropertyPath))
+					{
+						if (bGenerateGraph)
+						{
+							if (bEventPath)
+							{
+								Event->SetEventPath(NewPropertyPath);
+							}
+							else
+							{
+								Event->SetDestinationPath(NewPropertyPath);
+							}
+						}
+						else
+						{
+							const_cast<FMVVMBlueprintPropertyPath&>(PropertyPath) = NewPropertyPath;
+						}
+					}
+				};
+
+				TArray<TTuple<FName, FMVVMBlueprintPropertyPath>> NewPins;
+				for (const FMVVMBlueprintPin& Pin : Event->GetPins())
+				{
+					TOptional<FMVVMBlueprintPropertyPath> NewPath = PredPin(Pin);
+					if (NewPath.IsSet())
+					{
+						NewPins.Emplace(Pin.GetName(), MoveTemp(NewPath.GetValue()));
+					}
+				}
+				PredEventPath(Event->GetEventPath(), true);
+				PredEventPath(Event->GetDestinationPath(), false);
+				for (TTuple<FName, FMVVMBlueprintPropertyPath>& Pin : NewPins)
+				{
+					if (bGenerateGraph)
+					{
+						Event->SetPinPath(Pin.Get<0>(), Pin.Get<1>());
+					}
+					else
+					{
+						Event->SetPinPathNoGraphGeneration(Pin.Get<0>(), Pin.Get<1>());
+					}
+				}
+			}
+		}
+	}
+}
+
 void UMVVMBlueprintView::PostLoad()
 {
 	Super::PostLoad();
@@ -341,6 +443,17 @@ void UMVVMBlueprintView::PostLoad()
 	{
 		Binding.Conversion.DeprecateViewConversionFunction(GetOuterUMVVMWidgetBlueprintExtension_View()->GetWidgetBlueprint());
 	}
+
+	if (GetLinkerCustomVersion(FFortniteMainBranchObjectVersion::GUID) < FFortniteMainBranchObjectVersion::MVVMPropertyPathSelf)
+	{
+		const UWidgetBlueprint* ThisWidgetBlueprint = GetOuterUMVVMWidgetBlueprintExtension_View()->GetWidgetBlueprint();
+		auto DeprecateSelfPath = [ThisWidgetBlueprint](FMVVMBlueprintPropertyPath& Path) -> bool
+		{
+			Path.GetSource(ThisWidgetBlueprint);
+			return true;
+		};
+		UE::MVVM::Private::ForEachPropertyPath_Update(this, DeprecateSelfPath, false);
+	}
 }
 
 void UMVVMBlueprintView::PreSave(FObjectPreSaveContext Context)
@@ -415,17 +528,25 @@ void UMVVMBlueprintView::AddAssetTags(TArray<FAssetRegistryTag>& OutTags) const
 void UMVVMBlueprintView::WidgetRenamed(FName OldObjectName, FName NewObjectName)
 {
 	bool bRenamed = false;
-	for (FMVVMBlueprintViewBinding& Binding : Bindings)
+
+	UWidgetBlueprint* WidgetBlueprint = GetOuterUMVVMWidgetBlueprintExtension_View()->GetWidgetBlueprint();
+
+	auto RenameWidget = [WidgetBlueprint , &bRenamed, OldObjectName, NewObjectName](FMVVMBlueprintPropertyPath& PropertyPath) -> bool
 	{
-		if (Binding.DestinationPath.GetWidgetName() == OldObjectName)
+		if (PropertyPath.GetSource(WidgetBlueprint) == EMVVMBlueprintFieldPathSource::Widget && PropertyPath.GetWidgetName() == OldObjectName)
 		{
-			Binding.DestinationPath.SetWidgetName(NewObjectName);
+			PropertyPath.SetWidgetName(NewObjectName);
 			bRenamed = true;
+			return true;
 		}
-	}
+		return false;
+	};
+
+	UE::MVVM::Private::ForEachPropertyPath_Update(this, RenameWidget, true);
 
 	if (bRenamed)
 	{
+		Modify();
 		OnBindingsUpdated.Broadcast();
 	}
 }
diff --git a/Engine/Plugins/Runtime/ModelViewViewModel/Source/ModelViewViewModelBlueprint/Private/MVVMBlueprintViewEvent.cpp b/Engine/Plugins/Runtime/ModelViewViewModel/Source/ModelViewViewModelBlueprint/Private/MVVMBlueprintViewEvent.cpp
--- a/Engine/Plugins/Runtime/ModelViewViewModel/Source/ModelViewViewModelBlueprint/Private/MVVMBlueprintViewEvent.cpp
+++ b/Engine/Plugins/Runtime/ModelViewViewModel/Source/ModelViewViewModelBlueprint/Private/MVVMBlueprintViewEvent.cpp
@@ -35,7 +35,7 @@ void UMVVMBlueprintViewEvent::SetEventPath(FMVVMBlueprintPropertyPath InEventPat
 	EventPath = MoveTemp(InEventPath);
 	GraphName = FName();
 
-	if (!EventPath.IsEmpty())
+	if (EventPath.IsValid())
 	{
 		TStringBuilder<256> StringBuilder;
 		StringBuilder << TEXT("__");
@@ -142,6 +142,19 @@ void UMVVMBlueprintViewEvent::SetPinPath(FName PinName, const FMVVMBlueprintProp
 	}
 }
 
+void UMVVMBlueprintViewEvent::SetPinPathNoGraphGeneration(FName PinName, const FMVVMBlueprintPropertyPath& Path)
+{
+	FMVVMBlueprintPin* ViewPin = SavedPins.FindByPredicate([PinName](const FMVVMBlueprintPin& Other) { return PinName == Other.GetName(); });
+	if (!ViewPin)
+	{
+		ViewPin = &SavedPins.Emplace_GetRef(PinName);
+		ViewPin->SetPath(Path);
+	}
+
+	//A property (viewmodel or widget) may not be created yet and the skeletal needs to be recreated.
+	FBlueprintEditorUtils::MarkBlueprintAsStructurallyModified(GetWidgetBlueprintInternal());
+}
+
 UWidgetBlueprint* UMVVMBlueprintViewEvent::GetWidgetBlueprintInternal() const
 {
 	return GetOuterUMVVMBlueprintView()->GetOuterUMVVMWidgetBlueprintExtension_View()->GetWidgetBlueprint();
@@ -159,7 +172,7 @@ const UFunction* UMVVMBlueprintViewEvent::GetEventSignature() const
 
 const UFunction* UMVVMBlueprintViewEvent::GetEventSignature(const UWidgetBlueprint* WidgetBlueprint, const FMVVMBlueprintPropertyPath& PropertyPath)
 {
-	if (!PropertyPath.IsEmpty() && PropertyPath.GetFieldPaths().Num() > 0)
+	if (PropertyPath.IsValid() && PropertyPath.GetFieldPaths().Num() > 0)
 	{
 		const FMVVMBlueprintFieldPath& LastPath = PropertyPath.GetFieldPaths().Last();
 		UE::MVVM::FMVVMConstFieldVariant LastField = LastPath.GetField(WidgetBlueprint->SkeletonGeneratedClass);
@@ -176,7 +189,7 @@ const UFunction* UMVVMBlueprintViewEvent::GetEventSignature(const UWidgetBluepri
 
 UEdGraph* UMVVMBlueprintViewEvent::CreateWrapperGraphInternal()
 {
-	if (GraphName.IsNone() || DestinationPath.IsEmpty() || EventPath.IsEmpty())
+	if (GraphName.IsNone() || !DestinationPath.IsValid() || !EventPath.IsValid())
 	{
 		return nullptr;
 	}
diff --git a/Engine/Plugins/Runtime/ModelViewViewModel/Source/ModelViewViewModelBlueprint/Private/MVVMPropertyPath.cpp b/Engine/Plugins/Runtime/ModelViewViewModel/Source/ModelViewViewModelBlueprint/Private/MVVMPropertyPath.cpp
--- a/Engine/Plugins/Runtime/ModelViewViewModel/Source/ModelViewViewModelBlueprint/Private/MVVMPropertyPath.cpp
+++ b/Engine/Plugins/Runtime/ModelViewViewModel/Source/ModelViewViewModelBlueprint/Private/MVVMPropertyPath.cpp
@@ -322,21 +322,18 @@ FText GetViewModelDisplayName(const UWidgetBlueprint* WidgetBlueprint, FGuid Id,
 
 FText GetRootName(const UWidgetBlueprint* WidgetBlueprint, const FMVVMBlueprintPropertyPath& PropertyPath, bool bUseDisplayName, bool bIncludeMetaData)
 {
-	if (PropertyPath.IsFromWidget())
-	{
-		if (WidgetBlueprint->GetFName() == PropertyPath.GetWidgetName())
-		{
-			return bUseDisplayName ? LOCTEXT("Self", "Self") : FText::FromString(WidgetBlueprint->GetFriendlyName());
-		}
-		else
-		{
-			return GetWidgetDisplayName(WidgetBlueprint, PropertyPath.GetWidgetName(), bUseDisplayName, bIncludeMetaData);
-		}
-	}
-	else if (PropertyPath.IsFromViewModel())
+	switch (PropertyPath.GetSource(WidgetBlueprint))
 	{
+	case EMVVMBlueprintFieldPathSource::SelfContext:
+		return bUseDisplayName ? LOCTEXT("Self", "Self") : FText::FromString(WidgetBlueprint->GetFriendlyName());
+
+	case EMVVMBlueprintFieldPathSource::ViewModel:
 		return GetViewModelDisplayName(WidgetBlueprint, PropertyPath.GetViewModelId(), bUseDisplayName);
+
+	case EMVVMBlueprintFieldPathSource::Widget:
+		return GetWidgetDisplayName(WidgetBlueprint, PropertyPath.GetWidgetName(), bUseDisplayName, bIncludeMetaData);
 	}
+
 	return FText::GetEmpty();
 }
 }//namespace
@@ -428,4 +425,17 @@ FString FMVVMBlueprintPropertyPath::ToString(const UWidgetBlueprint* WidgetBluep
 	return Builder.ToString();
 }
 
+void FMVVMBlueprintPropertyPath::DeprecationUpdateSource(const UBlueprint* InContext)
+{
+	if (ContextId.IsValid())
+	{
+		Source = EMVVMBlueprintFieldPathSource::ViewModel;
+	}
+	else if (!WidgetName.IsNone())
+	{
+		Source = (InContext && InContext->GetFName() == WidgetName) ? EMVVMBlueprintFieldPathSource::SelfContext : EMVVMBlueprintFieldPathSource::Widget;
+	}
+	bDeprecatedSource = true;
+}
+
 #undef LOCTEXT_NAMESPACE
diff --git a/Engine/Plugins/Runtime/ModelViewViewModel/Source/ModelViewViewModelBlueprint/Private/MVVMViewBlueprintCompiler.cpp b/Engine/Plugins/Runtime/ModelViewViewModel/Source/ModelViewViewModelBlueprint/Private/MVVMViewBlueprintCompiler.cpp
--- a/Engine/Plugins/Runtime/ModelViewViewModel/Source/ModelViewViewModelBlueprint/Private/MVVMViewBlueprintCompiler.cpp
+++ b/Engine/Plugins/Runtime/ModelViewViewModel/Source/ModelViewViewModelBlueprint/Private/MVVMViewBlueprintCompiler.cpp
@@ -52,14 +52,21 @@ static const FText PropertyPathIsInvalidFormat = LOCTEXT("PropertyPathIsInvalid"
 
 FString PropertyPathToString(const UClass* InSelfContext, const UMVVMBlueprintView* BlueprintView, const FMVVMBlueprintPropertyPath& PropertyPath)
 {
-	if (PropertyPath.IsEmpty())
+	if (!PropertyPath.IsValid())
 	{
 		return FString();
 	}
 
 	TStringBuilder<512> Result;
-	if (PropertyPath.IsFromViewModel())
+	switch (PropertyPath.GetSource(BlueprintView->GetOuterUMVVMWidgetBlueprintExtension_View()->GetOuterUWidgetBlueprint()))
 	{
+	case EMVVMBlueprintFieldPathSource::SelfContext:
+		Result << InSelfContext->ClassGeneratedBy->GetName();
+		break;
+	case EMVVMBlueprintFieldPathSource::Widget:
+		Result << PropertyPath.GetWidgetName();
+		break;
+	case EMVVMBlueprintFieldPathSource::ViewModel:
 		if (const FMVVMBlueprintViewModelContext* SourceViewModelContext = BlueprintView->FindViewModel(PropertyPath.GetViewModelId()))
 		{
 			Result << SourceViewModelContext->GetViewModelName();
@@ -68,14 +75,10 @@ FString PropertyPathToString(const UClass* InSelfContext, const UMVVMBlueprintVi
 		{
 			Result << TEXT("<none>");
 		}
-	}
-	else if (PropertyPath.IsFromWidget())
-	{
-		Result << PropertyPath.GetWidgetName();
-	}
-	else
-	{
+		break;
+	default:
 		Result << TEXT("<none>");
+		break;
 	}
 
 	FString BasePropertyPath = PropertyPath.GetPropertyPath(InSelfContext);
@@ -535,14 +538,12 @@ void FMVVMViewBlueprintCompiler::CreateSourceLists(const FWidgetBlueprintCompile
 
 	auto GenerateCompilerSourceContext = [Self = this, BlueprintView, DefaultWidgetCategory, Class = Context.GetGeneratedClass(), &ViewModelGuids, &WidgetSources](const FMVVMBlueprintPropertyPath& PropertyPath) -> TValueOrError<void, FText>
 	{
-		if (PropertyPath.IsFromWidget())
+		switch (PropertyPath.GetSource(Self->WidgetBlueprintCompilerContext.WidgetBlueprint()))
+		{
+		case EMVVMBlueprintFieldPathSource::SelfContext:
+			return MakeValue();
+		case EMVVMBlueprintFieldPathSource::Widget:
 		{
-			if (PropertyPath.GetWidgetName() == Class->ClassGeneratedBy->GetFName())
-			{
-				// it's the userwidget
-				return MakeValue();
-			}
-
 			// If the widget doesn't have a property, add one automatically.
 			if (!WidgetSources.Contains(PropertyPath.GetWidgetName()))
 			{
@@ -564,8 +565,9 @@ void FMVVMViewBlueprintCompiler::CreateSourceLists(const FWidgetBlueprintCompile
 				SourceVariable.bPrivate = false;
 				Self->CompilerUserWidgetPropertyContexts.Emplace(MoveTemp(SourceVariable));
 			}
+			break;
 		}
-		else if (PropertyPath.IsFromViewModel())
+		case EMVVMBlueprintFieldPathSource::ViewModel:
 		{
 			const FMVVMBlueprintViewModelContext* SourceViewModelContext = BlueprintView->FindViewModel(PropertyPath.GetViewModelId());
 			if (SourceViewModelContext == nullptr)
@@ -577,9 +579,9 @@ void FMVVMViewBlueprintCompiler::CreateSourceLists(const FWidgetBlueprintCompile
 			{
 				return MakeError(FText::Format(LOCTEXT("BindingViewModelInvalid", "Viewmodel {0} {1} was invalid."), SourceViewModelContext->GetDisplayName(), GetViewModelIdText(PropertyPath)));
 			}
+			break;
 		}
-		else
-		{
+		default:
 			return MakeError(LOCTEXT("SourcePathNotSet", "A source path is required, but not set."));
 		}
 		return MakeValue();
@@ -1220,12 +1222,13 @@ bool FMVVMViewBlueprintCompiler::PreCompileBindingSources(UWidgetBlueprintGenera
 			ECreateSourcesForConversionFunctionResult ConversionFunctionResult = CreateSourcesForConversionFunction(true);
 			if (ConversionFunctionResult == ECreateSourcesForConversionFunctionResult::Continue)
 			{
-				if (!Binding.SourcePath.IsEmpty())
+				if (Binding.SourcePath.IsValid())
 				{
-					if (!Binding.DestinationPath.IsEmpty())
+					if (Binding.DestinationPath.IsValid())
 					{
 						AddWarningForPropertyWithMVVMAndLegacyBinding(Binding.DestinationPath);
 					}
+
 					if (!CreateSourceContextForPropertyPath(Binding.SourcePath, true, INDEX_NONE, FName()))
 					{
 						bAreBindingsValid = false;
@@ -1252,9 +1255,9 @@ bool FMVVMViewBlueprintCompiler::PreCompileBindingSources(UWidgetBlueprintGenera
 			ECreateSourcesForConversionFunctionResult ConversionFunctionResult = CreateSourcesForConversionFunction(false);
 			if (ConversionFunctionResult == ECreateSourcesForConversionFunctionResult::Continue)
 			{
-				if (!Binding.DestinationPath.IsEmpty())
+				if (Binding.DestinationPath.IsValid())
 				{
-					if (!Binding.SourcePath.IsEmpty())
+					if (Binding.SourcePath.IsValid())
 					{
 						AddWarningForPropertyWithMVVMAndLegacyBinding(Binding.SourcePath);
 					}
@@ -1647,7 +1650,7 @@ bool FMVVMViewBlueprintCompiler::PreCompileBindings(UWidgetBlueprintGeneratedCla
 				return MakeError();
 			}
 			FMVVMBlueprintPropertyPath DestinationPPropertyPath;
-			DestinationPPropertyPath.SetWidgetName(Class->ClassGeneratedBy->GetFName());
+			DestinationPPropertyPath.SetSelfContext();
 			DestinationPPropertyPath.SetPropertyPath(BlueprintView->GetOuterUMVVMWidgetBlueprintExtension_View()->GetWidgetBlueprint(), UE::MVVM::FMVVMConstFieldVariant(FoundFunction));
 			SetterPath = Self->CreateBindingDestinationPath(BlueprintView, Class, DestinationPPropertyPath);
 		}
@@ -2203,7 +2206,26 @@ bool FMVVMViewBlueprintCompiler::PreCompileEvents(UWidgetBlueprintGeneratedClass
 		}
 
 		FName SourceName;
-		if (EventPtr->GetEventPath().IsFromViewModel())
+		switch (EventPtr->GetEventPath().GetSource(WidgetBlueprintCompilerContext.WidgetBlueprint()))
+		{
+		case EMVVMBlueprintFieldPathSource::SelfContext:
+			SourceName = WidgetBlueprintCompilerContext.WidgetBlueprint()->GetFName();
+			break;
+		case EMVVMBlueprintFieldPathSource::Widget:
+		{
+			FName WidgetName = EventPtr->GetEventPath().GetWidgetName();
+			checkf(!WidgetName.IsNone(), TEXT("The destination should have been checked and set bAreSourceContextsValid."));
+			const bool bSourceIsUserWidget = WidgetName == Class->ClassGeneratedBy->GetFName();
+			ensure(!bSourceIsUserWidget);
+			
+			const int32 VariableContextIndex = CompilerUserWidgetPropertyContexts.IndexOfByPredicate([WidgetName](const FCompilerUserWidgetPropertyContext& Other) { return Other.PropertyName == WidgetName; });
+			if (ensureAlwaysMsgf(VariableContextIndex != INDEX_NONE, TEXT("Could not find source context for destination '%s'"), *WidgetName.ToString()))
+			{
+				SourceName = CompilerUserWidgetPropertyContexts[VariableContextIndex].PropertyName;
+			}
+			break;
+		}
+		case EMVVMBlueprintFieldPathSource::ViewModel:
 		{
 			const FMVVMBlueprintViewModelContext* SourceViewModelContext = BlueprintView->FindViewModel(EventPtr->GetEventPath().GetViewModelId());
 			check(SourceViewModelContext);
@@ -2211,20 +2233,10 @@ bool FMVVMViewBlueprintCompiler::PreCompileEvents(UWidgetBlueprintGeneratedClass
 			const int32 DestinationVariableContextIndex = CompilerUserWidgetPropertyContexts.IndexOfByPredicate([ViewModelName](const FCompilerUserWidgetPropertyContext& Other) { return Other.PropertyName == ViewModelName; });
 			check(DestinationVariableContextIndex != INDEX_NONE);
 			SourceName = CompilerUserWidgetPropertyContexts[DestinationVariableContextIndex].PropertyName;
+			break;
 		}
-		else if (EventPtr->GetEventPath().IsFromWidget())
-		{
-			FName WidgetName = EventPtr->GetEventPath().GetWidgetName();
-			checkf(!WidgetName.IsNone(), TEXT("The destination should have been checked and set bAreSourceContextsValid."));
-			const bool bSourceIsUserWidget = WidgetName == Class->ClassGeneratedBy->GetFName();
-			if (!bSourceIsUserWidget)
-			{
-				const int32 VariableContextIndex = CompilerUserWidgetPropertyContexts.IndexOfByPredicate([WidgetName](const FCompilerUserWidgetPropertyContext& Other) { return Other.PropertyName == WidgetName; });
-				if (ensureAlwaysMsgf(VariableContextIndex != INDEX_NONE, TEXT("Could not find source context for destination '%s'"), *WidgetName.ToString()))
-				{
-					SourceName = CompilerUserWidgetPropertyContexts[VariableContextIndex].PropertyName;
-				}
-			}
+		default:
+			ensureAlwaysMsgf(false, TEXT("An EMVVMBlueprintFieldPathSource case was not checked."));
 		}
 
 		// No need to add the generated function to the field compiler.
@@ -2367,14 +2379,16 @@ const FMVVMViewBlueprintCompiler::FCompilerSourceCreatorContext* FMVVMViewBluepr
 
 TValueOrError<FMVVMViewBlueprintCompiler::FBindingSourceContext, FText> FMVVMViewBlueprintCompiler::CreateBindingSourceContext(const UMVVMBlueprintView* BlueprintView, const UWidgetBlueprintGeneratedClass* Class, const FMVVMBlueprintPropertyPath& PropertyPath, bool bIsOneTimeBinding)
 {
-	if (PropertyPath.IsEmpty())
+	if (!PropertyPath.IsValid())
 	{
 		ensureAlways(false);
 		return MakeError(LOCTEXT("EmptyPropertyPath", "Empty property path found. This is ilegal."));
 	}
 
 	FBindingSourceContext Result;
-	if (PropertyPath.IsFromViewModel())
+	switch (PropertyPath.GetSource(WidgetBlueprintCompilerContext.WidgetBlueprint()))
+	{
+	case EMVVMBlueprintFieldPathSource::ViewModel:
 	{
 		Result.bIsRootWidget = false;
 
@@ -2386,30 +2400,34 @@ TValueOrError<FMVVMViewBlueprintCompiler::FBindingSourceContext, FText> FMVVMVie
 
 		Result.SourceClass = CompilerUserWidgetPropertyContexts[Result.UserWidgetPropertyContextIndex].Class;
 		Result.PropertyPath = CreatePropertyPath(Class, CompilerUserWidgetPropertyContexts[Result.UserWidgetPropertyContextIndex].PropertyName, PropertyPath.GetFields(Class));
+		break;
+	}
+
+	case EMVVMBlueprintFieldPathSource::SelfContext:
+	{
+		Result.bIsRootWidget = true;
+		Result.UserWidgetPropertyContextIndex = INDEX_NONE;
+		Result.SourceClass = const_cast<UWidgetBlueprintGeneratedClass*>(Class);
+		Result.PropertyPath = CreatePropertyPath(Class, FName(), PropertyPath.GetFields(Class));
+		break;
 	}
-	else if (PropertyPath.IsFromWidget())
+	case EMVVMBlueprintFieldPathSource::Widget:
 	{
+		Result.bIsRootWidget = false;
+
 		const FName SourceName = PropertyPath.GetWidgetName();
-		Result.bIsRootWidget = SourceName == Class->ClassGeneratedBy->GetFName();
-		if (Result.bIsRootWidget)
-		{
-			Result.UserWidgetPropertyContextIndex = INDEX_NONE;
-			Result.SourceClass = const_cast<UWidgetBlueprintGeneratedClass*>(Class);
-			Result.PropertyPath = CreatePropertyPath(Class, FName(), PropertyPath.GetFields(Class));
-		}
-		else
+		Result.UserWidgetPropertyContextIndex = CompilerUserWidgetPropertyContexts.IndexOfByPredicate([SourceName](const FCompilerUserWidgetPropertyContext& Other) { return Other.PropertyName == SourceName; });
+		if (!CompilerUserWidgetPropertyContexts.IsValidIndex(Result.UserWidgetPropertyContextIndex))
 		{
-			Result.UserWidgetPropertyContextIndex = CompilerUserWidgetPropertyContexts.IndexOfByPredicate([SourceName](const FCompilerUserWidgetPropertyContext& Other) { return Other.PropertyName == SourceName; });
-			if (!CompilerUserWidgetPropertyContexts.IsValidIndex(Result.UserWidgetPropertyContextIndex))
-			{
-				return MakeError(LOCTEXT("InvalidUserWidgetPropertyContextIndexInternal", "Internal error. UserWidgetPropertyContextIndex is invalid."));
-			}
-			Result.SourceClass = CompilerUserWidgetPropertyContexts[Result.UserWidgetPropertyContextIndex].Class;
-			Result.PropertyPath = CreatePropertyPath(Class, CompilerUserWidgetPropertyContexts[Result.UserWidgetPropertyContextIndex].PropertyName, PropertyPath.GetFields(Class));
+			return MakeError(LOCTEXT("InvalidUserWidgetPropertyContextIndexInternal", "Internal error. UserWidgetPropertyContextIndex is invalid."));
 		}
+
+		Result.SourceClass = CompilerUserWidgetPropertyContexts[Result.UserWidgetPropertyContextIndex].Class;
+		Result.PropertyPath = CreatePropertyPath(Class, CompilerUserWidgetPropertyContexts[Result.UserWidgetPropertyContextIndex].PropertyName, PropertyPath.GetFields(Class));
+		break;
 	}
-	else
-	{
+
+	default:
 		ensureAlwaysMsgf(false, TEXT("Not supported yet."));
 	}
 
@@ -2562,13 +2580,15 @@ TValueOrError<FMVVMViewBlueprintCompiler::FBindingSourceContext, FText> FMVVMVie
 
 TArray<FMVVMConstFieldVariant> FMVVMViewBlueprintCompiler::CreateBindingDestinationPath(const UMVVMBlueprintView* BlueprintView, const UWidgetBlueprintGeneratedClass* Class, const FMVVMBlueprintPropertyPath& PropertyPath) const
 {
-	if (PropertyPath.IsEmpty())
+	if (!PropertyPath.IsValid())
 	{
 		ensureAlwaysMsgf(false, TEXT("Empty property path found. This is legal."));
 		return TArray<FMVVMConstFieldVariant>();
 	}
 
-	if (PropertyPath.IsFromViewModel())
+	switch (PropertyPath.GetSource(WidgetBlueprintCompilerContext.WidgetBlueprint()))
+	{
+	case EMVVMBlueprintFieldPathSource::ViewModel:
 	{
 		const FMVVMBlueprintViewModelContext* SourceViewModelContext = BlueprintView->FindViewModel(PropertyPath.GetViewModelId());
 		check(SourceViewModelContext);
@@ -2578,31 +2598,26 @@ TArray<FMVVMConstFieldVariant> FMVVMViewBlueprintCompiler::CreateBindingDestinat
 
 		return CreatePropertyPath(Class, CompilerUserWidgetPropertyContexts[DestinationVariableContextIndex].PropertyName, PropertyPath.GetFields(Class));
 	}
-	else if (PropertyPath.IsFromWidget())
+	case EMVVMBlueprintFieldPathSource::SelfContext:
+	{
+		return CreatePropertyPath(Class, FName(), PropertyPath.GetFields(Class));
+	}
+	case EMVVMBlueprintFieldPathSource::Widget:
 	{
 		FName DestinationName = PropertyPath.GetWidgetName();
 		checkf(!DestinationName.IsNone(), TEXT("The destination should have been checked and set bAreSourceContextsValid."));
-		const bool bSourceIsUserWidget = DestinationName == Class->ClassGeneratedBy->GetFName();
-		if (bSourceIsUserWidget)
-		{
-			return CreatePropertyPath(Class, FName(), PropertyPath.GetFields(Class));
-		}
-		else
+		const int32 DestinationVariableContextIndex = CompilerUserWidgetPropertyContexts.IndexOfByPredicate([DestinationName](const FCompilerUserWidgetPropertyContext& Other) { return Other.PropertyName == DestinationName; });
+
+		if (ensureAlwaysMsgf(DestinationVariableContextIndex != INDEX_NONE, TEXT("Could not find source context for destination '%s'"), *DestinationName.ToString()))
 		{
-			const int32 DestinationVariableContextIndex = CompilerUserWidgetPropertyContexts.IndexOfByPredicate([DestinationName](const FCompilerUserWidgetPropertyContext& Other) { return Other.PropertyName == DestinationName; });
-			if (ensureAlwaysMsgf(DestinationVariableContextIndex != INDEX_NONE, TEXT("Could not find source context for destination '%s'"), *DestinationName.ToString()))
-			{
-				return CreatePropertyPath(Class, CompilerUserWidgetPropertyContexts[DestinationVariableContextIndex].PropertyName, PropertyPath.GetFields(Class));
-			}
+			return CreatePropertyPath(Class, CompilerUserWidgetPropertyContexts[DestinationVariableContextIndex].PropertyName, PropertyPath.GetFields(Class));
 		}
+		return TArray<FMVVMConstFieldVariant>();
 	}
-	else
-	{
+	default:
 		ensureAlwaysMsgf(false, TEXT("Not supported yet."));
 		return CreatePropertyPath(Class, FName(), PropertyPath.GetFields(Class));
 	}
-
-	return TArray<FMVVMConstFieldVariant>();
 }
 
 TArray<FMVVMConstFieldVariant> FMVVMViewBlueprintCompiler::CreatePropertyPath(const UClass* Class, FName PropertyName, TArray<FMVVMConstFieldVariant> Properties)
diff --git a/Engine/Plugins/Runtime/ModelViewViewModel/Source/ModelViewViewModelBlueprint/Public/MVVMBlueprintPin.h b/Engine/Plugins/Runtime/ModelViewViewModel/Source/ModelViewViewModelBlueprint/Public/MVVMBlueprintPin.h
--- a/Engine/Plugins/Runtime/ModelViewViewModel/Source/ModelViewViewModelBlueprint/Public/MVVMBlueprintPin.h
+++ b/Engine/Plugins/Runtime/ModelViewViewModel/Source/ModelViewViewModelBlueprint/Public/MVVMBlueprintPin.h
@@ -38,6 +38,7 @@ private:
 
 public:
 	FMVVMBlueprintPin() = default;
+	FMVVMBlueprintPin(FName PinName);
 
 	FName GetName() const
 	{
@@ -47,7 +48,7 @@ public:
 	/** Are we using the path. */
 	bool UsedPathAsValue() const
 	{
-		return !Path.IsEmpty();
+		return Path.IsValid();
 	}
 
 	const FMVVMBlueprintPropertyPath& GetPath() const
diff --git a/Engine/Plugins/Runtime/ModelViewViewModel/Source/ModelViewViewModelBlueprint/Public/MVVMBlueprintView.h b/Engine/Plugins/Runtime/ModelViewViewModel/Source/ModelViewViewModelBlueprint/Public/MVVMBlueprintView.h
--- a/Engine/Plugins/Runtime/ModelViewViewModel/Source/ModelViewViewModelBlueprint/Public/MVVMBlueprintView.h
+++ b/Engine/Plugins/Runtime/ModelViewViewModel/Source/ModelViewViewModelBlueprint/Public/MVVMBlueprintView.h
@@ -94,7 +94,6 @@ public:
 	void RemoveBinding(const FMVVMBlueprintViewBinding* Binding);
 	void RemoveBindingAt(int32 Index);
 
-	FMVVMBlueprintViewBinding& AddBinding(const UWidget* Widget, const FProperty* Property);
 	FMVVMBlueprintViewBinding& AddDefaultBinding();
 
 	int32 GetNumBindings() const
diff --git a/Engine/Plugins/Runtime/ModelViewViewModel/Source/ModelViewViewModelBlueprint/Public/MVVMBlueprintViewEvent.h b/Engine/Plugins/Runtime/ModelViewViewModel/Source/ModelViewViewModelBlueprint/Public/MVVMBlueprintViewEvent.h
--- a/Engine/Plugins/Runtime/ModelViewViewModel/Source/ModelViewViewModelBlueprint/Public/MVVMBlueprintViewEvent.h
+++ b/Engine/Plugins/Runtime/ModelViewViewModel/Source/ModelViewViewModelBlueprint/Public/MVVMBlueprintViewEvent.h
@@ -91,6 +91,8 @@ public:
 
 	FMVVMBlueprintPropertyPath GetPinPath(FName PinName) const;
 	void SetPinPath(FName PinName, const FMVVMBlueprintPropertyPath& Path);
+	// To set a pin when loading the asset (no graph generation)
+	void SetPinPathNoGraphGeneration(FName PinName, const FMVVMBlueprintPropertyPath& Path);
 
 public:
 	TArray<FText> GetCompilationMessages(EMessageType InMessageType) const;
diff --git a/Engine/Plugins/Runtime/ModelViewViewModel/Source/ModelViewViewModelBlueprint/Public/MVVMPropertyPath.h b/Engine/Plugins/Runtime/ModelViewViewModel/Source/ModelViewViewModelBlueprint/Public/MVVMPropertyPath.h
--- a/Engine/Plugins/Runtime/ModelViewViewModel/Source/ModelViewViewModelBlueprint/Public/MVVMPropertyPath.h
+++ b/Engine/Plugins/Runtime/ModelViewViewModel/Source/ModelViewViewModelBlueprint/Public/MVVMPropertyPath.h
@@ -73,6 +73,18 @@ private:
 	UE::MVVM::FMVVMConstFieldVariant GetFieldInternal(const UClass* SelfContext) const;
 };
 
+/**
+ * 
+ */
+UENUM()
+enum class EMVVMBlueprintFieldPathSource : uint8
+{
+	None,
+	Widget,
+	ViewModel,
+	SelfContext,
+};
+
 /**
  * Base path to properties for MVVM view models and widgets.
  * 
@@ -94,12 +106,18 @@ private:
 	UPROPERTY(VisibleAnywhere, Category = "MVVM")
 	FGuid ContextId;
 
+	UPROPERTY(VisibleAnywhere, Category = "MVVM")
+	EMVVMBlueprintFieldPathSource Source = EMVVMBlueprintFieldPathSource::None;
+
 #if WITH_EDITORONLY_DATA
 	// Use the Paths. BindingReference and BindingKind are deprecated.
 	UPROPERTY()
 	FMemberReference BindingReference_DEPRECATED;
 	UPROPERTY()
 	EBindingKind BindingKind_DEPRECATED = EBindingKind::Function;
+	// Is the deprecation source done.
+	UPROPERTY()
+	bool bDeprecatedSource = false;
 #endif
 
 public:
@@ -153,15 +171,31 @@ public:
 
 	void ResetSource()
 	{
+		Source = EMVVMBlueprintFieldPathSource::None;
 		ContextId = FGuid();
 		WidgetName = FName();
+		bDeprecatedSource = true;
+	}
+
+	EMVVMBlueprintFieldPathSource GetSource(const UBlueprint* InContext) const
+	{
+		if (!bDeprecatedSource)
+		{
+			const_cast<FMVVMBlueprintPropertyPath*>(this)->DeprecationUpdateSource(InContext);
+		}
+		return Source;
+		
+
+
 	}
 
+	UE_DEPRECATED(5.4, "Use GetSource instead.")
 	bool IsFromWidget() const
 	{
 		return !WidgetName.IsNone();
 	}
 
+	UE_DEPRECATED(5.4, "Use GetSource instead.")
 	bool IsFromViewModel() const
 	{
 		return ContextId.IsValid();
@@ -174,8 +208,9 @@ public:
 
 	void SetViewModelId(FGuid InViewModelId)
 	{
-		WidgetName = FName();
+		ResetSource();
 		ContextId = InViewModelId;
+		Source = EMVVMBlueprintFieldPathSource::ViewModel;
 	}
 
 	FName GetWidgetName() const
@@ -185,20 +220,36 @@ public:
 
 	void SetWidgetName(FName InWidgetName)
 	{
-		ContextId = FGuid();
+		ResetSource();
 		WidgetName = InWidgetName;
+		Source = EMVVMBlueprintFieldPathSource::Widget;
 	}
 
+	void SetSelfContext()
+	{
+		ResetSource();
+		Source = EMVVMBlueprintFieldPathSource::SelfContext;
+	}
+
+	bool IsValid() const
+	{
+		bool bHasValidSource = bDeprecatedSource && Source != EMVVMBlueprintFieldPathSource::None;
+		bool bHasValidDeprecatedSource = !bDeprecatedSource && (!WidgetName.IsNone() || ContextId.IsValid());
+		return bHasValidSource || bHasValidDeprecatedSource;
+	}
+
+	UE_DEPRECATED(5.4, "Use IsValid  instead.")
 	bool IsEmpty() const
 	{
-		return !IsFromWidget() && !IsFromViewModel() && BindingReference_DEPRECATED.GetMemberName() == FName();
+		return !IsValid();
 	}
 
 	bool operator==(const FMVVMBlueprintPropertyPath& Other) const
 	{
-		return WidgetName == Other.WidgetName && 
-			ContextId == Other.ContextId && 
-			Paths == Other.Paths;
+		return WidgetName == Other.WidgetName
+			&& ContextId == Other.ContextId
+			&& Source == Other.Source
+			&& Paths == Other.Paths;
 	}
 
 	bool operator!=(const FMVVMBlueprintPropertyPath& Other) const
@@ -221,6 +272,9 @@ public:
 public:
 	FText ToText(const UWidgetBlueprint* Blueprint, bool bUseDisplayName) const;
 	FString ToString(const UWidgetBlueprint* Blueprint, bool bUseDisplayName, bool bIncludeMetaData) const;
+
+private:
+	MODELVIEWVIEWMODELBLUEPRINT_API void DeprecationUpdateSource(const UBlueprint* InContext);
 };
 
 template<>
diff --git a/Engine/Plugins/Runtime/ModelViewViewModel/Source/ModelViewViewModelEditor/Private/Customizations/MVVMPropertyBindingExtension.cpp b/Engine/Plugins/Runtime/ModelViewViewModel/Source/ModelViewViewModelEditor/Private/Customizations/MVVMPropertyBindingExtension.cpp
--- a/Engine/Plugins/Runtime/ModelViewViewModel/Source/ModelViewViewModelEditor/Private/Customizations/MVVMPropertyBindingExtension.cpp
+++ b/Engine/Plugins/Runtime/ModelViewViewModel/Source/ModelViewViewModelEditor/Private/Customizations/MVVMPropertyBindingExtension.cpp
@@ -102,7 +102,14 @@ static void ExtendBindingsMenu(FMenuBuilder& MenuBuilder, const UWidgetBlueprint
 			DestinationPropertyPath.AppendPropertyPath(WidgetBlueprint, UE::MVVM::FMVVMConstFieldVariant(Field));
 		}
 
-		DestinationPropertyPath.SetWidgetName(Widget->GetFName());
+		if (Widget->GetFName() == WidgetBlueprint->GetFName())
+		{
+			DestinationPropertyPath.SetSelfContext();
+		}
+		else
+		{
+			DestinationPropertyPath.SetWidgetName(Widget->GetFName());
+		}
 		NewBinding.DestinationPath = DestinationPropertyPath;
 
 		NewBinding.BindingType = EMVVMBindingMode::OneWayToDestination;
@@ -312,7 +319,14 @@ IPropertyBindingExtension::EDropResult FMVVMPropertyBindingExtension::OnDrop(con
 			DestinationPropertyPath.AppendPropertyPath(WidgetBlueprint, UE::MVVM::FMVVMConstFieldVariant(Field));
 		}
 
-		DestinationPropertyPath.SetWidgetName(Widget->GetFName());
+		if (Widget->GetFName() == WidgetBlueprint->GetFName())
+		{
+			DestinationPropertyPath.SetSelfContext();
+		}
+		else
+		{
+			DestinationPropertyPath.SetWidgetName(Widget->GetFName());
+		}
 		EditorSubsystem->SetDestinationPathForBinding(WidgetBlueprint, NewBinding, DestinationPropertyPath);
 
 		return EDropResult::HandledContinue;
diff --git a/Engine/Plugins/Runtime/ModelViewViewModel/Source/ModelViewViewModelEditor/Private/MVVMEditorSubsystem.cpp b/Engine/Plugins/Runtime/ModelViewViewModel/Source/ModelViewViewModelEditor/Private/MVVMEditorSubsystem.cpp
--- a/Engine/Plugins/Runtime/ModelViewViewModel/Source/ModelViewViewModelEditor/Private/MVVMEditorSubsystem.cpp
+++ b/Engine/Plugins/Runtime/ModelViewViewModel/Source/ModelViewViewModelEditor/Private/MVVMEditorSubsystem.cpp
@@ -900,11 +900,7 @@ TArray<UE::MVVM::FBindingSource> UMVVMEditorSubsystem::GetBindableWidgets(const
 		if (Bindings.Num() > 0)
 		{
 			// at least one valid property, add it to our list
-			UE::MVVM::FBindingSource Source;
-			Source.Name = WidgetBlueprint->GetFName();
-			Source.DisplayName = FText::FromName(WidgetBlueprint->GetFName());
-			Source.Class = BPClass;
-			Sources.Add(Source);
+			Sources.Add(UE::MVVM::FBindingSource::CreateForBlueprint(WidgetBlueprint));
 		}
 	}
 
@@ -914,11 +910,7 @@ TArray<UE::MVVM::FBindingSource> UMVVMEditorSubsystem::GetBindableWidgets(const
 		if (Bindings.Num() > 0)
 		{
 			// at least one valid property, add it to our list
-			UE::MVVM::FBindingSource Source;
-			Source.Name = Widget->GetFName();
-			Source.DisplayName = Widget->GetLabelText();
-			Source.Class = Widget->GetClass();
-			Sources.Add(Source);
+			Sources.Add(UE::MVVM::FBindingSource::CreateForWidget(WidgetBlueprint, Widget));
 		}
 	}
 
@@ -934,13 +926,9 @@ TArray<UE::MVVM::FBindingSource> UMVVMEditorSubsystem::GetAllViewModels(const UW
 		const TArrayView<const FMVVMBlueprintViewModelContext> ViewModels = View->GetViewModels();
 		Sources.Reserve(ViewModels.Num());
 
-		for (const FMVVMBlueprintViewModelContext& ViewModel : ViewModels)
+		for (const FMVVMBlueprintViewModelContext& ViewModelContext : ViewModels)
 		{
-			UE::MVVM::FBindingSource Source;
-			Source.ViewModelId = ViewModel.GetViewModelId();
-			Source.DisplayName = ViewModel.GetDisplayName();
-			Source.Class = ViewModel.GetViewModelClass();
-			Sources.Add(Source);
+			Sources.Add(UE::MVVM::FBindingSource::CreateForViewModel(WidgetBlueprint, ViewModelContext));
 		}
 	}
 
diff --git a/Engine/Plugins/Runtime/ModelViewViewModel/Source/ModelViewViewModelEditor/Private/Types/MVVMBindingSource.cpp b/Engine/Plugins/Runtime/ModelViewViewModel/Source/ModelViewViewModelEditor/Private/Types/MVVMBindingSource.cpp
--- a/Engine/Plugins/Runtime/ModelViewViewModel/Source/ModelViewViewModelEditor/Private/Types/MVVMBindingSource.cpp
+++ b/Engine/Plugins/Runtime/ModelViewViewModel/Source/ModelViewViewModelEditor/Private/Types/MVVMBindingSource.cpp
@@ -4,15 +4,41 @@
 #include "Blueprint/WidgetTree.h"
 #include "Editor.h"
 #include "MVVMBlueprintView.h"
+#include "MVVMBlueprintViewModelContext.h"
 #include "MVVMPropertyPath.h"
 #include "MVVMEditorSubsystem.h"
 #include "Types/MVVMBindingName.h"
 #include "WidgetBlueprint.h"
 
-FMVVMBindingName UE::MVVM::FBindingSource::ToBindingName(const UWidgetBlueprint* WidgetBlueprint) const
+
+#define LOCTEXT_NAMESPACE "BindingSource"
+
+namespace UE::MVVM
+{
+
+const UClass* FBindingSource::GetClass() const
+{
+	return IsValid() ? Class.Get() : nullptr;
+}
+
+FText FBindingSource::GetDisplayName() const
+{
+	if (!IsValid())
+	{
+		return LOCTEXT("None", "<None>");
+	}
+
+	return DisplayName;
+}
+
+FMVVMBindingName FBindingSource::ToBindingName(const UWidgetBlueprint* WidgetBlueprint) const
 {
-	if (ViewModelId.IsValid())
+	switch (Source)
 	{
+	case EMVVMBlueprintFieldPathSource::SelfContext:
+		return FMVVMBindingName(WidgetBlueprint->GetFName());
+
+	case EMVVMBlueprintFieldPathSource::ViewModel:
 		if (UMVVMBlueprintView* View = GEditor->GetEditorSubsystem<UMVVMEditorSubsystem>()->GetView(WidgetBlueprint))
 		{
 			if (const FMVVMBlueprintViewModelContext* ViewModel = View->FindViewModel(ViewModelId))
@@ -21,25 +47,71 @@ FMVVMBindingName UE::MVVM::FBindingSource::ToBindingName(const UWidgetBlueprint*
 			}
 		}
 		return FMVVMBindingName();
+
+	case EMVVMBlueprintFieldPathSource::Widget:
+		return FMVVMBindingName(WidgetName);
 	}
-	else
+	return FMVVMBindingName();
+}
+
+void FBindingSource::SetSourceTo(FMVVMBlueprintPropertyPath& PropertyPath) const
+{
+	switch(Source)
 	{
-		return FMVVMBindingName(Name);
+	case EMVVMBlueprintFieldPathSource::SelfContext:
+		PropertyPath.SetSelfContext();
+		break;
+	case EMVVMBlueprintFieldPathSource::ViewModel:
+		PropertyPath.SetViewModelId(ViewModelId);
+		break;
+	case EMVVMBlueprintFieldPathSource::Widget:
+		PropertyPath.SetWidgetName(WidgetName);
+		break;
+	default:
+		PropertyPath.ResetSource();
 	}
 }
 
-UE::MVVM::FBindingSource UE::MVVM::FBindingSource::CreateForWidget(const UWidgetBlueprint* WidgetBlueprint, FName WidgetName)
+FBindingSource FBindingSource::CreateForBlueprint(const UWidgetBlueprint* WidgetBlueprint)
 {
-	UE::MVVM::FBindingSource Source;
+	FBindingSource Source;
+
+	Source.Source = EMVVMBlueprintFieldPathSource::SelfContext;
+	Source.DisplayName = FText::FromString(WidgetBlueprint->GetName());
+	Source.Class = WidgetBlueprint->GeneratedClass;
+	
+	return Source;
+}
+
+FBindingSource FBindingSource::CreateForWidget(const UWidgetBlueprint* WidgetBlueprint, const UWidget* Widget)
+{
+	if (Widget->GetFName() == WidgetBlueprint->GetFName())
+	{
+		return CreateForBlueprint(WidgetBlueprint);
+	}
+
+	FBindingSource Source;
 
-	Source.Name = WidgetName;
+	Source.Source = EMVVMBlueprintFieldPathSource::Widget;
+	Source.WidgetName = Widget->GetFName();
+	Source.DisplayName = Widget->GetLabelText();
+	Source.Class = Widget->GetClass();
+	
+	return Source;
+}
 
-	if (Source.Name == WidgetBlueprint->GetFName())
+FBindingSource FBindingSource::CreateForWidget(const UWidgetBlueprint* WidgetBlueprint, FName WidgetName)
+{
+	if (WidgetName == WidgetBlueprint->GetFName())
 	{
-		Source.DisplayName = FText::FromString(WidgetBlueprint->GetName());
-		Source.Class = WidgetBlueprint->GeneratedClass;
+		return CreateForBlueprint(WidgetBlueprint);
 	}
-	else if (UWidget* Widget = WidgetBlueprint->WidgetTree->FindWidget(Source.Name))
+
+	FBindingSource Source;
+
+	Source.Source = EMVVMBlueprintFieldPathSource::Widget;
+	Source.WidgetName = WidgetName;
+	if (UWidget* Widget = WidgetBlueprint->WidgetTree->FindWidget(Source.WidgetName))
 	{
 		Source.DisplayName = Widget->GetLabelText();
 		Source.Class = Widget->GetClass();
@@ -48,9 +120,11 @@ UE::MVVM::FBindingSource UE::MVVM::FBindingSource::CreateForWidget(const UWidget
 	return Source;
 }
 
-UE::MVVM::FBindingSource UE::MVVM::FBindingSource::CreateForViewModel(const UWidgetBlueprint* WidgetBlueprint, FGuid ViewModelId)
+FBindingSource FBindingSource::CreateForViewModel(const UWidgetBlueprint* WidgetBlueprint, FGuid ViewModelId)
 {
-	UE::MVVM::FBindingSource Source;
+	FBindingSource Source;
+
+	Source.Source = EMVVMBlueprintFieldPathSource::ViewModel;
 	Source.ViewModelId = ViewModelId;
 
 	UMVVMEditorSubsystem* Subsystem = GEditor->GetEditorSubsystem<UMVVMEditorSubsystem>();
@@ -66,15 +140,16 @@ UE::MVVM::FBindingSource UE::MVVM::FBindingSource::CreateForViewModel(const UWid
 	return Source;
 }
 
-UE::MVVM::FBindingSource UE::MVVM::FBindingSource::CreateForViewModel(const UWidgetBlueprint* WidgetBlueprint, FName ViewModelName)
+FBindingSource FBindingSource::CreateForViewModel(const UWidgetBlueprint* WidgetBlueprint, FName ViewModelName)
 {
-	UE::MVVM::FBindingSource Source;
+	FBindingSource Source;
 
 	UMVVMEditorSubsystem* Subsystem = GEditor->GetEditorSubsystem<UMVVMEditorSubsystem>();
 	if (UMVVMBlueprintView* View = Subsystem->GetView(WidgetBlueprint))
 	{
 		if (const FMVVMBlueprintViewModelContext* ViewModel = View->FindViewModel(ViewModelName))
 		{
+			Source.Source = EMVVMBlueprintFieldPathSource::ViewModel;
 			Source.ViewModelId = ViewModel->GetViewModelId();
 			Source.DisplayName = ViewModel->GetDisplayName();
 			Source.Class = ViewModel->GetViewModelClass();
@@ -84,16 +159,60 @@ UE::MVVM::FBindingSource UE::MVVM::FBindingSource::CreateForViewModel(const UWid
 	return Source;
 }
 
+FBindingSource FBindingSource::CreateForViewModel(const UWidgetBlueprint* WidgetBlueprint, const FMVVMBlueprintViewModelContext& ViewModelContext)
+{
+	FBindingSource Source;
+
+	Source.Source = EMVVMBlueprintFieldPathSource::ViewModel;
+	Source.ViewModelId = ViewModelContext.GetViewModelId();
+	Source.DisplayName = ViewModelContext.GetDisplayName();
+	Source.Class = ViewModelContext.GetViewModelClass();
+
+	return Source;
+}
+
+FBindingSource FBindingSource::CreateEmptySource(UClass* ViewModel)
+{
+	FBindingSource Source;
+
+	Source.DisplayName = ViewModel->GetDisplayNameText();
+	Source.Class = ViewModel;
+
+	return Source;
+}
 
-UE::MVVM::FBindingSource UE::MVVM::FBindingSource::CreateFromPropertyPath(const UWidgetBlueprint* WidgetBlueprint, const FMVVMBlueprintPropertyPath& Path)
+bool FBindingSource::Matches(const UWidgetBlueprint* WidgetBlueprint, const FMVVMBlueprintPropertyPath& PropertyPath) const
 {
-	if (Path.IsFromViewModel())
+	if (Source == PropertyPath.GetSource(WidgetBlueprint))
 	{
-		return FBindingSource::CreateForViewModel(WidgetBlueprint, Path.GetViewModelId());
+		switch (Source)
+		{
+		case EMVVMBlueprintFieldPathSource::ViewModel:
+			return PropertyPath.GetViewModelId() == ViewModelId;
+		case EMVVMBlueprintFieldPathSource::Widget:
+			return PropertyPath.GetWidgetName() == WidgetName;
+		case EMVVMBlueprintFieldPathSource::SelfContext:
+		default:
+			return true;
+		}
 	}
-	else if (Path.IsFromWidget())
+	return false;
+}
+
+FBindingSource FBindingSource::CreateFromPropertyPath(const UWidgetBlueprint* WidgetBlueprint, const FMVVMBlueprintPropertyPath& Path)
+{
+	switch (Path.GetSource(WidgetBlueprint))
 	{
+	case EMVVMBlueprintFieldPathSource::SelfContext:
+		return FBindingSource::CreateForBlueprint(WidgetBlueprint);
+	case EMVVMBlueprintFieldPathSource::ViewModel:
+		return FBindingSource::CreateForViewModel(WidgetBlueprint, Path.GetViewModelId());
+	case EMVVMBlueprintFieldPathSource::Widget:
 		return FBindingSource::CreateForWidget(WidgetBlueprint, Path.GetWidgetName());
 	}
 	return FBindingSource();
 }
+
+} // namespace
+
+#undef LOCTEXT_NAMESPACE
\ No newline at end of file
diff --git a/Engine/Plugins/Runtime/ModelViewViewModel/Source/ModelViewViewModelEditor/Private/Types/MVVMBindingSource.h b/Engine/Plugins/Runtime/ModelViewViewModel/Source/ModelViewViewModelEditor/Private/Types/MVVMBindingSource.h
--- a/Engine/Plugins/Runtime/ModelViewViewModel/Source/ModelViewViewModelEditor/Private/Types/MVVMBindingSource.h
+++ b/Engine/Plugins/Runtime/ModelViewViewModel/Source/ModelViewViewModelEditor/Private/Types/MVVMBindingSource.h
@@ -3,24 +3,31 @@
 #pragma once
 
 #include "Framework/Views/TableViewTypeTraits.h"
+#include "MVVMPropertyPath.h"
 #include "UObject/Class.h"
 
 struct FMVVMBindingName;
 struct FMVVMBlueprintPropertyPath;
+struct FMVVMBlueprintViewModelContext;
+class UWidget;
 class UWidgetBlueprint;
 
 namespace UE::MVVM
 {
+
 struct FBindingSource
 {
+private:
 	FGuid ViewModelId;
-	FName Name;
+	FName WidgetName;
 	TWeakObjectPtr<UClass> Class;
 	FText DisplayName;
+	EMVVMBlueprintFieldPathSource Source = EMVVMBlueprintFieldPathSource::None;
 
+public:
 	bool operator==(const FBindingSource& Other) const
 	{
-		return ViewModelId == Other.ViewModelId && Name == Other.Name && Class == Other.Class;
+		return Source == Other.Source && ViewModelId == Other.ViewModelId && WidgetName == Other.WidgetName && Class == Other.Class;
 	}
 
 	bool operator!=(const FBindingSource& Other) const
@@ -28,34 +35,62 @@ struct FBindingSource
 		return !(operator==(Other));
 	}
 
-	friend int32 GetTypeHash(const FBindingSource& Source)
+	friend int32 GetTypeHash(const FBindingSource& InSource)
 	{
-		uint32 Hash = HashCombine(GetTypeHash(Source.ViewModelId), GetTypeHash(Source.Name));
-		Hash = HashCombine(Hash, GetTypeHash(Source.Class));
+		uint32 Hash = HashCombine(GetTypeHash(InSource.ViewModelId), GetTypeHash(InSource.WidgetName));
+		Hash = HashCombine(Hash, GetTypeHash(InSource.Class));
+		Hash = HashCombine(Hash, GetTypeHash(InSource.Source));
 		return Hash;
 	}
 
 	bool IsValid() const
 	{
-		return Class != nullptr && (ViewModelId.IsValid() || !Name.IsNone());
+		return Class != nullptr && Source != EMVVMBlueprintFieldPathSource::None;
+	}
+
+	EMVVMBlueprintFieldPathSource GetSource() const
+	{
+		return Source;
+	}
+
+	const UClass* GetClass() const;
+	FText GetDisplayName() const;
+
+	FName GetWidgetName() const
+	{
+		return WidgetName;
+	}
+
+	FGuid GetViewModelId() const
+	{
+		return ViewModelId;
 	}
 
 	void Reset()
 	{
 		ViewModelId = FGuid();
-		Name = FName();
+		WidgetName = FName();
 		Class = nullptr;
 		DisplayName = FText::GetEmpty();
+		Source = EMVVMBlueprintFieldPathSource::None;
 	}
 
 	FMVVMBindingName ToBindingName(const UWidgetBlueprint* WidgetBlueprint) const;
+	void SetSourceTo(FMVVMBlueprintPropertyPath& PropertyPath) const;
+
+	bool Matches(const UWidgetBlueprint* WidgetBlueprint, const FMVVMBlueprintPropertyPath& PropertyPath) const;
 
+	static FBindingSource CreateForBlueprint(const UWidgetBlueprint* WidgetBlueprint);
+	static FBindingSource CreateForWidget(const UWidgetBlueprint* WidgetBlueprint, const UWidget* Widget);
 	static FBindingSource CreateForWidget(const UWidgetBlueprint* WidgetBlueprint, FName WidgetName);
 	static FBindingSource CreateForViewModel(const UWidgetBlueprint* WidgetBlueprint, FGuid ViewModelId);
 	static FBindingSource CreateForViewModel(const UWidgetBlueprint* WidgetBlueprint, FName ViewModelName);
+	static FBindingSource CreateForViewModel(const UWidgetBlueprint* WidgetBlueprint, const FMVVMBlueprintViewModelContext& ViewModelContext);
+	static FBindingSource CreateEmptySource(UClass* ViewModel);
 	static FBindingSource CreateFromPropertyPath(const UWidgetBlueprint* WidgetBlueprint, const FMVVMBlueprintPropertyPath& Path);
 };
-}
+
+} // namespace
 
 template <>
 struct TIsValidListItem<UE::MVVM::FBindingSource>
@@ -105,7 +140,7 @@ struct TListTypeTraits<UE::MVVM::FBindingSource>
 
 	static FString DebugDump(UE::MVVM::FBindingSource InPtr)
 	{
-		return InPtr.DisplayName.ToString();
+		return InPtr.GetDisplayName().ToString();
 	}
 
 	class SerializerType {};
diff --git a/Engine/Plugins/Runtime/ModelViewViewModel/Source/ModelViewViewModelEditor/Private/Widgets/SMVVMEventParameter.cpp b/Engine/Plugins/Runtime/ModelViewViewModel/Source/ModelViewViewModelEditor/Private/Widgets/SMVVMEventParameter.cpp
--- a/Engine/Plugins/Runtime/ModelViewViewModel/Source/ModelViewViewModelEditor/Private/Widgets/SMVVMEventParameter.cpp
+++ b/Engine/Plugins/Runtime/ModelViewViewModel/Source/ModelViewViewModelEditor/Private/Widgets/SMVVMEventParameter.cpp
@@ -67,7 +67,7 @@ void SEventParameter::Construct(const FArguments& InArgs, UWidgetBlueprint* InWi
 	}
 
 	FMVVMBlueprintPropertyPath Path = OnGetSelectedField();
-	bDefaultValueVisible = Path.IsEmpty();
+	bDefaultValueVisible = !Path.IsValid();
 
 	TSharedPtr<SHorizontalBox> HBox;
 
diff --git a/Engine/Plugins/Runtime/ModelViewViewModel/Source/ModelViewViewModelEditor/Private/Widgets/SMVVMFieldSelector.cpp b/Engine/Plugins/Runtime/ModelViewViewModel/Source/ModelViewViewModelEditor/Private/Widgets/SMVVMFieldSelector.cpp
--- a/Engine/Plugins/Runtime/ModelViewViewModel/Source/ModelViewViewModelEditor/Private/Widgets/SMVVMFieldSelector.cpp
+++ b/Engine/Plugins/Runtime/ModelViewViewModel/Source/ModelViewViewModelEditor/Private/Widgets/SMVVMFieldSelector.cpp
@@ -104,7 +104,7 @@ int32 SFieldSelector::GetCurrentDisplayIndex() const
 	}
 	if (OnGetPropertyPath.IsBound())
 	{
-		if (!OnGetPropertyPath.Execute().IsEmpty())
+		if (OnGetPropertyPath.Execute().IsValid())
 		{
 			return 0;
 		}
diff --git a/Engine/Plugins/Runtime/ModelViewViewModel/Source/ModelViewViewModelEditor/Private/Widgets/SMVVMFieldSelectorMenu.cpp b/Engine/Plugins/Runtime/ModelViewViewModel/Source/ModelViewViewModelEditor/Private/Widgets/SMVVMFieldSelectorMenu.cpp
--- a/Engine/Plugins/Runtime/ModelViewViewModel/Source/ModelViewViewModelEditor/Private/Widgets/SMVVMFieldSelectorMenu.cpp
+++ b/Engine/Plugins/Runtime/ModelViewViewModel/Source/ModelViewViewModelEditor/Private/Widgets/SMVVMFieldSelectorMenu.cpp
@@ -66,7 +66,7 @@ void SFieldSelectorMenu::Construct(const FArguments& InArgs, const UWidgetBluepr
 	SelectionContext = InArgs._SelectionContext;
 
 	bIsMenuInitialized = false;
-	bIsClearEnabled = (InArgs._CurrentPropertyPathSelected.IsSet() && !InArgs._CurrentPropertyPathSelected.GetValue().IsEmpty()) || InArgs._CurrentFunctionSelected != nullptr;
+	bIsClearEnabled = (InArgs._CurrentPropertyPathSelected.IsSet() && InArgs._CurrentPropertyPathSelected.GetValue().IsValid()) || InArgs._CurrentFunctionSelected != nullptr;
 
 	// If we're showing conversion functions, we don't want to set the AssignableTo property of SSourceBindingList, because then it will only show exact matches, 
 	// and since we're also showing conversion functions we know that's not what the user wants.
@@ -216,7 +216,7 @@ bool SFieldSelectorMenu::IsSelectEnabled() const
 	if (BindingList.IsValid())
 	{
 		FMVVMBlueprintPropertyPath Path = BindingList->GetSelectedProperty();
-		if (Path.IsFromViewModel() || Path.IsFromWidget())
+		if (Path.IsValid())
 		{
 			return true;
 		}
@@ -264,14 +264,7 @@ FReply SFieldSelectorMenu::HandleClearClicked()
 	if (SelectionContext.FixedBindingSource.IsSet())
 	{
 		FBindingSource Source = SelectionContext.FixedBindingSource.GetValue();
-		if (Source.ViewModelId.IsValid())
-		{
-			NewProperty.SetViewModelId(Source.ViewModelId);
-		}
-		else 
-		{
-			NewProperty.SetWidgetName(Source.Name);
-		}
+		Source.SetSourceTo(NewProperty);
 	}
 
 	OnFieldSelectionChanged.ExecuteIfBound(NewProperty, nullptr);
@@ -453,7 +446,7 @@ TSharedRef<ITableRow> SFieldSelectorMenu::HandleGenerateViewModelRow(FBindingSou
 			.AutoWidth()
 			[
 				SNew(SImage)
-				.Image(FSlateIconFinder::FindIconBrushForClass(ViewModel.Class.Get()))
+				.Image(FSlateIconFinder::FindIconBrushForClass(ViewModel.GetClass()))
 				.ColorAndOpacity(FSlateColor::UseForeground())
 			]
 			+ SHorizontalBox::Slot()
@@ -462,7 +455,7 @@ TSharedRef<ITableRow> SFieldSelectorMenu::HandleGenerateViewModelRow(FBindingSou
 			.AutoWidth()
 			[
 				SNew(STextBlock)
-				.Text(ViewModel.DisplayName)
+				.Text(ViewModel.GetDisplayName())
 				.ColorAndOpacity(FSlateColor::UseForeground())
 				.HighlightText_Lambda([this]() { return SearchBox.IsValid() ? SearchBox->GetText() : FText::GetEmpty(); })
 			]
@@ -746,19 +739,17 @@ void SFieldSelectorMenu::FilterViewModels(const FText& NewText)
 		return;
 	}
 
-	for (const FBindingSource& ViewModel : ViewModelSources)
+	for (const FBindingSource& BindingSource : ViewModelSources)
 	{
-		const FString DisplayName = ViewModel.DisplayName.ToString();
-		const FString ClassName = ViewModel.Class != nullptr ? ViewModel.Class->GetName() : FString();
-		const FString Name = ViewModel.Name.ToString();
+		const FString DisplayName = BindingSource.GetDisplayName().ToString();
+		const FString ClassName = BindingSource.GetClass() != nullptr ? BindingSource.GetClass()->GetName() : FString();
 
 		bool bMatchesFilters = true;
 		
 		for (const FString& Filter : FilterStrings)
 		{
 			if (!DisplayName.Contains(Filter) && 
-				!ClassName.Contains(Filter) &&
-				!Name.Contains(Filter))
+				!ClassName.Contains(Filter))
 			{
 				bMatchesFilters = false;
 				break;
@@ -767,11 +758,7 @@ void SFieldSelectorMenu::FilterViewModels(const FText& NewText)
 
 		if (bMatchesFilters)
 		{
-			FBindingSource& NewSource = FilteredViewModelSources.Add_GetRef(FBindingSource());
-			NewSource.Class = ViewModel.Class;
-			NewSource.Name = ViewModel.Name;
-			NewSource.DisplayName = ViewModel.DisplayName;
-			NewSource.ViewModelId = ViewModel.ViewModelId;
+			FilteredViewModelSources.Add(BindingSource);
 		}
 	}
 }
@@ -781,11 +768,13 @@ TSharedRef<SWidget> SFieldSelectorMenu::CreateBindingContextPanel(const FArgumen
 	// show source picker
 	TSharedRef<SVerticalBox> StackedSourcePicker = SNew(SVerticalBox);
 
+	const bool bUseFixedSource = SelectionContext.FixedBindingSource.IsSet() && SelectionContext.FixedBindingSource.GetValue().IsValid();
+	const EMVVMBlueprintFieldPathSource FixedFieldPathSource = bUseFixedSource ? SelectionContext.FixedBindingSource.GetValue().GetSource() : EMVVMBlueprintFieldPathSource::None;
+	const EMVVMBlueprintFieldPathSource CurrentPropertyFieldPathSource = InArgs._CurrentPropertyPathSelected.IsSet() ? InArgs._CurrentPropertyPathSelected.GetValue().GetSource(WidgetBlueprint.Get()) : EMVVMBlueprintFieldPathSource::None;
+
 	if (SelectionContext.bAllowViewModels)
 	{
-		if (SelectionContext.FixedBindingSource.IsSet()
-			&& SelectionContext.FixedBindingSource.GetValue().IsValid()
-			&& SelectionContext.FixedBindingSource.GetValue().ViewModelId.IsValid())
+		if (FixedFieldPathSource == EMVVMBlueprintFieldPathSource::ViewModel)
 		{
 			ViewModelSources.Add(SelectionContext.FixedBindingSource.GetValue());
 		}
@@ -803,11 +792,11 @@ TSharedRef<SWidget> SFieldSelectorMenu::CreateBindingContextPanel(const FArgumen
 			.OnSelectionChanged(this, &SFieldSelectorMenu::HandleViewModelSelected);
 
 		FBindingSource SelectedSource;
-		if (InArgs._CurrentPropertyPathSelected.IsSet() && InArgs._CurrentPropertyPathSelected.GetValue().IsFromViewModel())
+		if (CurrentPropertyFieldPathSource == EMVVMBlueprintFieldPathSource::ViewModel)
 		{
 			for (const FBindingSource& Source : FilteredViewModelSources)
 			{
-				if (Source.ViewModelId == InArgs._CurrentPropertyPathSelected.GetValue().GetViewModelId())
+				if (Source.GetViewModelId() == InArgs._CurrentPropertyPathSelected.GetValue().GetViewModelId())
 				{
 					SelectedSource = Source;
 				}
@@ -829,11 +818,13 @@ TSharedRef<SWidget> SFieldSelectorMenu::CreateBindingContextPanel(const FArgumen
 	if (SelectionContext.bAllowWidgets)
 	{
 		TArray<FName> ShowOnly;
-		if (SelectionContext.FixedBindingSource.IsSet()
-			&& SelectionContext.FixedBindingSource.GetValue().IsValid()
-			&& !SelectionContext.FixedBindingSource.GetValue().Name.IsNone())
+		if (FixedFieldPathSource == EMVVMBlueprintFieldPathSource::Widget)
 		{
-			ShowOnly.Add(SelectionContext.FixedBindingSource.GetValue().Name);
+			ShowOnly.Add(SelectionContext.FixedBindingSource.GetValue().GetWidgetName());
+		}
+		else if (FixedFieldPathSource == EMVVMBlueprintFieldPathSource::SelfContext)
+		{
+			ShowOnly.Add(WidgetBlueprint.Get()->GetFName());
 		}
 
 		WidgetList = SNew(SReadOnlyHierarchyView, WidgetBlueprint.Get())
@@ -843,10 +834,14 @@ TSharedRef<SWidget> SFieldSelectorMenu::CreateBindingContextPanel(const FArgumen
 			.ShowOnly(ShowOnly)
 			.ExpandAll(false);
 
-		if (InArgs._CurrentPropertyPathSelected.IsSet() && InArgs._CurrentPropertyPathSelected.GetValue().IsFromWidget())
+		if (CurrentPropertyFieldPathSource == EMVVMBlueprintFieldPathSource::Widget)
 		{
 			WidgetList->SetSelectedWidget(InArgs._CurrentPropertyPathSelected.GetValue().GetWidgetName());
 		}
+		else if (CurrentPropertyFieldPathSource == EMVVMBlueprintFieldPathSource::SelfContext)
+		{
+			WidgetList->SetSelectedWidget(WidgetBlueprint.Get()->GetFName());
+		}
 
 		StackedSourcePicker->AddSlot()
 			.AutoHeight()
@@ -901,9 +896,7 @@ TSharedRef<SWidget> SFieldSelectorMenu::CreateBindingListPanel(const FArguments&
 	}
 	else if (InArgs._CurrentPropertyPathSelected.IsSet())
 	{
-		FBindingSource Source = InArgs._CurrentPropertyPathSelected.GetValue().IsFromViewModel() ?
-			FBindingSource::CreateForViewModel(WidgetBlueprint.Get(), InArgs._CurrentPropertyPathSelected.GetValue().GetViewModelId()) :
-			FBindingSource::CreateForWidget(WidgetBlueprint.Get(), InArgs._CurrentPropertyPathSelected.GetValue().GetWidgetName());
+		FBindingSource Source = FBindingSource::CreateFromPropertyPath(WidgetBlueprint.Get(), InArgs._CurrentPropertyPathSelected.GetValue());
 		BindingList->AddSource(Source);
 	}
 
diff --git a/Engine/Plugins/Runtime/ModelViewViewModel/Source/ModelViewViewModelEditor/Private/Widgets/SMVVMFunctionParameter.cpp b/Engine/Plugins/Runtime/ModelViewViewModel/Source/ModelViewViewModelEditor/Private/Widgets/SMVVMFunctionParameter.cpp
--- a/Engine/Plugins/Runtime/ModelViewViewModel/Source/ModelViewViewModelEditor/Private/Widgets/SMVVMFunctionParameter.cpp
+++ b/Engine/Plugins/Runtime/ModelViewViewModel/Source/ModelViewViewModelEditor/Private/Widgets/SMVVMFunctionParameter.cpp
@@ -81,7 +81,7 @@ void SFunctionParameter::Construct(const FArguments& InArgs, UWidgetBlueprint* I
 	}
 
 	FMVVMBlueprintPropertyPath Path = OnGetSelectedField();
-	bDefaultValueVisible = Path.IsEmpty();
+	bDefaultValueVisible = !Path.IsValid();
 
 	bool bFromViewModel = UE::MVVM::IsForwardBinding(Binding->BindingType);
 	TSharedPtr<SHorizontalBox> HBox;
diff --git a/Engine/Plugins/Runtime/ModelViewViewModel/Source/ModelViewViewModelEditor/Private/Widgets/SMVVMPropertyPath.cpp b/Engine/Plugins/Runtime/ModelViewViewModel/Source/ModelViewViewModelEditor/Private/Widgets/SMVVMPropertyPath.cpp
--- a/Engine/Plugins/Runtime/ModelViewViewModel/Source/ModelViewViewModelEditor/Private/Widgets/SMVVMPropertyPath.cpp
+++ b/Engine/Plugins/Runtime/ModelViewViewModel/Source/ModelViewViewModelEditor/Private/Widgets/SMVVMPropertyPath.cpp
@@ -40,7 +40,7 @@ void SPropertyPath::SetPropertyPath(const FMVVMBlueprintPropertyPath& InProperty
 {
 	FieldBox->ClearChildren();
 
-	bool bHasSource = InPropertyPath.IsFromWidget() || InPropertyPath.IsFromViewModel();
+	bool bHasSource = InPropertyPath.IsValid();
 	if (bShowContext && bHasSource)
 	{
 		FieldBox->AddSlot()
diff --git a/Engine/Plugins/Runtime/ModelViewViewModel/Source/ModelViewViewModelEditor/Private/Widgets/SMVVMSelectViewModel.cpp b/Engine/Plugins/Runtime/ModelViewViewModel/Source/ModelViewViewModelEditor/Private/Widgets/SMVVMSelectViewModel.cpp
--- a/Engine/Plugins/Runtime/ModelViewViewModel/Source/ModelViewViewModelEditor/Private/Widgets/SMVVMSelectViewModel.cpp
+++ b/Engine/Plugins/Runtime/ModelViewViewModel/Source/ModelViewViewModelEditor/Private/Widgets/SMVVMSelectViewModel.cpp
@@ -137,7 +137,7 @@ void SMVVMSelectViewModel::HandleClassPicked(UClass* ClassPicked)
 	if (Private::IsValidViewModel(ClassPicked, DisallowedClassFlags))
 	{
 		SelectedClass = ClassPicked;
-		BindingListWidget->AddSource(ClassPicked, FName(), FGuid());
+		BindingListWidget->AddSource(FBindingSource::CreateEmptySource(ClassPicked));
 	}
 }
 
diff --git a/Engine/Plugins/Runtime/ModelViewViewModel/Source/ModelViewViewModelEditor/Private/Widgets/SMVVMSourceEntry.cpp b/Engine/Plugins/Runtime/ModelViewViewModel/Source/ModelViewViewModelEditor/Private/Widgets/SMVVMSourceEntry.cpp
--- a/Engine/Plugins/Runtime/ModelViewViewModel/Source/ModelViewViewModelEditor/Private/Widgets/SMVVMSourceEntry.cpp
+++ b/Engine/Plugins/Runtime/ModelViewViewModel/Source/ModelViewViewModelEditor/Private/Widgets/SMVVMSourceEntry.cpp
@@ -22,17 +22,18 @@ const FSlateBrush* GetSourceIcon(const FBindingSource& Source)
 		return nullptr;
 	}
 
-	return FSlateIconFinder::FindIconBrushForClass(Source.Class.Get());
+	return FSlateIconFinder::FindIconBrushForClass(Source.GetClass());
 }
 
 FLinearColor GetSourceColor(const FBindingSource& Source)
 {
-	if (!Source.IsValid())
+	const UClass* SourceClass = Source.GetClass();
+	if (SourceClass == nullptr)
 	{
 		return FLinearColor::White;
 	}
 
-	uint32 Hash = GetTypeHash(Source.Class->GetName());
+	uint32 Hash = GetTypeHash(SourceClass->GetName());
 	FLinearColor Color = FLinearColor::White;
 	Color.R = ((Hash * 1 % 96) + 32) / 256.f;
 	Color.G = ((Hash * 2 % 96) + 32) / 256.f;
@@ -42,12 +43,7 @@ FLinearColor GetSourceColor(const FBindingSource& Source)
 
 FText GetSourceDisplayName(const FBindingSource& Source)
 {
-	if (!Source.IsValid())
-	{
-		return LOCTEXT("None", "<None>");
-	}
-
-	return !Source.DisplayName.IsEmpty() ? Source.DisplayName : FText::FromString(Source.Name.ToString());
+	return Source.GetDisplayName();
 }
 
 } // namespace Private
@@ -86,14 +82,14 @@ void SBindingContextEntry::RefreshSource(const FBindingSource& Source)
 
 	if (Source.IsValid())
 	{
-		if (Source.Class != nullptr)
+		if (const UClass* SourceClass = Source.GetClass())
 		{
-			const FText ToolTipText = FText::Join(FText::FromString(TEXT("\n")), Source.DisplayName, Source.Class->GetDisplayNameText());
+			const FText ToolTipText = FText::Join(FText::FromString(TEXT("\n")), Source.GetDisplayName(), SourceClass->GetDisplayNameText());
 			SetToolTipText(ToolTipText);
 		}
 		else
 		{
-			SetToolTipText(Source.DisplayName);
+			SetToolTipText(Source.GetDisplayName());
 		}
 	}
 }
diff --git a/Engine/Plugins/Runtime/ModelViewViewModel/Source/ModelViewViewModelEditor/Private/Widgets/SMVVMSourceSelector.cpp b/Engine/Plugins/Runtime/ModelViewViewModel/Source/ModelViewViewModelEditor/Private/Widgets/SMVVMSourceSelector.cpp
--- a/Engine/Plugins/Runtime/ModelViewViewModel/Source/ModelViewViewModelEditor/Private/Widgets/SMVVMSourceSelector.cpp
+++ b/Engine/Plugins/Runtime/ModelViewViewModel/Source/ModelViewViewModelEditor/Private/Widgets/SMVVMSourceSelector.cpp
@@ -244,11 +244,7 @@ void SBindingContextSelector::Refresh()
 
 			for (const FMVVMBlueprintViewModelContext& ViewModel : ViewModels)
 			{
-				FBindingSource& Source = ViewModelSources.AddDefaulted_GetRef();
-				Source.Class = ViewModel.GetViewModelClass();
-				Source.Name = ViewModel.GetViewModelName();
-				Source.DisplayName = ViewModel.GetDisplayName();
-				Source.ViewModelId = ViewModel.GetViewModelId();
+				ViewModelSources.Add(FBindingSource::CreateForViewModel(WidgetBlueprint.Get(), ViewModel));
 			}
 		}
 	}
diff --git a/Engine/Plugins/Runtime/ModelViewViewModel/Source/ModelViewViewModelEditor/Private/Widgets/SMVVMViewBindingListView.cpp b/Engine/Plugins/Runtime/ModelViewViewModel/Source/ModelViewViewModelEditor/Private/Widgets/SMVVMViewBindingListView.cpp
--- a/Engine/Plugins/Runtime/ModelViewViewModel/Source/ModelViewViewModelEditor/Private/Widgets/SMVVMViewBindingListView.cpp
+++ b/Engine/Plugins/Runtime/ModelViewViewModel/Source/ModelViewViewModelEditor/Private/Widgets/SMVVMViewBindingListView.cpp
@@ -490,7 +490,7 @@ private:
 					if (FMVVMBlueprintViewBinding* Binding = ChildEntry->GetBinding(View))
 					{
 						FMVVMBlueprintPropertyPath CurrentPath = Binding->DestinationPath;
-						CurrentPath.SetWidgetName(Source.Name);
+						Source.SetSourceTo(CurrentPath);
 
 						EditorSubsystem->SetDestinationPathForBinding(WidgetBlueprint, *Binding, CurrentPath);
 					}
@@ -501,7 +501,7 @@ private:
 					if (UMVVMBlueprintViewEvent* Event = ChildEntry->GetEvent())
 					{
 						FMVVMBlueprintPropertyPath CurrentPath = Event->GetEventPath();
-						CurrentPath.SetWidgetName(Source.Name);
+						Source.SetSourceTo(CurrentPath);
 
 						EditorSubsystem->SetEventPath(Event, CurrentPath);
 					}
@@ -519,7 +519,14 @@ private:
 			FMVVMBlueprintPropertyPath Path;
 			if (Entry->IsGroupWidget())
 			{
-				Path.SetWidgetName(Entry->GetGroupName());
+				if (Entry->GetGroupName() == WidgetBlueprint->GetFName())
+				{
+					Path.SetSelfContext();
+				}
+				else
+				{
+					Path.SetWidgetName(Entry->GetGroupName());
+				}
 			}
 			else
 			{
@@ -1014,13 +1021,13 @@ private:
 						{
 							// Search for the widget by its name in the widget tree
 							// If the widget is not found, we know it is the root preview widget so we use the blueprint name.
-							if (WidgetBlueprintPtr->WidgetTree->FindWidget(OwnerWidgetPtr->GetFName()))
+							if (WidgetBlueprintPtr->WidgetTree->FindWidget(OwnerWidgetPtr->GetFName()) && WidgetBlueprintPtr->GetFName() != OwnerWidgetPtr->GetFName())
 							{
 								PropertyPath.SetWidgetName(OwnerWidgetPtr->GetFName());
 							}
 							else
 							{
-								PropertyPath.SetWidgetName(WidgetBlueprintPtr->GetFName());
+								PropertyPath.SetSelfContext();
 							}
 						}
 					}
@@ -1710,13 +1717,13 @@ private:
 			{
 				// Search for the widget by its name in the widget tree
 				// If the widget is not found, we know it is the root preview widget so we use the blueprint name.
-				if (WidgetBlueprintPtr->WidgetTree->FindWidget(OwnerWidgetPtr->GetFName()))
+				if (WidgetBlueprintPtr->WidgetTree->FindWidget(OwnerWidgetPtr->GetFName()) && WidgetBlueprintPtr->GetFName() != OwnerWidgetPtr->GetFName())
 				{
 					PropertyPath.SetWidgetName(OwnerWidgetPtr->GetFName());
 				}
 				else
 				{
-					PropertyPath.SetWidgetName(WidgetBlueprintPtr->GetFName());
+					PropertyPath.SetSelfContext();
 				}
 			}
 		}
@@ -2027,17 +2034,21 @@ void SBindingsList::Refresh()
 			
 			FName GroupName;
 			FGuid GroupViewModelId;
-			if (Binding.DestinationPath.IsFromWidget())
+			switch (Binding.DestinationPath.GetSource(WidgetBlueprint))
 			{
+			case EMVVMBlueprintFieldPathSource::SelfContext:
+				GroupName = WidgetBlueprint->GetFName();
+				break;
+			case EMVVMBlueprintFieldPathSource::Widget:
 				GroupName = Binding.DestinationPath.GetWidgetName();
-			}
-			else
-			{
+				break;
+			case EMVVMBlueprintFieldPathSource::ViewModel:
 				if (const FMVVMBlueprintViewModelContext* ViewModelContext = BlueprintView->FindViewModel(Binding.DestinationPath.GetViewModelId()))
 				{
 					GroupName = ViewModelContext->GetViewModelName();
 					GroupViewModelId = ViewModelContext->GetViewModelId();
 				}
+				break;
 			}
 
 			// Find the group entry
@@ -2107,17 +2118,21 @@ void SBindingsList::Refresh()
 
 			FName GroupName;
 			FGuid GroupViewModelId;
-			if (Event->GetEventPath().IsFromWidget())
+			switch (Event->GetEventPath().GetSource(WidgetBlueprint))
 			{
+			case EMVVMBlueprintFieldPathSource::SelfContext:
+				GroupName = WidgetBlueprint->GetFName();
+				break;
+			case EMVVMBlueprintFieldPathSource::Widget:
 				GroupName = Event->GetEventPath().GetWidgetName();
-			}
-			else if (Event->GetEventPath().IsFromViewModel())
-			{
+				break;
+			case EMVVMBlueprintFieldPathSource::ViewModel:
 				if (const FMVVMBlueprintViewModelContext* ViewModelContext = BlueprintView->FindViewModel(Event->GetEventPath().GetViewModelId()))
 				{
 					GroupName = ViewModelContext->GetViewModelName();
 					GroupViewModelId = ViewModelContext->GetViewModelId();
 				}
+				break;
 			}
 
 			// Find the group entry
diff --git a/Engine/Plugins/Runtime/ModelViewViewModel/Source/ModelViewViewModelEditor/Private/Widgets/SMVVMViewBindingPanel.cpp b/Engine/Plugins/Runtime/ModelViewViewModel/Source/ModelViewViewModelEditor/Private/Widgets/SMVVMViewBindingPanel.cpp
--- a/Engine/Plugins/Runtime/ModelViewViewModel/Source/ModelViewViewModelEditor/Private/Widgets/SMVVMViewBindingPanel.cpp
+++ b/Engine/Plugins/Runtime/ModelViewViewModel/Source/ModelViewViewModelEditor/Private/Widgets/SMVVMViewBindingPanel.cpp
@@ -247,7 +247,14 @@ void SBindingsPanel::AddDefaultBinding()
 					{
 						FMVVMBlueprintViewBinding& Binding = EditorSubsystem->AddBinding(MVVMExtensionPtr->GetWidgetBlueprint());
 						FMVVMBlueprintPropertyPath Path;
-						Path.SetWidgetName(WidgetReference.GetTemplate()->GetFName());
+						if (WidgetReference.GetTemplate()->GetFName() == MVVMExtensionPtr->GetWidgetBlueprint()->GetFName())
+						{
+							Path.SetSelfContext();
+						}
+						else
+						{
+							Path.SetWidgetName(WidgetReference.GetTemplate()->GetFName());
+						}
 						EditorSubsystem->SetDestinationPathForBinding(MVVMExtensionPtr->GetWidgetBlueprint(), Binding, Path);
 						AddedBindingId = Binding.BindingId;
 
diff --git a/Engine/Plugins/Runtime/ModelViewViewModel/Source/ModelViewViewModelEditor/Private/Widgets/SMVVMViewModelBindingListWidget.cpp b/Engine/Plugins/Runtime/ModelViewViewModel/Source/ModelViewViewModelEditor/Private/Widgets/SMVVMViewModelBindingListWidget.cpp
--- a/Engine/Plugins/Runtime/ModelViewViewModel/Source/ModelViewViewModelEditor/Private/Widgets/SMVVMViewModelBindingListWidget.cpp
+++ b/Engine/Plugins/Runtime/ModelViewViewModel/Source/ModelViewViewModelEditor/Private/Widgets/SMVVMViewModelBindingListWidget.cpp
@@ -385,23 +385,9 @@ void SSourceBindingList::ClearSources()
 	}
 }
 
-void SSourceBindingList::AddSource(UClass* Class, FName Name, FGuid Guid)
+void SSourceBindingList::AddWidgetBlueprint()
 {
-	FBindingSource Source;
-	Source.Class = Class;
-	Source.Name = Name;
-	Source.ViewModelId = Guid;
-
-	AddSources(MakeArrayView(&Source, 1));
-} 
-
-void SSourceBindingList::AddWidgetBlueprint(const UWidgetBlueprint* InWidgetBlueprint)
-{
-	FBindingSource Source;
-	Source.Class = InWidgetBlueprint->GeneratedClass;
-	Source.Name = InWidgetBlueprint->GetFName();
-	Source.DisplayName = FText::FromString(InWidgetBlueprint->GetName());
-
+	FBindingSource Source = FBindingSource::CreateForBlueprint(WidgetBlueprint.Get());
 	AddSources(MakeArrayView(&Source, 1 ));
 }
 
@@ -412,10 +398,7 @@ void SSourceBindingList::AddWidgets(TArrayView<const UWidget*> InWidgets)
 
 	for (const UWidget* Widget : InWidgets)
 	{
-		FBindingSource& Source = NewSources.AddDefaulted_GetRef();
-		Source.Class = Widget->GetClass();
-		Source.Name = Widget->GetFName();
-		Source.DisplayName = Widget->GetLabelText();
+		NewSources.Add(FBindingSource::CreateForWidget(WidgetBlueprint.Get(), Widget));
 	}
 
 	AddSources(NewSources);
@@ -428,9 +411,7 @@ void SSourceBindingList::AddViewModels(TArrayView<const FMVVMBlueprintViewModelC
 
 	for (const FMVVMBlueprintViewModelContext& ViewModelContext : InViewModels)
 	{
-		FBindingSource& Source = NewSources.AddDefaulted_GetRef();
-		Source.Class = ViewModelContext.GetViewModelClass();
-		Source.ViewModelId = ViewModelContext.GetViewModelId();
+		NewSources.Add(FBindingSource::CreateForViewModel(WidgetBlueprint.Get(), ViewModelContext));
 	}
 
 	AddSources(NewSources);
@@ -448,15 +429,15 @@ void SSourceBindingList::AddSources(TArrayView<const FBindingSource> InSources)
 		const UWidgetBlueprint* WidgetBlueprintPtr = WidgetBlueprint.Get();
 		for (const FBindingSource& Source : InSources)
 		{
-			if (const UClass* SourceClass = Source.Class.Get())
+			if (const UClass* SourceClass = Source.GetClass())
 			{
 				SPropertyViewer::FHandle Handle;
-				if (SourceClass && Source.Class->ImplementsInterface(UNotifyFieldValueChanged::StaticClass()))
+				if (SourceClass->ImplementsInterface(UNotifyFieldValueChanged::StaticClass()))
 				{
 					UWidget* Widget = nullptr;
-					if (WidgetBlueprintPtr && SourceClass->IsChildOf(UWidget::StaticClass()))
+					if (WidgetBlueprintPtr && Source.GetSource() == EMVVMBlueprintFieldPathSource::Widget)
 					{
-						Widget = WidgetBlueprintPtr->WidgetTree->FindWidget(Source.Name);
+						Widget = WidgetBlueprintPtr->WidgetTree->FindWidget(Source.GetWidgetName());
 					}
 
 					if (Widget)
@@ -552,15 +533,7 @@ FMVVMBlueprintPropertyPath SSourceBindingList::CreateBlueprintPropertyPath(SProp
 
 		if (bPassFilter)
 		{
-			if (Source->Key.ViewModelId.IsValid())
-			{
-				PropertyPath.SetViewModelId(Source->Key.ViewModelId);
-			}
-			else
-			{
-				PropertyPath.SetWidgetName(Source->Key.Name);
-			}
-
+			Source->Key.SetSourceTo(PropertyPath);
 			PropertyPath.ResetPropertyPath();
 			for (const FFieldVariant& Field : FieldPath)
 			{
@@ -595,14 +568,7 @@ FMVVMBlueprintPropertyPath SSourceBindingList::CreateBlueprintPropertyPath(SProp
 			}
 			if (bPassFilter)
 			{
-				if (Source->Key.ViewModelId.IsValid())
-				{
-					PropertyPath.SetViewModelId(Source->Key.ViewModelId);
-				}
-				else
-				{
-					PropertyPath.SetWidgetName(Source->Key.Name);
-				}
+				Source->Key.SetSourceTo(PropertyPath);
 				PropertyPath.ResetPropertyPath();
 			}
 		}
@@ -646,7 +612,7 @@ FMVVMBlueprintPropertyPath SSourceBindingList::GetSelectedProperty() const
 	return SelectedPath;
 }
 
-void SSourceBindingList::SetSelectedProperty(const FMVVMBlueprintPropertyPath& Property)
+void SSourceBindingList::SetSelectedProperty(const FMVVMBlueprintPropertyPath& PropertyPath)
 {
 	if (!PropertyViewer.IsValid())
 	{
@@ -662,8 +628,7 @@ void SSourceBindingList::SetSelectedProperty(const FMVVMBlueprintPropertyPath& P
 	SPropertyViewer::FHandle SelectedHandle;
 	for (TPair<FBindingSource, SPropertyViewer::FHandle>& Source : Sources)
 	{
-		if ((Property.IsFromViewModel() && Source.Key.ViewModelId == Property.GetViewModelId()) ||
-			(Property.IsFromWidget() && Source.Key.Name == Property.GetWidgetName()))
+		if (Source.Key.Matches(WidgetBlueprintPtr, PropertyPath))
 		{
 			SelectedHandle = Source.Value;
 			break;
@@ -674,7 +639,7 @@ void SSourceBindingList::SetSelectedProperty(const FMVVMBlueprintPropertyPath& P
 	TArray<FFieldVariant, TMemStackAllocator<>> FieldPath;
 	if (SelectedHandle.IsValid())
 	{
-		TArray<FMVVMConstFieldVariant> FieldVariants = Property.GetFields(WidgetBlueprintPtr->SkeletonGeneratedClass);
+		TArray<FMVVMConstFieldVariant> FieldVariants = PropertyPath.GetFields(WidgetBlueprintPtr->SkeletonGeneratedClass);
 		FieldPath.Reserve(FieldVariants.Num());
 
 		for (const FMVVMConstFieldVariant& Variant : FieldVariants)
diff --git a/Engine/Plugins/Runtime/ModelViewViewModel/Source/ModelViewViewModelEditor/Private/Widgets/SMVVMViewModelBindingListWidget.h b/Engine/Plugins/Runtime/ModelViewViewModel/Source/ModelViewViewModelEditor/Private/Widgets/SMVVMViewModelBindingListWidget.h
--- a/Engine/Plugins/Runtime/ModelViewViewModel/Source/ModelViewViewModelEditor/Private/Widgets/SMVVMViewModelBindingListWidget.h
+++ b/Engine/Plugins/Runtime/ModelViewViewModel/Source/ModelViewViewModelEditor/Private/Widgets/SMVVMViewModelBindingListWidget.h
@@ -92,11 +92,10 @@ namespace UE::MVVM
 
 		void ClearSources();
 
-		void AddSource(UClass* Class, FName Name, FGuid Guid);
 		void AddSource(const FBindingSource& InSource);
 		void AddSources(TArrayView<const FBindingSource> InSources);
 
-		void AddWidgetBlueprint(const UWidgetBlueprint* WidgetBlueprint);
+		void AddWidgetBlueprint();
 		void AddWidgets(TArrayView<const UWidget*> Widgets);
 		void AddViewModels(TArrayView<const FMVVMBlueprintViewModelContext> ViewModels);
 
diff --git a/Engine/Source/Runtime/Core/Public/UObject/FortniteMainBranchObjectVersion.h b/Engine/Source/Runtime/Core/Public/UObject/FortniteMainBranchObjectVersion.h
--- a/Engine/Source/Runtime/Core/Public/UObject/FortniteMainBranchObjectVersion.h
+++ b/Engine/Source/Runtime/Core/Public/UObject/FortniteMainBranchObjectVersion.h
@@ -384,6 +384,9 @@ struct FFortniteMainBranchObjectVersion
 		// Added support of external packaging of Data Layer Instances
 		AddDataLayerInstanceExternalPackage,
 
+		// Update paths to keep a flag if they are the widget BP
+		MVVMPropertyPathSelf,
+
 		// -----<new versions can be added above this line>-------------------------------------------------
 		VersionPlusOne,
 		LatestVersion = VersionPlusOne - 1
