MVVM: Change the description of the MVVM settings in the settings windows for "UMG Model View Viewmodel".

diff --git a/Engine/Plugins/Runtime/ModelViewViewModel/Source/ModelViewViewModelBlueprint/Private/MVVMDeveloperProjectSettings.cpp b/Engine/Plugins/Runtime/ModelViewViewModel/Source/ModelViewViewModelBlueprint/Private/MVVMDeveloperProjectSettings.cpp
--- a/Engine/Plugins/Runtime/ModelViewViewModel/Source/ModelViewViewModelBlueprint/Private/MVVMDeveloperProjectSettings.cpp
+++ b/Engine/Plugins/Runtime/ModelViewViewModel/Source/ModelViewViewModelBlueprint/Private/MVVMDeveloperProjectSettings.cpp
@@ -34,7 +34,7 @@ FName UMVVMDeveloperProjectSettings::GetCategoryName() const
 
 FText UMVVMDeveloperProjectSettings::GetSectionText() const
 {
-	return LOCTEXT("MVVMProjectSettings", "Model View Viewmodel");
+	return LOCTEXT("MVVMProjectSettings", "UMG Model View Viewmodel");
 }
 
 bool UMVVMDeveloperProjectSettings::PropertyHasFiltering(const FProperty* Property) const
