MVVM: Mark autogenerated functions transient when requested. That can create function to be null when loaded. We want to execute the function and compile the functions in the editor but we want to recreate them on every reload.

diff --git a/Engine/Plugins/Runtime/ModelViewViewModel/Source/ModelViewViewModel/Private/ViewModel/MVVMInstancedViewModelGeneratedClass.cpp b/Engine/Plugins/Runtime/ModelViewViewModel/Source/ModelViewViewModel/Private/ViewModel/MVVMInstancedViewModelGeneratedClass.cpp
--- a/Engine/Plugins/Runtime/ModelViewViewModel/Source/ModelViewViewModel/Private/ViewModel/MVVMInstancedViewModelGeneratedClass.cpp
+++ b/Engine/Plugins/Runtime/ModelViewViewModel/Source/ModelViewViewModel/Private/ViewModel/MVVMInstancedViewModelGeneratedClass.cpp
@@ -38,6 +38,8 @@ void UMVVMInstancedViewModelGeneratedClass::Link(FArchive& Ar, bool bRelinkExist
 {
 	Super::Link(Ar, bRelinkExistingProperties);
 
+	OnRepFunctionToLink.RemoveAllSwap([](UFunction* Other){ return Other == nullptr; }, true);
+
 	for (UFunction* OnRep : OnRepFunctionToLink)
 	{
 		NativeFunctionLookupTable.Emplace(OnRep->GetFName(), &UMVVMInstancedViewModelGeneratedClass::K2_CallNativeOnRep);
diff --git a/Engine/Plugins/Runtime/ModelViewViewModel/Source/ModelViewViewModelBlueprint/Private/MVVMWidgetBlueprintExtension_View.cpp b/Engine/Plugins/Runtime/ModelViewViewModel/Source/ModelViewViewModelBlueprint/Private/MVVMWidgetBlueprintExtension_View.cpp
--- a/Engine/Plugins/Runtime/ModelViewViewModel/Source/ModelViewViewModelBlueprint/Private/MVVMWidgetBlueprintExtension_View.cpp
+++ b/Engine/Plugins/Runtime/ModelViewViewModel/Source/ModelViewViewModelBlueprint/Private/MVVMWidgetBlueprintExtension_View.cpp
@@ -26,6 +26,14 @@ namespace UE::MVVM::Private
 		TEXT("Is the model view viewmodel view is allowed to be added to the generated Widget GeneratedClass."),
 		ECVF_ReadOnly
 	);
+
+	bool GAutogeneratedFunctionsAreTransient = false;
+	static FAutoConsoleVariableRef CVarAutogeneratedFunctionsAreTransient(
+		TEXT("MVVM.AutogeneratedFunctionsAreTransient"),
+		GAutogeneratedFunctionsAreTransient,
+		TEXT("Is the autogenerated function should be mark as transient."),
+		ECVF_ReadOnly
+	);
 }
 
 void UMVVMWidgetBlueprintExtension_View::CreateBlueprintViewInstance()
@@ -164,18 +172,6 @@ void UMVVMWidgetBlueprintExtension_View::HandleFinishCompilingClass(UWidgetBluep
 		if (bCompiled)
 		{
 			check(ViewExtension);
-			if (!UE::MVVM::Private::GAllowViewClass)
-			{
-				ViewExtension->SetFlags(RF_Transient);
-
-				// If we are not allowed to add the view class, add the transient flags on added conversion graph and event graph.
-				for (TFieldIterator<UFunction> FunctionIter(Class, EFieldIteratorFlags::ExcludeSuper); FunctionIter; ++FunctionIter)
-				{
-					UFunction* Function = *FunctionIter;
-					Function->SetFlags(RF_Transient);
-				}
-				Class->Children = nullptr; // clear all functions (including those that the user generated).
-			}
 
 			// Does it have any bindings
 			if (const_cast<const UMVVMViewClass*>(ViewExtension)->GetCompiledBindings().Num() > 0)
@@ -186,10 +182,25 @@ void UMVVMWidgetBlueprintExtension_View::HandleFinishCompilingClass(UWidgetBluep
 					CurrentCompilerContext->GetCompilerContext().MessageLog.Warning(*LOCTEXT("MoreThanOneViewWarning", "There is more than one view.").ToString());
 				}
 
-				CurrentCompilerContext->AddExtension(Class, ViewExtension);
+				// If we are not allowed to add the view class
+				if (UE::MVVM::Private::GAllowViewClass)
+				{
+					CurrentCompilerContext->AddExtension(Class, ViewExtension);
+				}
 			}
 		}
 	}
+
+	// If we can't auto-generate function add the transient flags
+	if (UE::MVVM::Private::GAutogeneratedFunctionsAreTransient)
+	{
+		// If we are not allowed to add the view class, add the transient flags on added conversion graphs and event graphs.
+		for (TFieldIterator<UFunction> FunctionIter(Class, EFieldIteratorFlags::ExcludeSuper); FunctionIter; ++FunctionIter)
+		{
+			UFunction* Function = *FunctionIter;
+			Function->SetFlags(RF_Transient);
+		}
+	}
 }
 
 
