MVVM: Split and recombine structure argument for conversion function parameters and events parameters.
Fix event broadcast when a ViewBinding changes.
Add option to reparents the viewmodel class. Force the list to be repainted when the viewmodel changes. That helps with display names when the viewmodel is renamed and with graphpin that needs to be rebuilt.

diff --git a/Engine/Plugins/Runtime/ModelViewViewModel/Source/ModelViewViewModelBlueprint/Private/MVVMBlueprintView.cpp b/Engine/Plugins/Runtime/ModelViewViewModel/Source/ModelViewViewModelBlueprint/Private/MVVMBlueprintView.cpp
--- a/Engine/Plugins/Runtime/ModelViewViewModel/Source/ModelViewViewModelBlueprint/Private/MVVMBlueprintView.cpp
+++ b/Engine/Plugins/Runtime/ModelViewViewModel/Source/ModelViewViewModelBlueprint/Private/MVVMBlueprintView.cpp
@@ -129,6 +129,35 @@ bool UMVVMBlueprintView::RenameViewModel(FName OldViewModelName, FName NewViewMo
 	return ViewModelContext != nullptr;
 }
 
+bool UMVVMBlueprintView::ReparentViewModel(FGuid ViewModelId, const UClass* ViewModelClass)
+{
+	if (ViewModelClass && ViewModelClass->ImplementsInterface(UNotifyFieldValueChanged::StaticClass()))
+	{
+		FMVVMBlueprintViewModelContext* ViewModelContext = AvailableViewModels.FindByPredicate([ViewModelId](const FMVVMBlueprintViewModelContext& Other)
+			{
+				return Other.GetViewModelId() == ViewModelId;
+			});
+		if (ViewModelContext)
+		{
+			ViewModelContext->NotifyFieldValueClass = const_cast<UClass*>(ViewModelClass);
+			TArray<EMVVMBlueprintViewModelContextCreationType> ValidCreationTypes = UE::MVVM::GetAllowedContextCreationType(ViewModelClass);
+			if (!ValidCreationTypes.Contains(ViewModelContext->CreationType))
+			{
+				if (ensureMsgf(ValidCreationTypes.Num() > 0, TEXT("There is no valid creation type for this class.")))
+				{
+					ViewModelContext->CreationType = ValidCreationTypes[0];
+				}
+			}
+
+			FBlueprintEditorUtils::MarkBlueprintAsModified(GetOuterUMVVMWidgetBlueprintExtension_View()->GetWidgetBlueprint());
+
+			OnViewModelsUpdated.Broadcast();
+			return true;
+		}
+	}
+	return false;
+}
+
 const FMVVMBlueprintViewBinding* UMVVMBlueprintView::FindBinding(const UWidget* Widget, const FProperty* Property) const
 {
 	return const_cast<UMVVMBlueprintView*>(this)->FindBinding(Widget, Property);
@@ -482,6 +511,10 @@ void UMVVMBlueprintView::PostEditChangeProperty(FPropertyChangedEvent& PropertyC
 	{
 		OnViewModelsUpdated.Broadcast();
 	}
+	if (PropertyChangedEvent.Property && PropertyChangedEvent.Property->GetFName() == GET_MEMBER_NAME_CHECKED(UMVVMBlueprintView, Events))
+	{
+		OnEventsUpdated.Broadcast();
+	}
 }
 
 void UMVVMBlueprintView::PostEditChangeChainProperty(FPropertyChangedChainEvent& PropertyChainEvent)
@@ -497,7 +530,7 @@ void UMVVMBlueprintView::PostEditChangeChainProperty(FPropertyChangedChainEvent&
 	}
 	if (PropertyChainEvent.PropertyChain.Contains(UMVVMBlueprintView::StaticClass()->FindPropertyByName(GET_MEMBER_NAME_CHECKED(UMVVMBlueprintView, Events))))
 	{
-		OnViewModelsUpdated.Broadcast();
+		OnEventsUpdated.Broadcast();
 	}
 }
 
diff --git a/Engine/Plugins/Runtime/ModelViewViewModel/Source/ModelViewViewModelBlueprint/Private/MVVMBlueprintViewEvent.cpp b/Engine/Plugins/Runtime/ModelViewViewModel/Source/ModelViewViewModelBlueprint/Private/MVVMBlueprintViewEvent.cpp
--- a/Engine/Plugins/Runtime/ModelViewViewModel/Source/ModelViewViewModelBlueprint/Private/MVVMBlueprintViewEvent.cpp
+++ b/Engine/Plugins/Runtime/ModelViewViewModel/Source/ModelViewViewModelBlueprint/Private/MVVMBlueprintViewEvent.cpp
@@ -402,7 +402,7 @@ FString UMVVMBlueprintViewEvent::GetSearchableString() const
 void UMVVMBlueprintViewEvent::PostEditChangeChainProperty(FPropertyChangedChainEvent& PropertyChainEvent)
 {
 	Super::PostEditChangeChainProperty(PropertyChainEvent);
-	GetOuterUMVVMBlueprintView()->OnViewModelsUpdated.Broadcast();
+	GetOuterUMVVMBlueprintView()->OnEventsUpdated.Broadcast();
 }
 
 #undef LOCTEXT_NAMESPACE
diff --git a/Engine/Plugins/Runtime/ModelViewViewModel/Source/ModelViewViewModelBlueprint/Public/MVVMBlueprintView.h b/Engine/Plugins/Runtime/ModelViewViewModel/Source/ModelViewViewModelBlueprint/Public/MVVMBlueprintView.h
--- a/Engine/Plugins/Runtime/ModelViewViewModel/Source/ModelViewViewModelBlueprint/Public/MVVMBlueprintView.h
+++ b/Engine/Plugins/Runtime/ModelViewViewModel/Source/ModelViewViewModelBlueprint/Public/MVVMBlueprintView.h
@@ -82,6 +82,7 @@ public:
 	bool RemoveViewModel(FGuid ViewModelId);
 	int32 RemoveViewModels(const TArrayView<FGuid> ViewModelIds);
 	bool RenameViewModel(FName OldViewModelName, FName NewViewModelName);
+	bool ReparentViewModel(FGuid ViewModelId, const UClass* ViewModelClass);
 
 	const TArrayView<const FMVVMBlueprintViewModelContext> GetViewModels() const
 	{
@@ -157,6 +158,9 @@ public:
 	DECLARE_EVENT(UMVVMBlueprintView, FOnBindingsAdded);
 	FOnBindingsAdded OnBindingsAdded;
 
+	DECLARE_EVENT(UMVVMBlueprintView, FOnEventsUpdated);
+	FOnEventsUpdated OnEventsUpdated;
+
 	DECLARE_EVENT(UMVVMBlueprintView, FOnViewModelsUpdated);
 	FOnViewModelsUpdated OnViewModelsUpdated;
 
diff --git a/Engine/Plugins/Runtime/ModelViewViewModel/Source/ModelViewViewModelBlueprint/Public/MVVMBlueprintViewModelContext.h b/Engine/Plugins/Runtime/ModelViewViewModel/Source/ModelViewViewModelBlueprint/Public/MVVMBlueprintViewModelContext.h
--- a/Engine/Plugins/Runtime/ModelViewViewModel/Source/ModelViewViewModelBlueprint/Public/MVVMBlueprintViewModelContext.h
+++ b/Engine/Plugins/Runtime/ModelViewViewModel/Source/ModelViewViewModelBlueprint/Public/MVVMBlueprintViewModelContext.h
@@ -94,7 +94,7 @@ private:
 	FGuid ViewModelContextId;
 
 public:
-	UPROPERTY(VisibleAnywhere, Category = "Viewmodel", NoClear, meta = (DisallowCreateNew, AllowedClasses = "/Script/FieldNotification.NotifyFieldValueChanged", DisallowedClasses = "/Script/UMG.Widget", NoResetToDefault))
+	UPROPERTY(EditAnywhere, Category = "Viewmodel", NoClear, meta = (DisallowCreateNew, AllowedClasses = "/Script/FieldNotification.NotifyFieldValueChanged", DisallowedClasses = "/Script/UMG.Widget", NoResetToDefault))
 	TObjectPtr<UClass> NotifyFieldValueClass = nullptr;
 
 	UPROPERTY()
diff --git a/Engine/Plugins/Runtime/ModelViewViewModel/Source/ModelViewViewModelEditor/Private/Customizations/MVVMBlueprintViewModelContextCustomization.cpp b/Engine/Plugins/Runtime/ModelViewViewModel/Source/ModelViewViewModelEditor/Private/Customizations/MVVMBlueprintViewModelContextCustomization.cpp
--- a/Engine/Plugins/Runtime/ModelViewViewModel/Source/ModelViewViewModelEditor/Private/Customizations/MVVMBlueprintViewModelContextCustomization.cpp
+++ b/Engine/Plugins/Runtime/ModelViewViewModel/Source/ModelViewViewModelEditor/Private/Customizations/MVVMBlueprintViewModelContextCustomization.cpp
@@ -3,7 +3,9 @@
 #include "MVVMBlueprintViewModelContextCustomization.h"
 
 #include "Bindings/MVVMBindingHelper.h"
+#include "DetailLayoutBuilder.h"
 #include "DetailWidgetRow.h"
+#include "Dialogs/Dialogs.h"
 #include "Features/IModularFeatures.h"
 #include "IDetailChildrenBuilder.h"
 #include "IPropertyAccessEditor.h"
@@ -19,6 +21,7 @@
 #include "Widgets/SBoxPanel.h"
 #include "Widgets/Input/SComboButton.h"
 #include "Widgets/Input/SEditableTextBox.h"
+#include "Widgets/SMVVMSelectViewModel.h"
 #include "Widgets/SMVVMViewModelPanel.h"
 
 #define LOCTEXT_NAMESPACE "BlueprintViewModelContextDetailCustomization"
@@ -201,9 +204,29 @@ void FBlueprintViewModelContextDetailCustomization::CustomizeChildren(TSharedRef
 
 	if (ensure(NotifyFieldValueClassHandle))
 	{
-		ChildBuilder.AddProperty(NotifyFieldValueClassHandle.ToSharedRef())
-			.IsEnabled(false)
-			.Visibility(MakeAttributeLambda([ContextPtr](){ return ContextPtr->InstancedViewModel != nullptr ? EVisibility::Collapsed : EVisibility::Visible; }));
+		IDetailPropertyRow& PropertyRow = ChildBuilder.AddProperty(NotifyFieldValueClassHandle.ToSharedRef())
+			.IsEnabled(bCanEdit)
+			.Visibility(MakeAttributeLambda([ContextPtr]() { return ContextPtr->InstancedViewModel != nullptr ? EVisibility::Collapsed : EVisibility::Visible; }));
+
+		TSharedPtr<SWidget> NameWidget, ValueWidget;
+		PropertyRow.GetDefaultWidgets(NameWidget, ValueWidget);
+		PropertyRow.CustomWidget()
+			.NameContent()
+			[
+				NameWidget.ToSharedRef()
+			]
+			.ValueContent()
+			[
+				SAssignNew(NotifyFieldValueClassComboButton, SComboButton)
+					.IsEnabled(bCanEdit)
+					.OnGetMenuContent(this, &FBlueprintViewModelContextDetailCustomization::HandleClassGetMenuContent)
+					.ButtonContent()
+					[
+						SNew(STextBlock)
+							.Text(this, &FBlueprintViewModelContextDetailCustomization::GetClassName)
+							.Font(IDetailLayoutBuilder::GetDetailFont())
+					]
+			];
 	}
 
 	if (ContextPtr->InstancedViewModel == nullptr)
@@ -349,6 +372,29 @@ void FBlueprintViewModelContextDetailCustomization::HandleClassChanged()
 	}
 }
 
+FText FBlueprintViewModelContextDetailCustomization::GetClassName() const
+{
+	UObject* Object = nullptr;
+	FPropertyAccess::Result ValueResult = NotifyFieldValueClassHandle->GetValue(Object);
+	if (ValueResult == FPropertyAccess::Success)
+	{
+		UClass* ViewModelClass = Cast<UClass>(Object);
+		if (ViewModelClass)
+		{
+			return ViewModelClass->GetDisplayNameText();
+		}
+		if (Object)
+		{
+			return FText::FromName(Object->GetFName());
+		}
+	}
+	else if (ValueResult == FPropertyAccess::MultipleValues)
+	{
+		return LOCTEXT("MultipleValues", "Multiple Values");
+	}
+	return LOCTEXT("None", "None");
+}
+
 TSharedRef<SWidget> FBlueprintViewModelContextDetailCustomization::CreateExecutionTypeMenuContent()
 {
 	const bool bCloseAfterSelection = true;
@@ -405,6 +451,69 @@ FText FBlueprintViewModelContextDetailCustomization::GetViewModalNameValueAsText
 	return Result;
 }
 
+TSharedRef<SWidget> FBlueprintViewModelContextDetailCustomization::HandleClassGetMenuContent()
+{
+	return SNew(SBox)
+		.WidthOverride(600)
+		.HeightOverride(500)
+		[
+			SNew(SMVVMSelectViewModel, WidgetBlueprintEditor.Pin()->GetWidgetBlueprintObj())
+			.OnCancel(this, &FBlueprintViewModelContextDetailCustomization::HandleClassCancelMenu)
+			.OnViewModelCommitted(this, &FBlueprintViewModelContextDetailCustomization::HandleClassCommitted)
+			.DisallowedClassFlags(CLASS_HideDropDown | CLASS_Hidden | CLASS_Deprecated | CLASS_NotPlaceable)
+		];
+}
+
+void FBlueprintViewModelContextDetailCustomization::HandleClassCancelMenu()
+{
+	if (NotifyFieldValueClassComboButton)
+	{
+		NotifyFieldValueClassComboButton->SetIsOpen(false, false);
+	}
+}
+
+void FBlueprintViewModelContextDetailCustomization::HandleClassCommitted(const UClass* SelectedClass)
+{
+	if (NotifyFieldValueClassComboButton)
+	{
+		NotifyFieldValueClassComboButton->SetIsOpen(false, false);
+	}
+	bool bReparent = false;
+	FName ViewModelName;
+	{
+		UObject* Object = nullptr;
+		FPropertyAccess::Result ClassValueResult = NotifyFieldValueClassHandle->GetValue(Object);
+		UClass* PreviousClass = Cast<UClass>(Object);
+		FPropertyAccess::Result NameValueResult = ViewModelNameHandle->GetValue(ViewModelName);
+		if (ClassValueResult == FPropertyAccess::Success && SelectedClass && SelectedClass != PreviousClass
+			&& NameValueResult == FPropertyAccess::Success && !ViewModelName.IsNone())
+		{
+			const FText Title = LOCTEXT("ReparentTitle", "Reparent Viewmodel");
+			const FText Message = LOCTEXT("ReparentWarning", "Reparenting the viewmodel may cause data loss. Continue reparenting?");
+
+			// Warn the user that this may result in data loss
+			FSuppressableWarningDialog::FSetupInfo Info(Message, Title, "Warning_ReparentTitle");
+			Info.ConfirmText = LOCTEXT("ReparentYesButton", "Reparent");
+			Info.CancelText = LOCTEXT("ReparentNoButton", "Cancel");
+			Info.CheckBoxText = FText::GetEmpty();	// not suppressible
+
+			FSuppressableWarningDialog ReparentBlueprintDlg(Info);
+			if (ReparentBlueprintDlg.ShowModal() == FSuppressableWarningDialog::Confirm)
+			{
+				bReparent = true;
+			}
+		}
+	}
+
+	if (bReparent)
+	{
+		UMVVMEditorSubsystem* EditorSubsystem = GEditor->GetEditorSubsystem<UMVVMEditorSubsystem>();
+		check(EditorSubsystem);
+		FText ErrorMessage;
+		EditorSubsystem->ReparentViewModel(WidgetBlueprintEditor.Pin()->GetWidgetBlueprintObj(), ViewModelName, SelectedClass, ErrorMessage);
+	}
+}
+
 namespace Private
 {
 bool VerifyViewModelName(TSharedPtr<FWidgetBlueprintEditor> WidgetBlueprintEditor, TSharedPtr<IPropertyHandle> ViewModelNameHandle, const FText& RenameTo, bool bCommit, FText& OutErrorMessage)
diff --git a/Engine/Plugins/Runtime/ModelViewViewModel/Source/ModelViewViewModelEditor/Private/Customizations/MVVMBlueprintViewModelContextCustomization.h b/Engine/Plugins/Runtime/ModelViewViewModel/Source/ModelViewViewModelEditor/Private/Customizations/MVVMBlueprintViewModelContextCustomization.h
--- a/Engine/Plugins/Runtime/ModelViewViewModel/Source/ModelViewViewModelEditor/Private/Customizations/MVVMBlueprintViewModelContextCustomization.h
+++ b/Engine/Plugins/Runtime/ModelViewViewModel/Source/ModelViewViewModelEditor/Private/Customizations/MVVMBlueprintViewModelContextCustomization.h
@@ -7,6 +7,7 @@
 
 class FWidgetBlueprintEditor;
 struct FBindingChainElement;
+class SComboButton;
 class SWidget;
 enum class EMVVMBlueprintViewModelContextCreationType : uint8;
 
@@ -50,13 +51,18 @@ private:
 	TSharedRef<SWidget> CreateExecutionTypeMenuContent();
 	FText GetCreationTypeValue() const;
 	FText GetExecutionTypeValueToolTip() const;
+	FText GetClassName() const;
 	void HandleClassChanged();
+	TSharedRef<SWidget> HandleClassGetMenuContent();
+	void HandleClassCancelMenu();
+	void HandleClassCommitted(const UClass* SelectedClass);
 	FText GetViewModalNameValueAsText() const;
 	void HandleNameTextCommitted(const FText& NewText, ETextCommit::Type /*CommitInfo*/);
 	bool HandleNameVerifyTextChanged(const FText& Text, FText& OutError) const;
 
 	FViewModelPropertyAccessEditor PropertyAccessEditor;
 	TWeakPtr<FWidgetBlueprintEditor> WidgetBlueprintEditor;
+	TSharedPtr<SComboButton> NotifyFieldValueClassComboButton;
 	TSharedPtr<IPropertyHandle> NotifyFieldValueClassHandle;
 	TSharedPtr<IPropertyHandle> PropertyPathHandle;
 	TSharedPtr<IPropertyHandle> CreationTypeHandle;
diff --git a/Engine/Plugins/Runtime/ModelViewViewModel/Source/ModelViewViewModelEditor/Private/MVVMEditorSubsystem.cpp b/Engine/Plugins/Runtime/ModelViewViewModel/Source/ModelViewViewModelEditor/Private/MVVMEditorSubsystem.cpp
--- a/Engine/Plugins/Runtime/ModelViewViewModel/Source/ModelViewViewModelEditor/Private/MVVMEditorSubsystem.cpp
+++ b/Engine/Plugins/Runtime/ModelViewViewModel/Source/ModelViewViewModelEditor/Private/MVVMEditorSubsystem.cpp
@@ -156,7 +156,7 @@ UMVVMBlueprintView* UMVVMEditorSubsystem::GetView(const UWidgetBlueprint* Widget
 FGuid UMVVMEditorSubsystem::AddViewModel(UWidgetBlueprint* WidgetBlueprint, const UClass* ViewModelClass)
 {
 	FGuid Result;
-	if (ViewModelClass)
+	if (ViewModelClass && ViewModelClass->ImplementsInterface(UNotifyFieldValueChanged::StaticClass()))
 	{
 		if (UMVVMBlueprintView* View = GetView(WidgetBlueprint))
 		{
@@ -289,6 +289,19 @@ bool UMVVMEditorSubsystem::RenameViewModel(UWidgetBlueprint* WidgetBlueprint, FN
 	return false;
 }
 
+bool UMVVMEditorSubsystem::ReparentViewModel(UWidgetBlueprint* WidgetBlueprint, FName ViewModel, const UClass* ViewModelClass, FText& OutError)
+{
+	if (UMVVMBlueprintView* View = GetView(WidgetBlueprint))
+	{
+		if (const FMVVMBlueprintViewModelContext* ViewModelContext = View->FindViewModel(ViewModel))
+		{
+			FScopedTransaction Transaction(LOCTEXT("ReparentViewmodel", "Reparent Viewmodel"));
+			return View->ReparentViewModel(ViewModelContext->GetViewModelId(), ViewModelClass);
+		}
+	}
+	return false;
+}
+
 FMVVMBlueprintViewBinding& UMVVMEditorSubsystem::AddBinding(UWidgetBlueprint* WidgetBlueprint)
 {
 	UMVVMBlueprintView* View = RequestView(WidgetBlueprint);
@@ -334,7 +347,7 @@ UFunction* UMVVMEditorSubsystem::GetConversionFunction(const UWidgetBlueprint* W
 	return nullptr;
 }
 
-UEdGraphPin* UMVVMEditorSubsystem::GetConversionFunctionArgumentPin(const UWidgetBlueprint* WidgetBlueprint, const FMVVMBlueprintViewBinding& Binding, FName ParameterName, bool bSourceToDestination)
+UEdGraphPin* UMVVMEditorSubsystem::GetConversionFunctionArgumentPin(const UWidgetBlueprint* WidgetBlueprint, const FMVVMBlueprintViewBinding& Binding, FName ParameterName, bool bSourceToDestination) const
 {
 	if (UMVVMBlueprintViewConversionFunction* ConversionFunction = Binding.Conversion.GetConversionFunction(bSourceToDestination))
 	{
@@ -878,6 +891,128 @@ void UMVVMEditorSubsystem::SetPathForConversionFunctionArgument(UWidgetBlueprint
 	FBlueprintEditorUtils::MarkBlueprintAsModified(WidgetBlueprint);
 }
 
+namespace Private
+{
+	UEdGraphPin* GetGraphPin(const UMVVMEditorSubsystem* Subsystem, const UWidgetBlueprint* WidgetBlueprint, const FMVVMBlueprintViewBinding& ViewBinding, FName PinName, bool bSourceToDestination)
+	{
+		if (WidgetBlueprint == nullptr)
+		{
+			return nullptr;
+		}
+		return Subsystem->GetConversionFunctionArgumentPin(WidgetBlueprint, ViewBinding, PinName, bSourceToDestination);
+	}
+
+	UEdGraphPin* GetGraphPin(const UWidgetBlueprint* WidgetBlueprint, UMVVMBlueprintViewEvent* ViewEvent, FName PinName)
+	{
+		if (WidgetBlueprint == nullptr || ViewEvent == nullptr)
+		{
+			return nullptr;
+		}
+
+		return ViewEvent->GetOrCreateGraphPin(PinName);
+	}
+}
+
+void UMVVMEditorSubsystem::SplitPin(UWidgetBlueprint* WidgetBlueprint, FMVVMBlueprintViewBinding& Binding, FName PinName, bool bSourceToDestination) const
+{
+	UEdGraphPin* GraphPin = Private::GetGraphPin(this, WidgetBlueprint, Binding, PinName, bSourceToDestination);
+	const UEdGraphSchema_K2* K2Schema = GetDefault<UEdGraphSchema_K2>();
+	if (GraphPin == nullptr || !K2Schema->CanSplitStructPin(*GraphPin))
+	{
+		return;
+	}
+
+	UMVVMBlueprintView* View = GetView(WidgetBlueprint);
+	UMVVMBlueprintViewConversionFunction* ConversionFunction = Binding.Conversion.GetConversionFunction(bSourceToDestination);
+
+	FScopedTransaction Transaction(LOCTEXT("BreakPin", "Split Struct Pin"));
+	UE::MVVM::Private::OnBindingPreEditChange(View, GET_MEMBER_NAME_CHECKED(FMVVMBlueprintViewBinding, Conversion));
+	K2Schema->SplitPin(GraphPin);
+	ConversionFunction->SavePinValues(WidgetBlueprint);
+	UE::MVVM::Private::OnBindingPostEditChange(View, GET_MEMBER_NAME_CHECKED(FMVVMBlueprintViewBinding, Conversion));
+}
+
+bool UMVVMEditorSubsystem::CanSplitPin(const UWidgetBlueprint* WidgetBlueprint, const FMVVMBlueprintViewBinding& Binding, FName PinName, bool bSourceToDestination) const
+{
+	UEdGraphPin* GraphPin = Private::GetGraphPin(this, WidgetBlueprint, Binding, PinName, bSourceToDestination);
+	const UEdGraphSchema_K2* K2Schema = GetDefault<UEdGraphSchema_K2>();
+	return GraphPin ? K2Schema->CanSplitStructPin(*GraphPin) : false;
+}
+
+void UMVVMEditorSubsystem::SplitPin(UWidgetBlueprint* WidgetBlueprint, UMVVMBlueprintViewEvent* ViewEvent, FName PinName) const
+{
+	UEdGraphPin* GraphPin = Private::GetGraphPin(WidgetBlueprint, ViewEvent, PinName);
+	const UEdGraphSchema_K2* K2Schema = GetDefault<UEdGraphSchema_K2>();
+	if (GraphPin == nullptr || !K2Schema->CanSplitStructPin(*GraphPin))
+	{
+		return;
+	}
+
+	FName NAME_SavedPin = "SavedPins";
+	FScopedTransaction Transaction(LOCTEXT("BreakPin", "Split Struct Pin"));
+	UE::MVVM::Private::OnEventPreEditChange(ViewEvent, NAME_SavedPin);
+	K2Schema->SplitPin(GraphPin);
+	ViewEvent->SavePinValues();
+	UE::MVVM::Private::OnEventPostEditChange(ViewEvent, NAME_SavedPin);
+}
+
+bool UMVVMEditorSubsystem::CanSplitPin(const UWidgetBlueprint* WidgetBlueprint, UMVVMBlueprintViewEvent* Event, FName PinName) const
+{
+	UEdGraphPin* GraphPin = Private::GetGraphPin(WidgetBlueprint, Event, PinName);
+	const UEdGraphSchema_K2* K2Schema = GetDefault<UEdGraphSchema_K2>();
+	return GraphPin ? K2Schema->CanSplitStructPin(*GraphPin) : false;
+}
+
+void UMVVMEditorSubsystem::RecombinePin(UWidgetBlueprint* WidgetBlueprint, FMVVMBlueprintViewBinding& Binding, FName PinName, bool bSourceToDestination) const
+{
+	UEdGraphPin* GraphPin = Private::GetGraphPin(this, WidgetBlueprint, Binding, PinName, bSourceToDestination);
+	const UEdGraphSchema_K2* K2Schema = GetDefault<UEdGraphSchema_K2>();
+	if (GraphPin == nullptr || !K2Schema->CanRecombineStructPin(*GraphPin))
+	{
+		return;
+	}
+
+	UMVVMBlueprintView* View = GetView(WidgetBlueprint);
+	UMVVMBlueprintViewConversionFunction* ConversionFunction = Binding.Conversion.GetConversionFunction(bSourceToDestination);
+
+	FScopedTransaction Transaction(LOCTEXT("RecombinePin", "Recombine Struct Pin"));
+	UE::MVVM::Private::OnBindingPreEditChange(View, GET_MEMBER_NAME_CHECKED(FMVVMBlueprintViewBinding, Conversion));
+	K2Schema->RecombinePin(GraphPin);
+	ConversionFunction->SavePinValues(WidgetBlueprint);
+	UE::MVVM::Private::OnBindingPostEditChange(View, GET_MEMBER_NAME_CHECKED(FMVVMBlueprintViewBinding, Conversion));
+}
+
+bool UMVVMEditorSubsystem::CanRecombinePin(const UWidgetBlueprint* WidgetBlueprint, const FMVVMBlueprintViewBinding& Binding, FName PinName, bool bSourceToDestination) const
+{
+	UEdGraphPin* GraphPin = Private::GetGraphPin(this, WidgetBlueprint, Binding, PinName, bSourceToDestination);
+	const UEdGraphSchema_K2* K2Schema = GetDefault<UEdGraphSchema_K2>();
+	return GraphPin ? K2Schema->CanRecombineStructPin(*GraphPin) : false;
+}
+
+void UMVVMEditorSubsystem::RecombinePin(UWidgetBlueprint* WidgetBlueprint, UMVVMBlueprintViewEvent* ViewEvent, FName PinName) const
+{
+	UEdGraphPin* GraphPin = Private::GetGraphPin(WidgetBlueprint, ViewEvent, PinName);
+	const UEdGraphSchema_K2* K2Schema = GetDefault<UEdGraphSchema_K2>();
+	if (GraphPin == nullptr || !K2Schema->CanRecombineStructPin(*GraphPin))
+	{
+		return;
+	}
+
+	FName NAME_SavedPin = "SavedPins";
+	FScopedTransaction Transaction(LOCTEXT("RecombinePin", "Recombine Struct Pin"));
+	UE::MVVM::Private::OnEventPreEditChange(ViewEvent, NAME_SavedPin);
+	K2Schema->RecombinePin(GraphPin);
+	ViewEvent->SavePinValues();
+	UE::MVVM::Private::OnEventPostEditChange(ViewEvent, NAME_SavedPin);
+}
+
+bool UMVVMEditorSubsystem::CanRecombinePin(const UWidgetBlueprint* WidgetBlueprint, UMVVMBlueprintViewEvent* Event, FName PinName) const
+{
+	UEdGraphPin* GraphPin = Private::GetGraphPin(WidgetBlueprint, Event, PinName);
+	const UEdGraphSchema_K2* K2Schema = GetDefault<UEdGraphSchema_K2>();
+	return GraphPin ? K2Schema->CanRecombineStructPin(*GraphPin) : false;
+}
+
 TArray<UE::MVVM::FBindingSource> UMVVMEditorSubsystem::GetBindableWidgets(const UWidgetBlueprint* WidgetBlueprint) const
 {
 	TArray<UE::MVVM::FBindingSource> Sources;
diff --git a/Engine/Plugins/Runtime/ModelViewViewModel/Source/ModelViewViewModelEditor/Private/Widgets/SMVVMCachedViewBindingPropertyPath.cpp b/Engine/Plugins/Runtime/ModelViewViewModel/Source/ModelViewViewModelEditor/Private/Widgets/SMVVMCachedViewBindingPropertyPath.cpp
--- a/Engine/Plugins/Runtime/ModelViewViewModel/Source/ModelViewViewModelEditor/Private/Widgets/SMVVMCachedViewBindingPropertyPath.cpp
+++ b/Engine/Plugins/Runtime/ModelViewViewModel/Source/ModelViewViewModelEditor/Private/Widgets/SMVVMCachedViewBindingPropertyPath.cpp
@@ -20,17 +20,13 @@ void SCachedViewBindingPropertyPath::Construct(const FArguments& InArgs, const U
 
 	WidgetBlueprint = InWidgetBlueprint;
 	OnGetPropertyPath = InArgs._OnGetPropertyPath;
-
-	const FMVVMBlueprintPropertyPath PropertyPath = OnGetPropertyPath.Execute();
-	CachedPropertyPath = PropertyPath.GetFields(InWidgetBlueprint->SkeletonGeneratedClass);
-	CachedWidgetName = PropertyPath.GetWidgetName();
-	CachedContextId = PropertyPath.GetViewModelId();
+	CachedPropertyPath = OnGetPropertyPath.Execute();
 
 	ChildSlot
 	[
 		SAssignNew(PropertyPathWidget, SPropertyPath, InWidgetBlueprint)
 		.TextStyle(InArgs._TextStyle)
-		.PropertyPath(PropertyPath)
+		.PropertyPath(CachedPropertyPath)
 		.ShowContext(InArgs._ShowContext)
 		.ShowOnlyLastPath(InArgs._ShowOnlyLastPath)
 	];
@@ -40,14 +36,11 @@ void SCachedViewBindingPropertyPath::Tick(const FGeometry& AllottedGeometry, con
 {
 	Super::Tick(AllottedGeometry, InCurrentTime, InDeltaTime);
 
-	const FMVVMBlueprintPropertyPath NewPropertyPath = OnGetPropertyPath.Execute();
-	TArray<FMVVMConstFieldVariant> NewCachedPropertyPath = NewPropertyPath.GetFields(WidgetBlueprint.Get()->SkeletonGeneratedClass);
-	if (NewCachedPropertyPath != CachedPropertyPath || NewPropertyPath.GetWidgetName() != CachedWidgetName || NewPropertyPath.GetViewModelId() != CachedContextId)
+	FMVVMBlueprintPropertyPath NewPropertyPath = OnGetPropertyPath.Execute();
+	if (NewPropertyPath != CachedPropertyPath)
 	{
-		CachedPropertyPath = MoveTemp(NewCachedPropertyPath);
-		CachedWidgetName = NewPropertyPath.GetWidgetName();
-		CachedContextId = NewPropertyPath.GetViewModelId();
 		PropertyPathWidget->SetPropertyPath(NewPropertyPath);
+		CachedPropertyPath = MoveTemp(NewPropertyPath);
 	}
 }
 
diff --git a/Engine/Plugins/Runtime/ModelViewViewModel/Source/ModelViewViewModelEditor/Private/Widgets/SMVVMCachedViewBindingPropertyPath.h b/Engine/Plugins/Runtime/ModelViewViewModel/Source/ModelViewViewModelEditor/Private/Widgets/SMVVMCachedViewBindingPropertyPath.h
--- a/Engine/Plugins/Runtime/ModelViewViewModel/Source/ModelViewViewModelEditor/Private/Widgets/SMVVMCachedViewBindingPropertyPath.h
+++ b/Engine/Plugins/Runtime/ModelViewViewModel/Source/ModelViewViewModelEditor/Private/Widgets/SMVVMCachedViewBindingPropertyPath.h
@@ -39,9 +39,7 @@ private:
 	TWeakObjectPtr<const UWidgetBlueprint> WidgetBlueprint;
 	TSharedPtr<SPropertyPath> PropertyPathWidget;
 	FOnGetPropertyPath OnGetPropertyPath;
-	TArray<FMVVMConstFieldVariant> CachedPropertyPath;
-	FName CachedWidgetName;
-	FGuid CachedContextId;
+	FMVVMBlueprintPropertyPath CachedPropertyPath;
 };
 
 } // namespace UE::MVVM
diff --git a/Engine/Plugins/Runtime/ModelViewViewModel/Source/ModelViewViewModelEditor/Private/Widgets/SMVVMFunctionParameter.cpp b/Engine/Plugins/Runtime/ModelViewViewModel/Source/ModelViewViewModelEditor/Private/Widgets/SMVVMFunctionParameter.cpp
--- a/Engine/Plugins/Runtime/ModelViewViewModel/Source/ModelViewViewModelEditor/Private/Widgets/SMVVMFunctionParameter.cpp
+++ b/Engine/Plugins/Runtime/ModelViewViewModel/Source/ModelViewViewModelEditor/Private/Widgets/SMVVMFunctionParameter.cpp
@@ -45,12 +45,8 @@ void SFunctionParameter::Construct(const FArguments& InArgs, UWidgetBlueprint* I
 	const FMVVMBlueprintViewBinding* Binding = View->GetBinding(BindingId);
 	check(Binding);
 
-	const UFunction* ConversionFunction = EditorSubsystem->GetConversionFunction(InWidgetBlueprint, *Binding, bSourceToDestination);
-	const FProperty* Property = ConversionFunction->FindPropertyByName(ParameterName);
-	check(Property);
-
 	TSharedRef<SWidget> ValueWidget = SNullWidget::NullWidget;
-	bool bIsBooleanPin = CastField<const FBoolProperty>(Property) != nullptr;
+	bool bIsBooleanPin = false;
 
 	if (UEdGraphPin* Pin = EditorSubsystem->GetConversionFunctionArgumentPin(InWidgetBlueprint, *Binding, ParameterName, bSourceToDestination))
 	{
diff --git a/Engine/Plugins/Runtime/ModelViewViewModel/Source/ModelViewViewModelEditor/Private/Widgets/SMVVMViewBindingListView.cpp b/Engine/Plugins/Runtime/ModelViewViewModel/Source/ModelViewViewModelEditor/Private/Widgets/SMVVMViewBindingListView.cpp
--- a/Engine/Plugins/Runtime/ModelViewViewModel/Source/ModelViewViewModelEditor/Private/Widgets/SMVVMViewBindingListView.cpp
+++ b/Engine/Plugins/Runtime/ModelViewViewModel/Source/ModelViewViewModelEditor/Private/Widgets/SMVVMViewBindingListView.cpp
@@ -10,6 +10,7 @@
 #include "DragAndDrop/DecoratedDragDropOp.h"
 #include "Framework/MultiBox/MultiBoxBuilder.h"
 #include "Framework/Views/TableViewMetadata.h"
+#include "Kismet2/BlueprintEditorUtils.h"
 #include "Misc/MessageDialog.h"
 
 #include "MVVMBlueprintView.h"
@@ -1915,8 +1916,9 @@ void SBindingsList::Construct(const FArguments& InArgs, TSharedPtr<SBindingsPane
 
 	MVVMExtension->OnBlueprintViewChangedDelegate().AddSP(this, &SBindingsList::Refresh);
 	MVVMExtension->GetBlueprintView()->OnBindingsUpdated.AddSP(this, &SBindingsList::Refresh);
+	MVVMExtension->GetBlueprintView()->OnEventsUpdated.AddSP(this, &SBindingsList::Refresh);
 	MVVMExtension->GetBlueprintView()->OnBindingsAdded.AddSP(this, &SBindingsList::ClearFilterText);
-	MVVMExtension->GetBlueprintView()->OnViewModelsUpdated.AddSP(this, &SBindingsList::Refresh);
+	MVVMExtension->GetBlueprintView()->OnViewModelsUpdated.AddSP(this, &SBindingsList::ForceRefresh);
 
 	ChildSlot
 	[
@@ -1939,6 +1941,8 @@ SBindingsList::~SBindingsList()
 	{
 		MVVMExtensionPtr->OnBlueprintViewChangedDelegate().RemoveAll(this);
 		MVVMExtensionPtr->GetBlueprintView()->OnBindingsUpdated.RemoveAll(this);
+		MVVMExtensionPtr->GetBlueprintView()->OnEventsUpdated.RemoveAll(this);
+		MVVMExtensionPtr->GetBlueprintView()->OnBindingsAdded.RemoveAll(this);
 		MVVMExtensionPtr->GetBlueprintView()->OnViewModelsUpdated.RemoveAll(this);
 	}
 }
@@ -2079,33 +2083,68 @@ void SBindingsList::Refresh()
 			}
 
 			// Create/Find entries for conversion function parameters
+			if (UMVVMBlueprintViewConversionFunction* ConversionFunction = Binding.Conversion.GetConversionFunction(UE::MVVM::IsForwardBinding(Binding.BindingType)))
 			{
-				const UFunction* Function = GEditor->GetEditorSubsystem<UMVVMEditorSubsystem>()->GetConversionFunction(MVVMExtensionPtr->GetWidgetBlueprint(), Binding, UE::MVVM::IsForwardBinding(Binding.BindingType));
-				if (Function != nullptr)
+				ConversionFunction->GetOrCreateWrapperGraph(MVVMExtensionPtr->GetWidgetBlueprint());
+				if (ConversionFunction->NeedsWrapperGraph())
 				{
-					TValueOrError<TArray<const FProperty*>, FText> ArgumentsResult = BindingHelper::TryGetArgumentsForConversionFunction(Function);
-					if (ArgumentsResult.HasValue())
+					for (const FMVVMBlueprintPin& Pin : ConversionFunction->GetPins())
 					{
-						for (const FProperty* Argument : ArgumentsResult.GetValue())
+						UEdGraphPin* GraphPin = ConversionFunction->GetOrCreateGraphPin(MVVMExtensionPtr->GetWidgetBlueprint(), Pin.GetName());
+						if (GraphPin && GraphPin->bHidden)
 						{
-							TSharedPtr<FBindingEntry> ArgumentEntry;
-							if (PreviousGroupEntry)
+							continue;
+						}
+
+						TSharedPtr<FBindingEntry> ArgumentEntry;
+						if (PreviousGroupEntry)
+						{
+							if (TSharedPtr<FBindingEntry>* FoundParameter = PreviousGroupEntry->Children.FindByPredicate([BindingId, ArgumentName = Pin.GetName()](const TSharedPtr<FBindingEntry>& Other)
+								{ return Other->GetBindingId() == BindingId && Other->GetRowType() == FBindingEntry::ERowType::BindingParameter && Other->GetBindingParameterName() == ArgumentName; }))
 							{
-								if (TSharedPtr<FBindingEntry>* FoundParameter = PreviousGroupEntry->Children.FindByPredicate([BindingId, ArgumentName = Argument->GetFName()](const TSharedPtr<FBindingEntry>& Other)
-									{ return Other->GetBindingId() == BindingId && Other->GetRowType() == FBindingEntry::ERowType::BindingParameter && Other->GetBindingParameterName() == ArgumentName; }))
-								{
-									ArgumentEntry = *FoundParameter;
-								}
+								ArgumentEntry = *FoundParameter;
 							}
+						}
 
-							if (!ArgumentEntry.IsValid())
+						if (!ArgumentEntry.IsValid())
+						{
+							ArgumentEntry = MakeShared<FBindingEntry>();
+							ArgumentEntry->SetBindingParameterName(Binding.BindingId, Pin.GetName());
+
+							NewEntries.Add(ArgumentEntry);
+						}
+						BindingEntry->AddChild(ArgumentEntry);
+					}
+				}
+				else // simple conversion function with one parameters (no subpins)
+				{
+					const UFunction* Function = GEditor->GetEditorSubsystem<UMVVMEditorSubsystem>()->GetConversionFunction(MVVMExtensionPtr->GetWidgetBlueprint(), Binding, UE::MVVM::IsForwardBinding(Binding.BindingType));
+					if (Function != nullptr)
+					{
+						TValueOrError<TArray<const FProperty*>, FText> ArgumentsResult = BindingHelper::TryGetArgumentsForConversionFunction(Function);
+						if (ArgumentsResult.HasValue())
+						{
+							for (const FProperty* Argument : ArgumentsResult.GetValue())
 							{
-								ArgumentEntry = MakeShared<FBindingEntry>();
-								ArgumentEntry->SetBindingParameterName(Binding.BindingId, Argument->GetFName());
+								TSharedPtr<FBindingEntry> ArgumentEntry;
+								if (PreviousGroupEntry)
+								{
+									if (TSharedPtr<FBindingEntry>* FoundParameter = PreviousGroupEntry->Children.FindByPredicate([BindingId, ArgumentName = Argument->GetFName()](const TSharedPtr<FBindingEntry>& Other)
+										{ return Other->GetBindingId() == BindingId && Other->GetRowType() == FBindingEntry::ERowType::BindingParameter && Other->GetBindingParameterName() == ArgumentName; }))
+									{
+										ArgumentEntry = *FoundParameter;
+									}
+								}
 
-								NewEntries.Add(ArgumentEntry);
+								if (!ArgumentEntry.IsValid())
+								{
+									ArgumentEntry = MakeShared<FBindingEntry>();
+									ArgumentEntry->SetBindingParameterName(Binding.BindingId, Argument->GetFName());
+
+									NewEntries.Add(ArgumentEntry);
+								}
+								BindingEntry->AddChild(ArgumentEntry);
 							}
-							BindingEntry->AddChild(ArgumentEntry);
 						}
 					}
 				}
@@ -2166,6 +2205,12 @@ void SBindingsList::Refresh()
 			// Create/Find entries for function parameters
 			for (const FMVVMBlueprintPin& Pin : Event->GetPins())
 			{
+				UEdGraphPin* GraphPin = Event->GetOrCreateGraphPin(Pin.GetName());
+				if (GraphPin && GraphPin->bHidden)
+				{
+					continue;
+				}
+
 				TSharedPtr<FBindingEntry> ArgumentEntry;
 				if (PreviousGroupEntry)
 				{
@@ -2200,6 +2245,13 @@ void SBindingsList::Refresh()
 	}
 }
 
+void SBindingsList::ForceRefresh()
+{
+	AllRootGroups.Reset();
+	FilteredRootGroups.Reset();
+	Refresh();
+}
+
 TSharedRef<ITableRow> SBindingsList::GenerateEntryRow(TSharedPtr<FBindingEntry> Entry, const TSharedRef<STableViewBase>& OwnerTable) const
 {
 	TSharedPtr<ITableRow> Row;
@@ -2275,9 +2327,9 @@ namespace Private
 			GatherAllChildBindings(BlueprintView, Entry->GetAllChildren(), OutBindings, OutEvents);
 		}
 	}
-}
+} //namespace
 
-void SBindingsList::OnDeleteSelected()
+void SBindingsList::HandleDeleteSelected()
 {
 	TArray<TSharedPtr<FBindingEntry>> Selection = TreeView->GetSelectedItems();
 	if (Selection.Num() == 0)
@@ -2342,6 +2394,62 @@ void SBindingsList::OnDeleteSelected()
 	}
 }
 
+void SBindingsList::HandleBreakSelectedPin()
+{
+	if (UMVVMWidgetBlueprintExtension_View* MVVMExtensionPtr = MVVMExtension.Get())
+	{
+		UWidgetBlueprint* WidgetBlueprint = MVVMExtensionPtr->GetWidgetBlueprint();
+		check(WidgetBlueprint);
+		UMVVMEditorSubsystem* EditorSubsystem = GEditor->GetEditorSubsystem<UMVVMEditorSubsystem>();
+
+		TArray<TSharedPtr<FBindingEntry>> Selection = TreeView->GetSelectedItems();
+		for (const TSharedPtr<FBindingEntry>& Entry : Selection)
+		{
+			if (Entry->GetRowType() == FBindingEntry::ERowType::EventParameter)
+			{
+				EditorSubsystem->SplitPin(WidgetBlueprint, Entry->GetEvent(), Entry->GetEventParameterName());
+			}
+			else if(Entry->GetRowType() == FBindingEntry::ERowType::BindingParameter)
+			{
+				if (FMVVMBlueprintViewBinding* Binding = Entry->GetBinding(MVVMExtensionPtr->GetBlueprintView()))
+				{
+					const bool bSourceToDestination = UE::MVVM::IsForwardBinding(Binding->BindingType);
+					EditorSubsystem->SplitPin(WidgetBlueprint, *Binding, Entry->GetEventParameterName(), bSourceToDestination);
+				}
+			}
+		}
+	}
+}
+
+void SBindingsList::HandleRecombineSelectedPin()
+{
+	if (UMVVMWidgetBlueprintExtension_View* MVVMExtensionPtr = MVVMExtension.Get())
+	{
+		const FScopedTransaction Transaction(LOCTEXT("RecombineStructPin", "Recombine Struct Pin"));
+
+		UWidgetBlueprint* WidgetBlueprint = MVVMExtensionPtr->GetWidgetBlueprint();
+		check(WidgetBlueprint);
+		UMVVMEditorSubsystem* EditorSubsystem = GEditor->GetEditorSubsystem<UMVVMEditorSubsystem>();
+
+		TArray<TSharedPtr<FBindingEntry>> Selection = TreeView->GetSelectedItems();
+		for (const TSharedPtr<FBindingEntry>& Entry : Selection)
+		{
+			if (Entry->GetRowType() == FBindingEntry::ERowType::EventParameter)
+			{
+				EditorSubsystem->RecombinePin(WidgetBlueprint, Entry->GetEvent(), Entry->GetEventParameterName());
+			}
+			else if (Entry->GetRowType() == FBindingEntry::ERowType::BindingParameter)
+			{
+				if (FMVVMBlueprintViewBinding* Binding = Entry->GetBinding(MVVMExtensionPtr->GetBlueprintView()))
+				{
+					const bool bSourceToDestination = UE::MVVM::IsForwardBinding(Binding->BindingType);
+					EditorSubsystem->RecombinePin(WidgetBlueprint, *Binding, Entry->GetEventParameterName(), bSourceToDestination);
+				}
+			}
+		}
+	}
+}
+
 TSharedPtr<SWidget> SBindingsList::OnSourceConstructContextMenu()
 {
 	const bool bShouldCloseWindowAfterMenuSelection = true;
@@ -2350,12 +2458,101 @@ TSharedPtr<SWidget> SBindingsList::OnSourceConstructContextMenu()
 	TArray<TSharedPtr<FBindingEntry>> Selection = TreeView->GetSelectedItems();
 	if (Selection.Num() > 0)
 	{
-		FUIAction RemoveAction;
-		RemoveAction.ExecuteAction = FExecuteAction::CreateSP(this, &SBindingsList::OnDeleteSelected);
-		MenuBuilder.AddMenuEntry(LOCTEXT("RemoveBinding", "Remove Binding"), 
-			LOCTEXT("RemoveBindingTooltip", "Remove this binding."), 
-			FSlateIcon(FAppStyle::GetAppStyleSetName(), "Icons.Delete"), 
-			RemoveAction);
+		{
+			FUIAction RemoveAction;
+			RemoveAction.ExecuteAction = FExecuteAction::CreateSP(this, &SBindingsList::HandleDeleteSelected);
+			MenuBuilder.AddMenuEntry(LOCTEXT("RemoveBinding", "Remove Binding"),
+				LOCTEXT("RemoveBindingTooltip", "Remove this binding."),
+				FSlateIcon(FAppStyle::GetAppStyleSetName(), "Icons.Delete"),
+				RemoveAction);
+		}
+
+		{
+			const UMVVMWidgetBlueprintExtension_View* MVVMExtensionPtr = MVVMExtension.Get();
+			const UWidgetBlueprint* WidgetBlueprint = MVVMExtensionPtr ? MVVMExtensionPtr->GetWidgetBlueprint() : nullptr;
+			const UMVVMBlueprintView* View = MVVMExtensionPtr ? MVVMExtensionPtr->GetBlueprintView() : nullptr;
+			const UMVVMEditorSubsystem* EditorSubsystem = GEditor->GetEditorSubsystem<UMVVMEditorSubsystem>();
+
+			if (View && WidgetBlueprint)
+			{
+				bool bCanSplitPin = true;
+				bool bCanRecombinePin = true;
+				bool bCanRecombinePinVisible = true;
+				for (const TSharedPtr<FBindingEntry>& Entry : Selection)
+				{
+					if (Entry->GetRowType() == FBindingEntry::ERowType::EventParameter)
+					{
+						if (!EditorSubsystem->CanSplitPin(WidgetBlueprint, Entry->GetEvent(), Entry->GetEventParameterName()))
+						{
+							bCanSplitPin = false;
+						}
+						if (!EditorSubsystem->CanRecombinePin(WidgetBlueprint, Entry->GetEvent(), Entry->GetEventParameterName()))
+						{
+							bCanRecombinePin = false;
+						}
+						UMVVMBlueprintViewEvent* ViewEvent = Entry->GetEvent();
+						UEdGraphPin* GraphPin = ViewEvent ? ViewEvent->GetOrCreateGraphPin(Entry->GetEventParameterName()) : nullptr;
+						if (GraphPin == nullptr || GraphPin->ParentPin == nullptr)
+						{
+							bCanRecombinePinVisible = false;
+						}
+					}
+					else if (Entry->GetRowType() == FBindingEntry::ERowType::BindingParameter)
+					{
+						const FMVVMBlueprintViewBinding* Binding = Entry->GetBinding(View);
+						if (Binding)
+						{
+							const bool bSourceToDestination = UE::MVVM::IsForwardBinding(Binding->BindingType);
+							if (!EditorSubsystem->CanSplitPin(WidgetBlueprint, *Binding, Entry->GetBindingParameterName(), bSourceToDestination))
+							{
+								bCanSplitPin = false;
+							}
+							if (!EditorSubsystem->CanRecombinePin(WidgetBlueprint, *Binding, Entry->GetBindingParameterName(), bSourceToDestination))
+							{
+								bCanRecombinePin = false;
+							}
+
+							UEdGraphPin* GraphPin = GraphPin = EditorSubsystem->GetConversionFunctionArgumentPin(WidgetBlueprint, *Binding, Entry->GetBindingParameterName(), bSourceToDestination);
+							if (GraphPin == nullptr || GraphPin->ParentPin == nullptr)
+							{
+								bCanRecombinePinVisible = false;
+							}
+						}
+						else
+						{
+							bCanSplitPin = false;
+							bCanRecombinePin = false;
+							bCanRecombinePinVisible = false;
+						}
+					}
+
+					if (!bCanRecombinePin && !bCanSplitPin && !bCanRecombinePinVisible)
+					{
+						break;
+					}
+				}
+
+				if (bCanSplitPin)
+				{
+					FUIAction SplitPinAction;
+					SplitPinAction.ExecuteAction = FExecuteAction::CreateSP(this, &SBindingsList::HandleBreakSelectedPin);
+					MenuBuilder.AddMenuEntry(LOCTEXT("BreakPin", "Split Struct Pin"),
+						LOCTEXT("BreakPinTooltip", "Breaks a struct pin in to a separate pin per element."),
+						FSlateIcon(),
+						SplitPinAction);
+				}
+				if (bCanRecombinePinVisible)
+				{
+					FUIAction RecombinePinAction;
+					RecombinePinAction.ExecuteAction = FExecuteAction::CreateSP(this, &SBindingsList::HandleRecombineSelectedPin);
+					RecombinePinAction.CanExecuteAction = FCanExecuteAction::CreateLambda([bCanRecombinePin](){ return bCanRecombinePin; });
+					MenuBuilder.AddMenuEntry(LOCTEXT("RecombinePin", "Recombine Struct Pin"),
+						LOCTEXT("RecombinePinTooltip", "Takes struct pins that have been broken in to composite elements and combines them back to a single struct pin."),
+						FSlateIcon(),
+						RecombinePinAction);
+				}
+			}
+		}
 	}
 
 	return MenuBuilder.MakeWidget();
@@ -2383,7 +2580,7 @@ FReply SBindingsList::OnKeyDown(const FGeometry& MyGeometry, const FKeyEvent& In
 {
 	if (InKeyEvent.GetKey() == EKeys::Delete)
 	{
-		OnDeleteSelected();
+		HandleDeleteSelected();
 		return FReply::Handled();
 	}
 
diff --git a/Engine/Plugins/Runtime/ModelViewViewModel/Source/ModelViewViewModelEditor/Private/Widgets/SMVVMViewBindingListView.h b/Engine/Plugins/Runtime/ModelViewViewModel/Source/ModelViewViewModelEditor/Private/Widgets/SMVVMViewBindingListView.h
--- a/Engine/Plugins/Runtime/ModelViewViewModel/Source/ModelViewViewModelEditor/Private/Widgets/SMVVMViewBindingListView.h
+++ b/Engine/Plugins/Runtime/ModelViewViewModel/Source/ModelViewViewModelEditor/Private/Widgets/SMVVMViewBindingListView.h
@@ -45,7 +45,11 @@ private:
 	void GetChildrenOfEntry(TSharedPtr<FBindingEntry> Entry, TArray<TSharedPtr<FBindingEntry>>& OutChildren) const;
 
 	virtual FReply OnKeyDown(const FGeometry& MyGeometry, const FKeyEvent& InKeyEvent) override;
-	void OnDeleteSelected();
+	void HandleDeleteSelected();
+	void HandleBreakSelectedPin();
+	void HandleRecombineSelectedPin();
+
+	void ForceRefresh();
 
 private:
 	TWeakPtr<SBindingsPanel> BindingPanel;
diff --git a/Engine/Plugins/Runtime/ModelViewViewModel/Source/ModelViewViewModelEditor/Public/MVVMEditorSubsystem.h b/Engine/Plugins/Runtime/ModelViewViewModel/Source/ModelViewViewModelEditor/Public/MVVMEditorSubsystem.h
--- a/Engine/Plugins/Runtime/ModelViewViewModel/Source/ModelViewViewModelEditor/Public/MVVMEditorSubsystem.h
+++ b/Engine/Plugins/Runtime/ModelViewViewModel/Source/ModelViewViewModelEditor/Public/MVVMEditorSubsystem.h
@@ -51,6 +51,9 @@ public:
 	UFUNCTION(BlueprintCallable, Category = "Viewmodel")
 	bool RenameViewModel(UWidgetBlueprint* WidgetBlueprint, FName ViewModel, FName NewViewModel, FText& OutError);
 
+	UFUNCTION(BlueprintCallable, Category = "Viewmodel")
+	bool ReparentViewModel(UWidgetBlueprint* WidgetBlueprint, FName ViewModel, const UClass* NewViewModel, FText& OutError);
+
 	UFUNCTION(BlueprintCallable, Category = "Viewmodel")
 	FMVVMBlueprintViewBinding& AddBinding(UWidgetBlueprint* WidgetBlueprint);
 
@@ -103,7 +106,16 @@ public:
 
 	FMVVMBlueprintPropertyPath GetPathForConversionFunctionArgument(const UWidgetBlueprint* WidgetBlueprint, const FMVVMBlueprintViewBinding& Binding, FName ArgumentName, bool bSourceToDestination) const;
 	void SetPathForConversionFunctionArgument(UWidgetBlueprint* WidgetBlueprint, FMVVMBlueprintViewBinding& Binding, FName ArgumentName, const FMVVMBlueprintPropertyPath& Path, bool bSourceToDestination) const;
-	UEdGraphPin* GetConversionFunctionArgumentPin(const UWidgetBlueprint* WidgetBlueprint, const FMVVMBlueprintViewBinding& Binding, FName ParameterName, bool bSourceToDestination);
+	UEdGraphPin* GetConversionFunctionArgumentPin(const UWidgetBlueprint* WidgetBlueprint, const FMVVMBlueprintViewBinding& Binding, FName ParameterName, bool bSourceToDestination) const;
+
+	void SplitPin(UWidgetBlueprint* WidgetBlueprint, FMVVMBlueprintViewBinding& Binding, FName PinName, bool bSourceToDestination) const;
+	bool CanSplitPin(const UWidgetBlueprint* WidgetBlueprint, const FMVVMBlueprintViewBinding& Binding, FName PinName, bool bSourceToDestination) const;
+	void SplitPin(UWidgetBlueprint* WidgetBlueprint, UMVVMBlueprintViewEvent* Event, FName PinName) const;
+	bool CanSplitPin(const UWidgetBlueprint* WidgetBlueprint, UMVVMBlueprintViewEvent* Event, FName PinName) const;
+	void RecombinePin(UWidgetBlueprint* WidgetBlueprint, FMVVMBlueprintViewBinding& Binding, FName PinName, bool bSourceToDestination) const;
+	bool CanRecombinePin(const UWidgetBlueprint* WidgetBlueprint, const FMVVMBlueprintViewBinding& Binding, FName PinName, bool bSourceToDestination) const;
+	void RecombinePin(UWidgetBlueprint* WidgetBlueprint, UMVVMBlueprintViewEvent* Event, FName PinName) const;
+	bool CanRecombinePin(const UWidgetBlueprint* WidgetBlueprint, UMVVMBlueprintViewEvent* Event, FName PinName) const;
 
 	TArray<UE::MVVM::FBindingSource> GetBindableWidgets(const UWidgetBlueprint* WidgetBlueprint) const;
 	TArray<UE::MVVM::FBindingSource> GetAllViewModels(const UWidgetBlueprint* WidgetBlueprint) const;
