MVVM: Reset the NumReplicatedProperties when building an InstancedViewModel.

diff --git a/Engine/Plugins/Runtime/ModelViewViewModel/Source/ModelViewViewModelBlueprint/Private/MVVMBlueprintInstancedViewModel.cpp b/Engine/Plugins/Runtime/ModelViewViewModel/Source/ModelViewViewModelBlueprint/Private/MVVMBlueprintInstancedViewModel.cpp
--- a/Engine/Plugins/Runtime/ModelViewViewModel/Source/ModelViewViewModelBlueprint/Private/MVVMBlueprintInstancedViewModel.cpp
+++ b/Engine/Plugins/Runtime/ModelViewViewModel/Source/ModelViewViewModelBlueprint/Private/MVVMBlueprintInstancedViewModel.cpp
@@ -144,6 +144,7 @@ void UMVVMBlueprintInstancedViewModelBase::CleanClass()
 	GeneratedClass->ClassWithin = UObject::StaticClass();
 	GeneratedClass->ClassConfigName = ParentClass->ClassConfigName;
 	GeneratedClass->ClassFlags |= CLASS_NotPlaceable;
+	GeneratedClass->NumReplicatedProperties = 0;
 	GeneratedClass->FieldNotifies.Empty();
 }
 
@@ -238,7 +239,12 @@ void UMVVMBlueprintInstancedViewModelBase::InitializeProperty(FProperty* NewProp
 	}
 
 	EPropertyFlags NewFlags = NewProperty->GetPropertyFlags() | EPropertyFlags::CPF_Edit;
-	NewFlags |= !NewProperty->RepNotifyFunc.IsNone() ? EPropertyFlags::CPF_Net | EPropertyFlags::CPF_RepNotify : EPropertyFlags::CPF_None;
+	if (Args.bNetwork)
+	{
+		NewFlags |= EPropertyFlags::CPF_Net;
+		NewFlags |= !NewProperty->RepNotifyFunc.IsNone() ? EPropertyFlags::CPF_RepNotify : EPropertyFlags::CPF_None;
+		++GeneratedClass->NumReplicatedProperties;
+	}
 	NewFlags |= !Args.bPrivate ? EPropertyFlags::CPF_BlueprintVisible : EPropertyFlags::CPF_None;
 	NewFlags |= Args.bReadOnly ? EPropertyFlags::CPF_BlueprintReadOnly : EPropertyFlags::CPF_None;
 	NewProperty->SetPropertyFlags(NewFlags);
