MVVM: Test permission setting with the authorithy class.

diff --git a/Engine/Plugins/Runtime/ModelViewViewModel/Source/ModelViewViewModelBlueprint/Private/MVVMDeveloperProjectSettings.cpp b/Engine/Plugins/Runtime/ModelViewViewModel/Source/ModelViewViewModelBlueprint/Private/MVVMDeveloperProjectSettings.cpp
--- a/Engine/Plugins/Runtime/ModelViewViewModel/Source/ModelViewViewModelBlueprint/Private/MVVMDeveloperProjectSettings.cpp
+++ b/Engine/Plugins/Runtime/ModelViewViewModel/Source/ModelViewViewModelBlueprint/Private/MVVMDeveloperProjectSettings.cpp
@@ -69,7 +69,7 @@ bool UMVVMDeveloperProjectSettings::PropertyHasFiltering(const FProperty* Proper
 
 namespace UE::MVVM::Private
 {
-bool ShouldDoPropertyEditorPermission(const UBlueprint* GeneratingFor, UClass* FieldOwner)
+bool ShouldDoPropertyEditorPermission(const UBlueprint* GeneratingFor, const UClass* FieldOwner)
 {
 	if (GeneratingFor && FieldOwner)
 	{
@@ -84,46 +84,88 @@ bool UMVVMDeveloperProjectSettings::IsPropertyAllowed(const UBlueprint* Generati
 {
 	check(Property);
 	check(GeneratingFor);
-	const bool bDoPropertyEditorPermission = UE::MVVM::Private::ShouldDoPropertyEditorPermission(GeneratingFor, Property->GetTypedOwner<UClass>());
-	if (bDoPropertyEditorPermission && !FPropertyEditorPermissionList::Get().DoesPropertyPassFilter(Property->GetOwnerStruct(), Property->GetFName()))
+
+	const UStruct* OwnerStruct = Property->GetOwnerStruct();
+	check(OwnerStruct);
+
+	// The editor permission doesn't work with skeletal class
+	const UClass* OwnerClass = Cast<UClass>(OwnerStruct);
+	if (OwnerClass)
 	{
-		return false;
+		OwnerClass = OwnerClass->GetAuthoritativeClass();
+		OwnerStruct = OwnerClass;
 	}
 
-	TStringBuilder<512> StringBuilder;
-	Property->GetOwnerClass()->GetPathName(nullptr, StringBuilder);
-	FSoftClassPath StructPath;
-	StructPath.SetPath(StringBuilder);
-	if (const FMVVMDeveloperProjectWidgetSettings* Settings = FieldSelectorPermissions.Find(StructPath))
+	const bool bDoPropertyEditorPermission = UE::MVVM::Private::ShouldDoPropertyEditorPermission(GeneratingFor, OwnerClass);
+	if (bDoPropertyEditorPermission)
 	{
-		return !Settings->DisallowedFieldNames.Find(Property->GetFName());
+		if (!FPropertyEditorPermissionList::Get().DoesPropertyPassFilter(OwnerStruct, Property->GetFName()))
+		{
+			return false;
+		}
+	}
+
+	if (OwnerClass)
+	{
+		TStringBuilder<512> StringBuilder;
+		OwnerClass->GetPathName(nullptr, StringBuilder);
+		FSoftClassPath StructPath;
+		StructPath.SetPath(StringBuilder);
+		if (const FMVVMDeveloperProjectWidgetSettings* Settings = FieldSelectorPermissions.Find(StructPath))
+		{
+			return !Settings->DisallowedFieldNames.Find(Property->GetFName());
+		}
 	}
 	return true;
 }
 
 bool UMVVMDeveloperProjectSettings::IsFunctionAllowed(const UBlueprint* GeneratingFor, const UFunction* Function) const
 {
-	check(Function);
+	if (Function == nullptr)
+	{
+		return false;
+	}
+
+	const UClass* OriginalOwnerClass = Function->GetOwnerClass();
+	const UClass* OwnerClass = OriginalOwnerClass ? OriginalOwnerClass->GetAuthoritativeClass() : nullptr;
+	if (OwnerClass == nullptr)
+	{
+		return false;
+	}
 
 	TStringBuilder<512> StringBuilder;
 	const FPathPermissionList& FunctionPermissions = GetMutableDefault<UBlueprintEditorSettings>()->GetFunctionPermissions();
 	if (FunctionPermissions.HasFiltering())
 	{
-		Function->GetPathName(nullptr, StringBuilder);
+		const UFunction* FunctionToTest = Function;
+		if (OriginalOwnerClass != OwnerClass)
+		{
+			FunctionToTest = OwnerClass->FindFunctionByName(Function->GetFName());
+		}
+		if (FunctionToTest == nullptr)
+		{
+			return false;
+		}
+
+		StringBuilder.Reset();
+		FunctionToTest->GetPathName(nullptr, StringBuilder);
 		if (!FunctionPermissions.PassesFilter(StringBuilder.ToView()))
 		{
 			return false;
 		}
 	}
 
-	StringBuilder.Reset();
-	Function->GetOwnerClass()->GetPathName(nullptr, StringBuilder);
-	FSoftClassPath StructPath;
-	StructPath.SetPath(StringBuilder);
-	if (const FMVVMDeveloperProjectWidgetSettings* Settings = FieldSelectorPermissions.Find(StructPath))
 	{
-		return !Settings->DisallowedFieldNames.Find(Function->GetFName());
+		StringBuilder.Reset();
+		OwnerClass->GetPathName(nullptr, StringBuilder);
+		FSoftClassPath StructPath;
+		StructPath.SetPath(StringBuilder);
+		if (const FMVVMDeveloperProjectWidgetSettings* Settings = FieldSelectorPermissions.Find(StructPath))
+		{
+			return !Settings->DisallowedFieldNames.Find(Function->GetFName());
+		}
 	}
+
 	return true;
 }
 
