warning V700: Consider inspecting the 'T foo = foo = ...' expression. It is odd that variable is initialized through itself.
#rnx

diff --git a/Engine/Plugins/Experimental/ChaosVD/Source/ChaosVD/Private/ChaosVDScene.cpp b/Engine/Plugins/Experimental/ChaosVD/Source/ChaosVD/Private/ChaosVDScene.cpp
--- a/Engine/Plugins/Experimental/ChaosVD/Source/ChaosVD/Private/ChaosVDScene.cpp
+++ b/Engine/Plugins/Experimental/ChaosVD/Source/ChaosVD/Private/ChaosVDScene.cpp
@@ -413,7 +413,7 @@ void FChaosVDScene::CreateBaseLights(UWorld* TargetWorld) const
 UWorld* FChaosVDScene::CreatePhysicsVDWorld()
 {
 	const FName UniqueWorldName = FName(FGuid::NewGuid().ToString());
-	UWorld* NewWorld = NewWorld = NewObject<UWorld>( GetTransientPackage(), UniqueWorldName );
+	UWorld* NewWorld = NewObject<UWorld>( GetTransientPackage(), UniqueWorldName );
 	
 	NewWorld->WorldType = EWorldType::EditorPreview;
 
diff --git a/Engine/Plugins/Runtime/ModelViewViewModel/Source/ModelViewViewModelEditor/Private/Widgets/SMVVMViewBindingListView.cpp b/Engine/Plugins/Runtime/ModelViewViewModel/Source/ModelViewViewModelEditor/Private/Widgets/SMVVMViewBindingListView.cpp
--- a/Engine/Plugins/Runtime/ModelViewViewModel/Source/ModelViewViewModelEditor/Private/Widgets/SMVVMViewBindingListView.cpp
+++ b/Engine/Plugins/Runtime/ModelViewViewModel/Source/ModelViewViewModelEditor/Private/Widgets/SMVVMViewBindingListView.cpp
@@ -2513,7 +2513,7 @@ TSharedPtr<SWidget> SBindingsList::OnSourceConstructContextMenu()
 								bCanRecombinePin = false;
 							}
 
-							UEdGraphPin* GraphPin = GraphPin = EditorSubsystem->GetConversionFunctionArgumentPin(WidgetBlueprint, *Binding, Entry->GetBindingParameterName(), bSourceToDestination);
+							UEdGraphPin* GraphPin = EditorSubsystem->GetConversionFunctionArgumentPin(WidgetBlueprint, *Binding, Entry->GetBindingParameterName(), bSourceToDestination);
 							if (GraphPin == nullptr || GraphPin->ParentPin == nullptr)
 							{
 								bCanRecombinePinVisible = false;
diff --git a/Engine/Source/Editor/UnrealEd/Private/Commandlets/CookCommandlet.cpp b/Engine/Source/Editor/UnrealEd/Private/Commandlets/CookCommandlet.cpp
--- a/Engine/Source/Editor/UnrealEd/Private/Commandlets/CookCommandlet.cpp
+++ b/Engine/Source/Editor/UnrealEd/Private/Commandlets/CookCommandlet.cpp
@@ -530,7 +530,7 @@ void UCookCommandlet::RunCookByTheBookCook(UCookOnTheFlyServer* CookOnTheFlyServ
 		DECLARE_SCOPE_CYCLE_COUNTER(TEXT("CookByTheBook.MainLoop"), STAT_CookByTheBook_MainLoop, STATGROUP_LoadTime);
 		while (CookOnTheFlyServer->IsInSession())
 		{
-			uint32 TickResults = TickResults = CookOnTheFlyServer->TickCookByTheBook(MAX_flt,
+			uint32 TickResults = CookOnTheFlyServer->TickCookByTheBook(MAX_flt,
 				ShowProgress ? ECookTickFlags::None : ECookTickFlags::HideProgressDisplay);
 			ConditionalCollectGarbage(TickResults, *CookOnTheFlyServer);
 		}
