MVVM: Cache the default execution mode for performance.

diff --git a/Engine/Plugins/Runtime/ModelViewViewModel/Source/ModelViewViewModel/Private/ModelViewViewModelModule.cpp b/Engine/Plugins/Runtime/ModelViewViewModel/Source/ModelViewViewModel/Private/ModelViewViewModelModule.cpp
--- a/Engine/Plugins/Runtime/ModelViewViewModel/Source/ModelViewViewModel/Private/ModelViewViewModelModule.cpp
+++ b/Engine/Plugins/Runtime/ModelViewViewModel/Source/ModelViewViewModel/Private/ModelViewViewModelModule.cpp
@@ -9,9 +9,10 @@
 
 void FModelViewViewModelModule::StartupModule()
 {
-	static IConsoleVariable* CVarDefaultExecutionMode = IConsoleManager::Get().FindConsoleVariable(TEXT("MVVM.DefaultExecutionMode"));
-	if (CVarDefaultExecutionMode)
+	IConsoleVariable* CVarDefaultExecutionMode = IConsoleManager::Get().FindConsoleVariable(TEXT("MVVM.DefaultExecutionMode"));
+	if (ensure(CVarDefaultExecutionMode))
 	{
+		HandleDefaultExecutionModeChanged(CVarDefaultExecutionMode);
 		CVarDefaultExecutionMode->OnChangedDelegate().AddRaw(this, &FModelViewViewModelModule::HandleDefaultExecutionModeChanged);
 	}
 }
@@ -20,8 +21,7 @@ void FModelViewViewModelModule::ShutdownModule()
 {
 	if (!IsEngineExitRequested())
 	{
-		static IConsoleVariable* CVarDefaultExecutionMode = IConsoleManager::Get().FindConsoleVariable(TEXT("MVVM.DefaultExecutionMode"));
-		if (CVarDefaultExecutionMode)
+		if (IConsoleVariable* CVarDefaultExecutionMode = IConsoleManager::Get().FindConsoleVariable(TEXT("MVVM.DefaultExecutionMode")))
 		{
 			CVarDefaultExecutionMode->OnChangedDelegate().RemoveAll(this);
 		}
@@ -30,13 +30,18 @@ void FModelViewViewModelModule::ShutdownModule()
 
 void FModelViewViewModelModule::HandleDefaultExecutionModeChanged(IConsoleVariable* Variable)
 {
-	EMVVMExecutionMode Value = (EMVVMExecutionMode)Variable->GetInt();
-	if (Value != EMVVMExecutionMode::Delayed
-		&& Value != EMVVMExecutionMode::Immediate
-		&& Value != EMVVMExecutionMode::Tick
-		&& Value != EMVVMExecutionMode::DelayedWhenSharedElseImmediate)
+	const int32 Value = Variable->GetInt();
+	switch(Value)
 	{
-		Variable->Set((int32)EMVVMExecutionMode::Immediate, (EConsoleVariableFlags)(Variable->GetFlags() & ECVF_SetByMask));
+	case (int32)EMVVMExecutionMode::Delayed:
+	case (int32)EMVVMExecutionMode::Immediate:
+	case (int32)EMVVMExecutionMode::Tick:
+	case (int32)EMVVMExecutionMode::DelayedWhenSharedElseImmediate:
+		break;
+	default:
+		ensureMsgf(false, TEXT("MVVM.DefaultExecutionMode default value is not a valid value."));
+		Variable->Set((int32)EMVVMExecutionMode::DelayedWhenSharedElseImmediate, (EConsoleVariableFlags)(Variable->GetFlags() & ECVF_SetByMask));
+		break;
 	}
 }
 
diff --git a/Engine/Plugins/Runtime/ModelViewViewModel/Source/ModelViewViewModel/Private/View/MVVMViewClass.cpp b/Engine/Plugins/Runtime/ModelViewViewModel/Source/ModelViewViewModel/Private/View/MVVMViewClass.cpp
--- a/Engine/Plugins/Runtime/ModelViewViewModel/Source/ModelViewViewModel/Private/View/MVVMViewClass.cpp
+++ b/Engine/Plugins/Runtime/ModelViewViewModel/Source/ModelViewViewModel/Private/View/MVVMViewClass.cpp
@@ -274,18 +274,31 @@ void FMVVMViewClass_SourceCreator::PostSerialize(const FArchive& Ar)
 ///////////////////////////////////////////////////////////////////////
 namespace UE::MVVM::Private
 {
-	int32 GDefaultEvaluationMode = (int32)EMVVMExecutionMode::Immediate;
-	static FAutoConsoleVariableRef CVarDefaultEvaluationMode(
+	FAutoConsoleVariable CVarDefaultEvaluationMode(
 		TEXT("MVVM.DefaultExecutionMode"),
-		GDefaultEvaluationMode,
-		TEXT("The default evaluation mode of a MVVM binding.")
+		static_cast<uint8>(EMVVMExecutionMode::DelayedWhenSharedElseImmediate),
+		TEXT("The default evaluation mode of a MVVM binding. Must restart if changed."),
+		EConsoleVariableFlags::ECVF_ReadOnly
 	);
 }
 
 EMVVMExecutionMode FMVVMViewClass_CompiledBinding::GetExecuteMode() const
 {
-	EMVVMExecutionMode DefaultMode = (EMVVMExecutionMode)UE::MVVM::Private::GDefaultEvaluationMode;
-	EMVVMExecutionMode Result = (Flags & EBindingFlags::OverrideExecuteMode) == 0 ? DefaultMode : ExecutionMode;
+	struct FLocal
+	{
+		FLocal()
+		{
+			IConsoleVariable* CVarDefaultExecutionMode = IConsoleManager::Get().FindConsoleVariable(TEXT("MVVM.DefaultExecutionMode"));
+			if (ensure(CVarDefaultExecutionMode))
+			{
+				DefaultMode = (EMVVMExecutionMode)CVarDefaultExecutionMode->GetInt();
+			}
+		}
+		EMVVMExecutionMode DefaultMode = EMVVMExecutionMode::DelayedWhenSharedElseImmediate;
+	};
+	static FLocal Local;
+
+	EMVVMExecutionMode Result = (Flags & EBindingFlags::OverrideExecuteMode) == 0 ? Local.DefaultMode : ExecutionMode;
 	return Result == EMVVMExecutionMode::DelayedWhenSharedElseImmediate ? (Binding.IsShared() ? EMVVMExecutionMode::Delayed : EMVVMExecutionMode::Immediate) : Result;
 }
 
diff --git a/Engine/Plugins/Runtime/ModelViewViewModel/Source/ModelViewViewModelBlueprint/Private/MVVMViewBlueprintCompiler.cpp b/Engine/Plugins/Runtime/ModelViewViewModel/Source/ModelViewViewModelBlueprint/Private/MVVMViewBlueprintCompiler.cpp
--- a/Engine/Plugins/Runtime/ModelViewViewModel/Source/ModelViewViewModelBlueprint/Private/MVVMViewBlueprintCompiler.cpp
+++ b/Engine/Plugins/Runtime/ModelViewViewModel/Source/ModelViewViewModelBlueprint/Private/MVVMViewBlueprintCompiler.cpp
@@ -1858,12 +1858,19 @@ bool FMVVMViewBlueprintCompiler::CompileBindings(const FCompiledBindingLibraryCo
 		return false;
 	}
 
-	static IConsoleVariable* CVarDefaultExecutionMode = IConsoleManager::Get().FindConsoleVariable(TEXT("MVVM.DefaultExecutionMode"));
-	if (!CVarDefaultExecutionMode)
+	struct FLocal
 	{
-		WidgetBlueprintCompilerContext.MessageLog.Error(*LOCTEXT("CantFindDefaultExecutioMode", "The default execution mode cannot be found.").ToString());
-		return false;
-	}
+		FLocal()
+		{
+			IConsoleVariable* CVarDefaultExecutionMode = IConsoleManager::Get().FindConsoleVariable(TEXT("MVVM.DefaultExecutionMode"));
+			if (ensure(CVarDefaultExecutionMode))
+			{
+				DefaultMode = (EMVVMExecutionMode)CVarDefaultExecutionMode->GetInt();
+			}
+		}
+		EMVVMExecutionMode DefaultMode = EMVVMExecutionMode::DelayedWhenSharedElseImmediate;
+	};
+	static FLocal LocalDefaultMode;
 
 	// Store bindings with corresponding ComplexConversionFunctionContextIndex to be sorted and mark the last binding 
 	// in the complex conversion as bExecuteAtInitialization == true
@@ -2000,7 +2007,7 @@ bool FMVVMViewBlueprintCompiler::CompileBindings(const FCompiledBindingLibraryCo
 				ComplexConversionFunctionContextIndex = CompileBinding.ComplexConversionFunctionContextIndex;
 			}
 
-			NewBinding.ExecutionMode = ViewBinding->bOverrideExecutionMode ? ViewBinding->OverrideExecutionMode : (EMVVMExecutionMode)CVarDefaultExecutionMode->GetInt();;
+			NewBinding.ExecutionMode = ViewBinding->bOverrideExecutionMode ? ViewBinding->OverrideExecutionMode : LocalDefaultMode.DefaultMode;
 			NewBinding.EditorId = ViewBinding->BindingId;
 
 			NewBinding.Flags |= (ViewBinding->bEnabled) ? FMVVMViewClass_CompiledBinding::EBindingFlags::EnabledByDefault : 0;
diff --git a/Engine/Plugins/Runtime/ModelViewViewModel/Source/ModelViewViewModelEditor/Private/Widgets/SMVVMViewBindingListView.cpp b/Engine/Plugins/Runtime/ModelViewViewModel/Source/ModelViewViewModelEditor/Private/Widgets/SMVVMViewBindingListView.cpp
--- a/Engine/Plugins/Runtime/ModelViewViewModel/Source/ModelViewViewModelEditor/Private/Widgets/SMVVMViewBindingListView.cpp
+++ b/Engine/Plugins/Runtime/ModelViewViewModel/Source/ModelViewViewModelEditor/Private/Widgets/SMVVMViewBindingListView.cpp
@@ -559,8 +559,10 @@ public:
 		Entry = InEntry;
 		WidgetBlueprintWeak = InWidgetBlueprint;
 
-		static IConsoleVariable* StaticCVarDefaultExecutionMode = IConsoleManager::Get().FindConsoleVariable(TEXT("MVVM.DefaultExecutionMode"));
-		CVarDefaultExecutionMode = StaticCVarDefaultExecutionMode;
+		static IConsoleVariable* CVarDefaultExecutionMode = IConsoleManager::Get().FindConsoleVariable(TEXT("MVVM.DefaultExecutionMode"));
+		ensure(CVarDefaultExecutionMode);
+		DefaultExecutionMode = CVarDefaultExecutionMode ? (EMVVMExecutionMode)CVarDefaultExecutionMode->GetInt() : EMVVMExecutionMode::DelayedWhenSharedElseImmediate;
+
 		
 		FMVVMBlueprintViewBinding* ViewBinding = GetThisViewBinding();
 
@@ -1102,8 +1104,7 @@ private:
 			if (NewState == ECheckBoxState::Checked)
 			{
 				UMVVMEditorSubsystem* Subsystem = GEditor->GetEditorSubsystem<UMVVMEditorSubsystem>();
-				EMVVMExecutionMode ExecutionMode = CVarDefaultExecutionMode ? (EMVVMExecutionMode)CVarDefaultExecutionMode->GetInt() : EMVVMExecutionMode::Immediate;
-				Subsystem->OverrideExecutionModeForBinding(WidgetBlueprintWeak.Get(), *ViewBinding, ExecutionMode);
+				Subsystem->OverrideExecutionModeForBinding(WidgetBlueprintWeak.Get(), *ViewBinding, DefaultExecutionMode);
 			}
 			else
 			{
@@ -1155,7 +1156,7 @@ private:
 
 	FText GetExecutioModeValue() const
 	{
-		EMVVMExecutionMode ExecutionMode = CVarDefaultExecutionMode ? (EMVVMExecutionMode)CVarDefaultExecutionMode->GetInt() : EMVVMExecutionMode::Immediate;
+		EMVVMExecutionMode ExecutionMode = DefaultExecutionMode;
 		if (FMVVMBlueprintViewBinding* ViewBinding = GetThisViewBinding())
 		{
 			if (ViewBinding->bOverrideExecutionMode)
@@ -1168,7 +1169,7 @@ private:
 	
 	FText GetExecutioModeValueToolTip() const
 	{
-		EMVVMExecutionMode ExecutionMode = CVarDefaultExecutionMode ? (EMVVMExecutionMode)CVarDefaultExecutionMode->GetInt() : EMVVMExecutionMode::Immediate;
+		EMVVMExecutionMode ExecutionMode = DefaultExecutionMode;
 		if (FMVVMBlueprintViewBinding* ViewBinding = GetThisViewBinding())
 		{
 			if (ViewBinding->bOverrideExecutionMode)
@@ -1324,7 +1325,7 @@ private:
 	TWeakObjectPtr<UWidgetBlueprint> WidgetBlueprintWeak;
 	TSharedPtr<SCustomDialog> ErrorDialog;
 	TArray<TSharedPtr<FText>> ErrorItems;
-	IConsoleVariable* CVarDefaultExecutionMode = nullptr;
+	EMVVMExecutionMode DefaultExecutionMode = EMVVMExecutionMode::DelayedWhenSharedElseImmediate;
 };
 
 /**
