MVVM: Test if a field is allowed base on the path, not only on the field owner. That allow to have a Widget's property allowed on a child widget but not on the base widget.

diff --git a/Engine/Plugins/Runtime/ModelViewViewModel/Source/ModelViewViewModel/Private/Bindings/MVVMFieldPathHelper.cpp b/Engine/Plugins/Runtime/ModelViewViewModel/Source/ModelViewViewModel/Private/Bindings/MVVMFieldPathHelper.cpp
--- a/Engine/Plugins/Runtime/ModelViewViewModel/Source/ModelViewViewModel/Private/Bindings/MVVMFieldPathHelper.cpp
+++ b/Engine/Plugins/Runtime/ModelViewViewModel/Source/ModelViewViewModel/Private/Bindings/MVVMFieldPathHelper.cpp
@@ -491,6 +491,32 @@ TValueOrError<UObject*, void> EvaluateObjectProperty(const FFieldContext& InSour
 	}
 }
 
+TValueOrError<const UStruct*, void> GetFieldAsContainer(const UE::MVVM::FMVVMConstFieldVariant Field)
+{
+	const FProperty* PropertyToTest = nullptr;
+	if (Field.IsProperty())
+	{
+		PropertyToTest = Field.GetProperty();
+	}
+	else if (Field.IsFunction())
+	{
+		PropertyToTest = BindingHelper::GetReturnProperty(Field.GetFunction());
+	}
+
+	if (PropertyToTest)
+	{
+		if (const FObjectPropertyBase* ObjectProperty = CastField<const FObjectPropertyBase>(PropertyToTest))
+		{
+			return MakeValue(ObjectProperty->PropertyClass);
+		}
+		else if (const FStructProperty* StructProperty = CastField<const FStructProperty>(PropertyToTest))
+		{
+			return MakeValue(StructProperty->Struct);
+		}
+	}
+	return MakeError();
+}
+
 } // namespace
 
 #undef LOCTEXT_NAMESPACE
diff --git a/Engine/Plugins/Runtime/ModelViewViewModel/Source/ModelViewViewModel/Public/Bindings/MVVMFieldPathHelper.h b/Engine/Plugins/Runtime/ModelViewViewModel/Source/ModelViewViewModel/Public/Bindings/MVVMFieldPathHelper.h
--- a/Engine/Plugins/Runtime/ModelViewViewModel/Source/ModelViewViewModel/Public/Bindings/MVVMFieldPathHelper.h
+++ b/Engine/Plugins/Runtime/ModelViewViewModel/Source/ModelViewViewModel/Public/Bindings/MVVMFieldPathHelper.h
@@ -58,9 +58,17 @@ namespace UE::MVVM::FieldPathHelper
 	UE_NODISCARD MODELVIEWVIEWMODEL_API FText ToText(TArrayView<const FMVVMFieldVariant> Fields);
 	UE_NODISCARD MODELVIEWVIEWMODEL_API FText ToText(TArrayView<const FMVVMConstFieldVariant> Fields);
 
-	/** */
+	/** Returns the object value from the FieldContext. */
 	UE_NODISCARD MODELVIEWVIEWMODEL_API TValueOrError<UObject*, void> EvaluateObjectProperty(const FFieldContext& InSource);
 
+	/**
+	 * Returns the container that the field represents.
+	 * If the field is a Object property, then returns the PropertyClass.
+	 * If the field is a Struct property, then returns the Struct.
+	 * If the field is a Function, then returns the return value property's container.
+	 */
+	UE_NODISCARD MODELVIEWVIEWMODEL_API TValueOrError<const UStruct*, void> GetFieldAsContainer(const UE::MVVM::FMVVMConstFieldVariant Field);
+
 } // namespace
 
 #if UE_ENABLE_INCLUDE_ORDER_DEPRECATED_IN_5_2
diff --git a/Engine/Plugins/Runtime/ModelViewViewModel/Source/ModelViewViewModelBlueprint/Private/Bindings/MVVMCompiledBindingLibraryCompiler.cpp b/Engine/Plugins/Runtime/ModelViewViewModel/Source/ModelViewViewModelBlueprint/Private/Bindings/MVVMCompiledBindingLibraryCompiler.cpp
--- a/Engine/Plugins/Runtime/ModelViewViewModel/Source/ModelViewViewModelBlueprint/Private/Bindings/MVVMCompiledBindingLibraryCompiler.cpp
+++ b/Engine/Plugins/Runtime/ModelViewViewModel/Source/ModelViewViewModelBlueprint/Private/Bindings/MVVMCompiledBindingLibraryCompiler.cpp
@@ -3,6 +3,7 @@
 #include "Bindings/MVVMCompiledBindingLibraryCompiler.h"
 
 #include "Bindings/MVVMBindingHelper.h"
+#include "Bindings/MVVMFieldPathHelper.h"
 #include "Engine/Blueprint.h"
 #include "Engine/Engine.h"
 #include "Kismet2/BlueprintEditorUtils.h"
@@ -315,8 +316,21 @@ TValueOrError<FCompiledBindingLibraryCompiler::FFieldPathHandle, FText> FCompile
 	TArray<int32> RawFieldIndexes;
 	RawFieldIndexes.Reserve(InFieldPath.Num());
 
+	const UBlueprint* BlueprintContext = Impl->BlueprintContext.Get();
+	if (BlueprintContext == nullptr)
+	{
+		return MakeError(LOCTEXT("InvalidBlueprint", "The Blueprint is invalid."));
+	}
+
+	const UStruct* CurrentContainer = BlueprintContext->GeneratedClass ? BlueprintContext->GeneratedClass : BlueprintContext->SkeletonGeneratedClass;
+
 	for (int32 Index = 0; Index < InFieldPath.Num(); ++Index)
 	{
+		if (CurrentContainer == nullptr)
+		{
+			return MakeError(LOCTEXT("InvalidContainer", "The path has an invalid container."));
+		}
+
 		// Make sure the FieldVariant is not from a skeletalclass
 		FMVVMConstFieldVariant FieldVariant = InFieldPath[Index];
 		if (!FieldVariant.IsValid())
@@ -346,7 +360,7 @@ TValueOrError<FCompiledBindingLibraryCompiler::FFieldPathHandle, FText> FCompile
 				return MakeError(ValidatedStr);
 			}
 
-			if (!GetDefault<UMVVMDeveloperProjectSettings>()->IsPropertyAllowed(Impl->BlueprintContext.Get(), FieldVariant.GetProperty()))
+			if (!GetDefault<UMVVMDeveloperProjectSettings>()->IsPropertyAllowed(Impl->BlueprintContext.Get(), CurrentContainer, FieldVariant.GetProperty()))
 			{
 				return MakeError(LOCTEXT("PropertyNotAllow", "A property is not allowed."));
 			}
@@ -367,7 +381,13 @@ TValueOrError<FCompiledBindingLibraryCompiler::FFieldPathHandle, FText> FCompile
 				return MakeError(FText::Format(LOCTEXT("FunctionNotReadableAtRuntime", "Function '{0}' is not readable at runtime."), FieldVariant.GetFunction()->GetDisplayNameText()));
 			}
 
-			if (!GetDefault<UMVVMDeveloperProjectSettings>()->IsFunctionAllowed(Impl->BlueprintContext.Get(), FieldVariant.GetFunction()))
+			const UClass* CurrentContainerAsClass = Cast<const UClass>(CurrentContainer);
+			if (CurrentContainerAsClass == nullptr)
+			{
+				return MakeError(LOCTEXT("InvalidContainer", "The path has an invalid container."));
+			}
+
+			if (!GetDefault<UMVVMDeveloperProjectSettings>()->IsFunctionAllowed(Impl->BlueprintContext.Get(), CurrentContainerAsClass, FieldVariant.GetFunction()))
 			{
 				return MakeError(LOCTEXT("FunctionNotAllow", "A function is not allowed."));
 			}
@@ -389,6 +409,9 @@ TValueOrError<FCompiledBindingLibraryCompiler::FFieldPathHandle, FText> FCompile
 		{
 			return MakeError(LOCTEXT("InvalidFieldInPath", "There is an invalid field in the field path."));
 		}
+
+		TValueOrError<const UStruct*, void> FieldAsContainerResult = UE::MVVM::FieldPathHelper::GetFieldAsContainer(FieldVariant);
+		CurrentContainer = FieldAsContainerResult.HasValue() ? FieldAsContainerResult.GetValue() : nullptr;
 	}
 
 	int32 FoundFieldPath = Impl->FieldPaths.IndexOfByPredicate([&RawFieldIndexes](const Private::FRawFieldPath& Other)
diff --git a/Engine/Plugins/Runtime/ModelViewViewModel/Source/ModelViewViewModelBlueprint/Private/MVVMDeveloperProjectSettings.cpp b/Engine/Plugins/Runtime/ModelViewViewModel/Source/ModelViewViewModelBlueprint/Private/MVVMDeveloperProjectSettings.cpp
--- a/Engine/Plugins/Runtime/ModelViewViewModel/Source/ModelViewViewModelBlueprint/Private/MVVMDeveloperProjectSettings.cpp
+++ b/Engine/Plugins/Runtime/ModelViewViewModel/Source/ModelViewViewModelBlueprint/Private/MVVMDeveloperProjectSettings.cpp
@@ -37,20 +37,13 @@ FText UMVVMDeveloperProjectSettings::GetSectionText() const
 	return LOCTEXT("MVVMProjectSettings", "UMG Model View Viewmodel");
 }
 
-bool UMVVMDeveloperProjectSettings::PropertyHasFiltering(const FProperty* Property) const
+bool UMVVMDeveloperProjectSettings::PropertyHasFiltering(const UStruct* ObjectStruct, const FProperty* Property) const
 {
+	check(ObjectStruct);
 	check(Property);
-	const UStruct* ObjectStruct = nullptr;
-
-	if (const FObjectProperty* ObjectProperty = CastField<FObjectProperty>(Property))
-	{
-		ObjectStruct = ObjectProperty->PropertyClass;
-	}
-	else
-	{
-		ObjectStruct = Property->GetOwnerStruct();
-	}
 
+	const UClass* AuthoritativeClass = Cast<const UClass>(ObjectStruct);
+	ObjectStruct = AuthoritativeClass ? AuthoritativeClass->GetAuthoritativeClass() : ObjectStruct;
 	if (!FPropertyEditorPermissionList::Get().HasFiltering(ObjectStruct))
 	{
 		return false;
@@ -80,55 +73,49 @@ bool ShouldDoPropertyEditorPermission(const UBlueprint* GeneratingFor, const UCl
 }
 }//namespace
 
-bool UMVVMDeveloperProjectSettings::IsPropertyAllowed(const UBlueprint* GeneratingFor, const FProperty* Property) const
+bool UMVVMDeveloperProjectSettings::IsPropertyAllowed(const UBlueprint* GeneratingFor, const UStruct* ObjectStruct, const FProperty* Property) const
 {
-	check(Property);
 	check(GeneratingFor);
+	check(ObjectStruct);
+	check(Property);
 
-	const UStruct* OwnerStruct = Property->GetOwnerStruct();
-	check(OwnerStruct);
-
-	// The editor permission doesn't work with skeletal class
-	const UClass* OwnerClass = Cast<UClass>(OwnerStruct);
-	if (OwnerClass)
-	{
-		OwnerClass = OwnerClass->GetAuthoritativeClass();
-		OwnerStruct = OwnerClass;
-	}
+	const UClass* AuthoritativeClass = Cast<const UClass>(ObjectStruct);
+	AuthoritativeClass = AuthoritativeClass ? AuthoritativeClass->GetAuthoritativeClass() : nullptr;
 
-	const bool bDoPropertyEditorPermission = UE::MVVM::Private::ShouldDoPropertyEditorPermission(GeneratingFor, OwnerClass);
+	const bool bDoPropertyEditorPermission = UE::MVVM::Private::ShouldDoPropertyEditorPermission(GeneratingFor, AuthoritativeClass);
 	if (bDoPropertyEditorPermission)
 	{
-		if (!FPropertyEditorPermissionList::Get().DoesPropertyPassFilter(OwnerStruct, Property->GetFName()))
+		if (!FPropertyEditorPermissionList::Get().DoesPropertyPassFilter(AuthoritativeClass, Property->GetFName()))
 		{
 			return false;
 		}
 	}
 
-	if (OwnerClass)
+	if (AuthoritativeClass)
 	{
 		TStringBuilder<512> StringBuilder;
-		OwnerClass->GetPathName(nullptr, StringBuilder);
+		AuthoritativeClass->GetPathName(nullptr, StringBuilder);
 		FSoftClassPath StructPath;
 		StructPath.SetPath(StringBuilder);
 		if (const FMVVMDeveloperProjectWidgetSettings* Settings = FieldSelectorPermissions.Find(StructPath))
 		{
-			return !Settings->DisallowedFieldNames.Find(Property->GetFName());
+			if (Settings->DisallowedFieldNames.Find(Property->GetFName()))
+			{
+				return false;
+			}
 		}
 	}
 	return true;
 }
 
-bool UMVVMDeveloperProjectSettings::IsFunctionAllowed(const UBlueprint* GeneratingFor, const UFunction* Function) const
+bool UMVVMDeveloperProjectSettings::IsFunctionAllowed(const UBlueprint* GeneratingFor, const UClass* ObjectClass, const UFunction* Function) const
 {
-	if (Function == nullptr)
-	{
-		return false;
-	}
+	check(GeneratingFor);
+	check(ObjectClass);
+	check(Function);
 
-	const UClass* OriginalOwnerClass = Function->GetOwnerClass();
-	const UClass* OwnerClass = OriginalOwnerClass ? OriginalOwnerClass->GetAuthoritativeClass() : nullptr;
-	if (OwnerClass == nullptr)
+	const UClass* AuthoritativeClass = ObjectClass->GetAuthoritativeClass();
+	if (AuthoritativeClass == nullptr)
 	{
 		return false;
 	}
@@ -137,11 +124,7 @@ bool UMVVMDeveloperProjectSettings::IsFunctionAllowed(const UBlueprint* Generati
 	const FPathPermissionList& FunctionPermissions = GetMutableDefault<UBlueprintEditorSettings>()->GetFunctionPermissions();
 	if (FunctionPermissions.HasFiltering())
 	{
-		const UFunction* FunctionToTest = Function;
-		if (OriginalOwnerClass != OwnerClass)
-		{
-			FunctionToTest = OwnerClass->FindFunctionByName(Function->GetFName());
-		}
+		const UFunction* FunctionToTest = AuthoritativeClass->FindFunctionByName(Function->GetFName());
 		if (FunctionToTest == nullptr)
 		{
 			return false;
@@ -157,12 +140,15 @@ bool UMVVMDeveloperProjectSettings::IsFunctionAllowed(const UBlueprint* Generati
 
 	{
 		StringBuilder.Reset();
-		OwnerClass->GetPathName(nullptr, StringBuilder);
+		AuthoritativeClass->GetPathName(nullptr, StringBuilder);
 		FSoftClassPath StructPath;
 		StructPath.SetPath(StringBuilder);
 		if (const FMVVMDeveloperProjectWidgetSettings* Settings = FieldSelectorPermissions.Find(StructPath))
 		{
-			return !Settings->DisallowedFieldNames.Find(Function->GetFName());
+			if (Settings->DisallowedFieldNames.Find(Function->GetFName()))
+			{
+				return false;
+			}
 		}
 	}
 
@@ -179,7 +165,7 @@ bool UMVVMDeveloperProjectSettings::IsConversionFunctionAllowed(const UBlueprint
 
 	if (ConversionFunctionFilter == EMVVMDeveloperConversionFunctionFilterType::BlueprintActionRegistry)
 	{
-		return IsFunctionAllowed(GeneratingFor, Function);
+		return IsFunctionAllowed(GeneratingFor, Function->GetOwnerClass(), Function);
 	}
 	else
 	{
@@ -205,7 +191,7 @@ bool UMVVMDeveloperProjectSettings::IsConversionFunctionAllowed(const UBlueprint
 		else
 		{
 			// The function is on self and may have been filtered.
-			return IsFunctionAllowed(GeneratingFor, Function);
+			return IsFunctionAllowed(GeneratingFor, Function->GetOwnerClass(), Function);
 		}
 	}
 }
diff --git a/Engine/Plugins/Runtime/ModelViewViewModel/Source/ModelViewViewModelBlueprint/Private/MVVMViewBlueprintCompiler.cpp b/Engine/Plugins/Runtime/ModelViewViewModel/Source/ModelViewViewModelBlueprint/Private/MVVMViewBlueprintCompiler.cpp
--- a/Engine/Plugins/Runtime/ModelViewViewModel/Source/ModelViewViewModelBlueprint/Private/MVVMViewBlueprintCompiler.cpp
+++ b/Engine/Plugins/Runtime/ModelViewViewModel/Source/ModelViewViewModelBlueprint/Private/MVVMViewBlueprintCompiler.cpp
@@ -2643,11 +2643,16 @@ TArray<FMVVMConstFieldVariant> FMVVMViewBlueprintCompiler::CreatePropertyPath(co
 
 bool FMVVMViewBlueprintCompiler::IsPropertyPathValid(const UBlueprint* Context, TArrayView<const FMVVMConstFieldVariant> PropertyPath)
 {
+	const UStruct* CurrentContainer = Context->GeneratedClass ? Context->GeneratedClass : Context->SkeletonGeneratedClass;
 	int32 PathLength = PropertyPath.Num();
 	for (int32 Index = 0; Index < PathLength; Index++)
 	{
 		const FMVVMConstFieldVariant& Field = PropertyPath[Index];
 
+		if (CurrentContainer == nullptr)
+		{
+			return false;
+		}
 		if (Field.IsEmpty())
 		{
 			return false;
@@ -2659,7 +2664,7 @@ bool FMVVMViewBlueprintCompiler::IsPropertyPathValid(const UBlueprint* Context,
 				return false;
 			}
 
-			if (!GetDefault<UMVVMDeveloperProjectSettings>()->IsPropertyAllowed(Context, Field.GetProperty()))
+			if (!GetDefault<UMVVMDeveloperProjectSettings>()->IsPropertyAllowed(Context, CurrentContainer, Field.GetProperty()))
 			{
 				return false;
 			}
@@ -2670,11 +2675,19 @@ bool FMVVMViewBlueprintCompiler::IsPropertyPathValid(const UBlueprint* Context,
 			{
 				return false;
 			}
-			if (!GetDefault<UMVVMDeveloperProjectSettings>()->IsFunctionAllowed(Context, Field.GetFunction()))
+			const UClass* CurrentContainerAsClass = Cast<const UClass>(CurrentContainer);
+			if (CurrentContainerAsClass == nullptr)
+			{
+				return false;
+			}
+			if (!GetDefault<UMVVMDeveloperProjectSettings>()->IsFunctionAllowed(Context, CurrentContainerAsClass, Field.GetFunction()))
 			{
 				return false;
 			}
 		}
+
+		TValueOrError<const UStruct*, void> FieldAsContainerResult = UE::MVVM::FieldPathHelper::GetFieldAsContainer(Field);
+		CurrentContainer = FieldAsContainerResult.HasValue() ? FieldAsContainerResult.GetValue() : nullptr;
 	}
 	return true;
 }
diff --git a/Engine/Plugins/Runtime/ModelViewViewModel/Source/ModelViewViewModelBlueprint/Public/MVVMDeveloperProjectSettings.h b/Engine/Plugins/Runtime/ModelViewViewModel/Source/ModelViewViewModelBlueprint/Public/MVVMDeveloperProjectSettings.h
--- a/Engine/Plugins/Runtime/ModelViewViewModel/Source/ModelViewViewModelBlueprint/Public/MVVMDeveloperProjectSettings.h
+++ b/Engine/Plugins/Runtime/ModelViewViewModel/Source/ModelViewViewModelBlueprint/Public/MVVMDeveloperProjectSettings.h
@@ -70,9 +70,9 @@ public:
 	virtual FText GetSectionText() const override;
 
 
-	bool PropertyHasFiltering(const FProperty* Property) const;
-	bool IsPropertyAllowed(const UBlueprint* Context, const FProperty* Property) const;
-	bool IsFunctionAllowed(const UBlueprint* Context, const UFunction* Function) const;
+	bool PropertyHasFiltering(const UStruct* ObjectStruct, const FProperty* Property) const;
+	bool IsPropertyAllowed(const UBlueprint* Context, const UStruct* ObjectStruct, const FProperty* Property) const;
+	bool IsFunctionAllowed(const UBlueprint* Context, const UClass* ObjectClass, const UFunction* Function) const;
 	bool IsConversionFunctionAllowed(const UBlueprint* Context, const UFunction* Function) const;
 
 	bool IsExecutionModeAllowed(EMVVMExecutionMode ExecutionMode) const
diff --git a/Engine/Plugins/Runtime/ModelViewViewModel/Source/ModelViewViewModelEditor/Private/Widgets/SMVVMViewModelBindingListWidget.cpp b/Engine/Plugins/Runtime/ModelViewViewModel/Source/ModelViewViewModelEditor/Private/Widgets/SMVVMViewModelBindingListWidget.cpp
--- a/Engine/Plugins/Runtime/ModelViewViewModel/Source/ModelViewViewModelEditor/Private/Widgets/SMVVMViewModelBindingListWidget.cpp
+++ b/Engine/Plugins/Runtime/ModelViewViewModel/Source/ModelViewViewModelEditor/Private/Widgets/SMVVMViewModelBindingListWidget.cpp
@@ -26,7 +26,7 @@ namespace UE::MVVM
 
 namespace Private
 {
-	TOptional<FFieldVariant> PassFilter(const UBlueprint* Blueprint, const FMVVMAvailableBinding& Binding, const FMVVMFieldVariant& FieldVariant, EFieldVisibility FieldVisibilityFlags, const FProperty* AssignableTo, bool bDoObjectProperty);
+	TOptional<FFieldVariant> PassFilter(const UBlueprint* Blueprint, const FMVVMAvailableBinding& Binding, const UStruct* Struct, const FMVVMFieldVariant& FieldVariant, EFieldVisibility FieldVisibilityFlags, const FProperty* AssignableTo, bool bDoObjectProperty);
 	TOptional<FFieldVariant> PassFilter(const UBlueprint* Blueprint, const FMVVMAvailableBinding& Binding, const UStruct* Struct, EFieldVisibility FieldVisibilityFlags, const FProperty* AssignableTo, bool bDoObjectProperty);
 
 	TOptional<FFieldVariant> PassFilter(const UBlueprint* Blueprint, const FMVVMAvailableBinding& Binding, const UStruct* Struct, EFieldVisibility FieldVisibilityFlags, const FProperty* AssignableTo, bool bDoObjectProperty)
@@ -42,10 +42,10 @@ namespace Private
 		}
 
 		FMVVMFieldVariant FieldVariant = BindingHelper::FindFieldByName(Struct, Binding.GetBindingName());
-		return PassFilter(Blueprint, Binding, FieldVariant, FieldVisibilityFlags, AssignableTo, bDoObjectProperty);
+		return PassFilter(Blueprint, Binding, Struct, FieldVariant, FieldVisibilityFlags, AssignableTo, bDoObjectProperty);
 	}
 
-	TOptional<FFieldVariant> PassFilter(const UBlueprint* Blueprint, const FMVVMAvailableBinding& Binding, const FMVVMFieldVariant& FieldVariant, EFieldVisibility FieldVisibilityFlags, const FProperty* AssignableTo, bool bDoObjectProperty)
+	TOptional<FFieldVariant> PassFilter(const UBlueprint* Blueprint, const FMVVMAvailableBinding& Binding, const UStruct* Struct, const FMVVMFieldVariant& FieldVariant, EFieldVisibility FieldVisibilityFlags, const FProperty* AssignableTo, bool bDoObjectProperty)
 	{
 		if (ensure(!FieldVariant.IsEmpty()))
 		{
@@ -93,7 +93,7 @@ namespace Private
 					}
 				}
 
-				if (!GetDefault<UMVVMDeveloperProjectSettings>()->IsFunctionAllowed(Blueprint, Function))
+				if (!GetDefault<UMVVMDeveloperProjectSettings>()->IsFunctionAllowed(Blueprint, Cast<const UClass>(Struct), Function))
 				{
 					return TOptional<FFieldVariant>();
 				}
@@ -144,7 +144,7 @@ namespace Private
 					}
 				}
 
-				if (!GetDefault<UMVVMDeveloperProjectSettings>()->IsPropertyAllowed(Blueprint, Property))
+				if (!GetDefault<UMVVMDeveloperProjectSettings>()->IsPropertyAllowed(Blueprint, Struct, Property))
 				{
 					return TOptional<FFieldVariant>();
 				}
