MVVM: Big rework of the MVVM view compiler. It add supports for conversion function to go from complex to simple. Supports pin split on simple conversion function. Add ability to to support setter function with more than one argument. Build the list of all the sources.

diff --git a/Engine/Plugins/Runtime/ModelViewViewModel/Source/ModelViewViewModel/Private/Types/MVVMBindingMode.cpp b/Engine/Plugins/Runtime/ModelViewViewModel/Source/ModelViewViewModel/Private/Types/MVVMBindingMode.cpp
--- a/Engine/Plugins/Runtime/ModelViewViewModel/Source/ModelViewViewModel/Private/Types/MVVMBindingMode.cpp
+++ b/Engine/Plugins/Runtime/ModelViewViewModel/Source/ModelViewViewModel/Private/Types/MVVMBindingMode.cpp
@@ -10,11 +10,13 @@ namespace UE::MVVM
 
 	bool IsForwardBinding(EMVVMBindingMode Mode)
 	{
+		//OneTimeToDestination || OneWayToDestination || TwoWay
 		return static_cast<uint8>(Mode) <= static_cast<uint8>(EMVVMBindingMode::TwoWay);
 	}
 
 	bool IsBackwardBinding(EMVVMBindingMode Mode)
 	{
+		//OneTimeToSource || OneWayToSource || TwoWay
 		return static_cast<uint8>(Mode) >= static_cast<uint8>(EMVVMBindingMode::TwoWay);
 	}
 
diff --git a/Engine/Plugins/Runtime/ModelViewViewModel/Source/ModelViewViewModel/Private/View/MVVMViewClass.cpp b/Engine/Plugins/Runtime/ModelViewViewModel/Source/ModelViewViewModel/Private/View/MVVMViewClass.cpp
--- a/Engine/Plugins/Runtime/ModelViewViewModel/Source/ModelViewViewModel/Private/View/MVVMViewClass.cpp
+++ b/Engine/Plugins/Runtime/ModelViewViewModel/Source/ModelViewViewModel/Private/View/MVVMViewClass.cpp
@@ -50,6 +50,10 @@ UObject* FMVVMViewClass_SourceCreator::CreateInstance(const UMVVMViewClass* InVi
 			Log.Error(FText::Format(LOCTEXT("CreateInstanceCreateInstance", "The source '{0}' could not be created. The class is not loaded."), FText::FromName(PropertyName)));
 		}
 	}
+	else if ((Flags & (uint8)ESourceFlags::SelfReference) != 0)
+	{ 
+		Result = InUserWidget;
+	}
 	else if (Resolver)
 	{
 		Result = Resolver->CreateInstance(ExpectedSourceType.Get(), InUserWidget, InView);
diff --git a/Engine/Plugins/Runtime/ModelViewViewModel/Source/ModelViewViewModel/Public/View/MVVMViewClass.h b/Engine/Plugins/Runtime/ModelViewViewModel/Source/ModelViewViewModel/Public/View/MVVMViewClass.h
--- a/Engine/Plugins/Runtime/ModelViewViewModel/Source/ModelViewViewModel/Public/View/MVVMViewClass.h
+++ b/Engine/Plugins/Runtime/ModelViewViewModel/Source/ModelViewViewModel/Public/View/MVVMViewClass.h
@@ -58,6 +58,11 @@ public:
 	{
 		return (Flags & (uint8)ESourceFlags::CanBeEvaluated) != 0;
 	}
+	
+	bool IsOptional() const
+	{
+		return (Flags & (uint8)ESourceFlags::IsOptional) != 0;
+	}
 
 	FName GetSourceName() const
 	{
@@ -125,6 +130,7 @@ private:
 		IsOptional = 1 << 2,
 		CanBeSet = 1 << 3,
 		CanBeEvaluated = 1 << 4,
+		SelfReference = 1 << 5,
 	};
 
 	UPROPERTY()
diff --git a/Engine/Plugins/Runtime/ModelViewViewModel/Source/ModelViewViewModelBlueprint/Private/Bindings/MVVMCompiledBindingLibraryCompiler.cpp b/Engine/Plugins/Runtime/ModelViewViewModel/Source/ModelViewViewModelBlueprint/Private/Bindings/MVVMCompiledBindingLibraryCompiler.cpp
--- a/Engine/Plugins/Runtime/ModelViewViewModel/Source/ModelViewViewModelBlueprint/Private/Bindings/MVVMCompiledBindingLibraryCompiler.cpp
+++ b/Engine/Plugins/Runtime/ModelViewViewModel/Source/ModelViewViewModelBlueprint/Private/Bindings/MVVMCompiledBindingLibraryCompiler.cpp
@@ -225,7 +225,7 @@ TValueOrError<FCompiledBindingLibraryCompiler::FFieldIdHandle, FText> FCompiledB
 
 	if (FieldId.IsNone())
 	{
-		return MakeError(LOCTEXT("FieldNotDefined", "The Field does not have the specifier FieldNotify and cannot be used as a binding source. You may want to use the 'One Time' binding mode."));
+		return MakeError(LOCTEXT("FieldNotDefined", "The Field does not have the specifier FieldNotify or cannot be used as a binding source. You may want to use the 'One Time' binding mode or make sure it is readable."));
 	}
 
 	if (!SourceClass->ImplementsInterface(UNotifyFieldValueChanged::StaticClass()))
@@ -436,7 +436,7 @@ TValueOrError<FCompiledBindingLibraryCompiler::FFieldPathHandle, FText> FCompile
 }
 
 
-TValueOrError<FCompiledBindingLibraryCompiler::FFieldPathHandle, FText> FCompiledBindingLibraryCompiler::AddObjectFieldPath(TArrayView<const UE::MVVM::FMVVMConstFieldVariant> FieldPath, UClass* ExpectedType, bool bInRead)
+TValueOrError<FCompiledBindingLibraryCompiler::FFieldPathHandle, FText> FCompiledBindingLibraryCompiler::AddObjectFieldPath(TArrayView<const UE::MVVM::FMVVMConstFieldVariant> FieldPath, const UClass* ExpectedType, bool bInRead)
 {
 	Impl->bCompiled = false;
 
diff --git a/Engine/Plugins/Runtime/ModelViewViewModel/Source/ModelViewViewModelBlueprint/Private/MVVMBlueprintPin.cpp b/Engine/Plugins/Runtime/ModelViewViewModel/Source/ModelViewViewModelBlueprint/Private/MVVMBlueprintPin.cpp
--- a/Engine/Plugins/Runtime/ModelViewViewModel/Source/ModelViewViewModelBlueprint/Private/MVVMBlueprintPin.cpp
+++ b/Engine/Plugins/Runtime/ModelViewViewModel/Source/ModelViewViewModelBlueprint/Private/MVVMBlueprintPin.cpp
@@ -4,6 +4,7 @@
 
 #include "Bindings/MVVMConversionFunctionHelper.h"
 #include "EdGraph/EdGraphPin.h"
+#include "EdGraphSchema_K2.h"
 
 #include UE_INLINE_GENERATED_CPP_BY_NAME(MVVMBlueprintPin)
 
@@ -15,11 +16,14 @@ FMVVMBlueprintPin::FMVVMBlueprintPin(FName InPinName)
 FMVVMBlueprintPin FMVVMBlueprintPin::CreateFromPin(const UBlueprint* Blueprint, const UEdGraphPin* Pin)
 {
 	FMVVMBlueprintPin Result;
+	Result.PinId = Pin->PinId;
 	Result.PinName = Pin->PinName;
 	Result.Path = UE::MVVM::ConversionFunctionHelper::GetPropertyPathForPin(Blueprint, Pin, true);
 	Result.DefaultObject = Pin->DefaultObject;
 	Result.DefaultString = Pin->DefaultValue;
 	Result.DefaultText = Pin->DefaultTextValue;
+	Result.bSplit = Pin->SubPins.Num() > 0;
+	Result.Status = Pin->bOrphanedPin ? EMVVMBlueprintPinStatus::Orphaned : EMVVMBlueprintPinStatus::Valid;
 	return Result;
 }
 
@@ -69,19 +73,90 @@ FString FMVVMBlueprintPin::GetValueAsString(const UClass* SelfContext) const
 	}
 }
 
-void FMVVMBlueprintPin::CopyTo(const UBlueprint* Blueprint, UEdGraphPin* Pin) const
+bool FMVVMBlueprintPin::IsInputPin(const UEdGraphPin* Pin)
 {
-	Pin->DefaultObject = DefaultObject;
-	Pin->DefaultValue = DefaultString;
-	Pin->DefaultTextValue = DefaultText;
+	return Pin->PinName != UEdGraphSchema_K2::PN_Self && Pin->PinName != UEdGraphSchema_K2::PN_Execute && Pin->Direction == EGPD_Input && (!Pin->bOrphanedPin || Pin->ShouldSavePinIfOrphaned());
+}
+
+TArray<FMVVMBlueprintPin> FMVVMBlueprintPin::CopyAndReturnMissingPins(UBlueprint* Blueprint, UEdGraphNode* GraphNode, const TArray<FMVVMBlueprintPin>& Pins)
+{
+	check(GraphNode);
+	check(Blueprint);
+
+	TSet<const UEdGraphPin*> AllGraphPins;
+	AllGraphPins.Reserve(GraphNode->Pins.Num());
+	for (const FMVVMBlueprintPin& Pin : Pins)
+	{
+		Pin.CopyTo(Blueprint, GraphNode);
+		if (UEdGraphPin* GraphPin = Pin.FindGraphPin(GraphNode))
+		{
+			AllGraphPins.Add(GraphPin);
+		}
+	}
+
+	// Create the missing pins.
+	TArray<FMVVMBlueprintPin> MissingPins;
+	for (const UEdGraphPin* GraphPin : GraphNode->Pins)
+	{
+		if (IsInputPin(GraphPin) && !AllGraphPins.Contains(GraphPin))
+		{
+			MissingPins.Add(FMVVMBlueprintPin::CreateFromPin(Blueprint, GraphPin));
+		}
+	}
+	return MissingPins;
+}
+
+void FMVVMBlueprintPin::CopyTo(const UBlueprint* Blueprint, UEdGraphNode* Node) const
+{
+	Status = EMVVMBlueprintPinStatus::Orphaned;
+	if (UEdGraphPin* GraphPin = FindGraphPin(Node))
+	{
+		if (IsInputPin(GraphPin) && !GraphPin->bOrphanedPin)
+		{
+			Status = EMVVMBlueprintPinStatus::Valid;
+			if (bSplit && GraphPin->SubPins.Num() == 0 && GetDefault<UEdGraphSchema_K2>()->CanSplitStructPin(*GraphPin))
+			{
+				GetDefault<UEdGraphSchema_K2>()->SplitPin(GraphPin, false);
+			}
 
-	UE::MVVM::ConversionFunctionHelper::SetPropertyPathForPin(Blueprint, Path, Pin);
+			GraphPin->DefaultObject = DefaultObject;
+			GraphPin->DefaultValue = DefaultString;
+			GraphPin->DefaultTextValue = DefaultText;
+			UE::MVVM::ConversionFunctionHelper::SetPropertyPathForPin(Blueprint, Path, GraphPin);
+		}
+	}
+}
+
+TArray<FMVVMBlueprintPin> FMVVMBlueprintPin::CreateFromNode(UBlueprint* Blueprint, UEdGraphNode* GraphNode)
+{
+	TArray<FMVVMBlueprintPin> Result;
+	Result.Reserve(GraphNode->Pins.Num());
+	for (const UEdGraphPin* GraphPin : GraphNode->Pins)
+	{
+		if (IsInputPin(GraphPin))
+		{
+			Result.Add(FMVVMBlueprintPin::CreateFromPin(Blueprint, GraphPin));
+		}
+	}
+	return Result;
+}
+
+UEdGraphPin* FMVVMBlueprintPin::FindGraphPin(UEdGraphNode* Node) const
+{
+	if (PinId.IsValid())
+	{
+		return Node->FindPinById(PinId);
+	}
+	return Node->FindPin(PinName, EGPD_Input);
 }
 
 void FMVVMBlueprintPin::Reset()
 {
+	PinId.Invalidate();
 	Path = FMVVMBlueprintPropertyPath();
 	DefaultString.Empty();
 	DefaultText = FText();
 	DefaultObject = nullptr;
+	bSplit = false;
+	Status = EMVVMBlueprintPinStatus::Valid;
 }
diff --git a/Engine/Plugins/Runtime/ModelViewViewModel/Source/ModelViewViewModelBlueprint/Private/MVVMBlueprintView.cpp b/Engine/Plugins/Runtime/ModelViewViewModel/Source/ModelViewViewModelBlueprint/Private/MVVMBlueprintView.cpp
--- a/Engine/Plugins/Runtime/ModelViewViewModel/Source/ModelViewViewModelBlueprint/Private/MVVMBlueprintView.cpp
+++ b/Engine/Plugins/Runtime/ModelViewViewModel/Source/ModelViewViewModelBlueprint/Private/MVVMBlueprintView.cpp
@@ -470,7 +470,7 @@ void UMVVMBlueprintView::PostLoad()
 
 	for (FMVVMBlueprintViewBinding& Binding : Bindings)
 	{
-		Binding.Conversion.DeprecateViewConversionFunction(GetOuterUMVVMWidgetBlueprintExtension_View()->GetWidgetBlueprint());
+		Binding.Conversion.DeprecateViewConversionFunction(GetOuterUMVVMWidgetBlueprintExtension_View()->GetWidgetBlueprint(), Binding);
 	}
 
 	if (GetLinkerCustomVersion(FFortniteMainBranchObjectVersion::GUID) < FFortniteMainBranchObjectVersion::MVVMPropertyPathSelf)
diff --git a/Engine/Plugins/Runtime/ModelViewViewModel/Source/ModelViewViewModelBlueprint/Private/MVVMBlueprintViewBinding.cpp b/Engine/Plugins/Runtime/ModelViewViewModel/Source/ModelViewViewModelBlueprint/Private/MVVMBlueprintViewBinding.cpp
--- a/Engine/Plugins/Runtime/ModelViewViewModel/Source/ModelViewViewModelBlueprint/Private/MVVMBlueprintViewBinding.cpp
+++ b/Engine/Plugins/Runtime/ModelViewViewModel/Source/ModelViewViewModelBlueprint/Private/MVVMBlueprintViewBinding.cpp
@@ -39,31 +39,47 @@ void FMVVMBlueprintViewConversionPath::SavePinValues(UBlueprint* Blueprint)
 	}
 }
 
-void FMVVMBlueprintViewConversionPath::DeprecateViewConversionFunction(UBlueprint* Blueprint)
+void FMVVMBlueprintViewConversionPath::DeprecateViewConversionFunction(UBlueprint* Blueprint, FMVVMBlueprintViewBinding& Owner)
 {
-	auto Deprecate = [Blueprint](FName& Wrapper, FMemberReference& Reference, TObjectPtr<UMVVMBlueprintViewConversionFunction>& ConversionFunction)
+	auto Deprecate = [&Owner, Blueprint](bool bSourceToDestination, FName& Wrapper, FMemberReference& Reference, TObjectPtr<UMVVMBlueprintViewConversionFunction>& ConversionFunction, FMVVMBlueprintPropertyPath& Path)
 	{
-		if (!Wrapper.IsNone())
+		if (ConversionFunction)
 		{
-			TObjectPtr<UEdGraph>* GraphPtr = Blueprint->FunctionGraphs.FindByPredicate([Wrapper](const UEdGraph* Other) { return Other->GetFName() == Wrapper; });
-			if (GraphPtr)
+			ConversionFunction->ConditionalPostLoad();
+			if (ConversionFunction->GetWrapperGraphName().IsNone())
 			{
-				ConversionFunction = NewObject<UMVVMBlueprintViewConversionFunction>(Blueprint);
-				ConversionFunction->InitializeFromWrapperGraph(Blueprint, *GraphPtr);
+				// The function was a simple conversion function but now we want to always have a wrapper
+				FName GraphName = UE::MVVM::ConversionFunctionHelper::CreateWrapperName(Owner, bSourceToDestination);
+				ConversionFunction->Deprecation_SetWrapperGraphName(Blueprint, GraphName, Path);
+				Path = FMVVMBlueprintPropertyPath();
 			}
 		}
-		else if (!Reference.GetMemberName().IsNone())
+		else
 		{
-			ConversionFunction = NewObject<UMVVMBlueprintViewConversionFunction>(Blueprint);
-			ConversionFunction->InitializeFromMemberReference(Blueprint, Reference);
+			if (!Wrapper.IsNone())
+			{
+				TObjectPtr<UEdGraph>* GraphPtr = Blueprint->FunctionGraphs.FindByPredicate([Wrapper](const UEdGraph* Other) { return Other->GetFName() == Wrapper; });
+				if (GraphPtr)
+				{
+					ConversionFunction = NewObject<UMVVMBlueprintViewConversionFunction>(Blueprint);
+					ConversionFunction->Deprecation_InitializeFromWrapperGraph(Blueprint, *GraphPtr);
+				}
+			}
+			else if (!Reference.GetMemberName().IsNone())
+			{
+				ConversionFunction = NewObject<UMVVMBlueprintViewConversionFunction>(Blueprint);
+				FName GraphName = UE::MVVM::ConversionFunctionHelper::CreateWrapperName(Owner, bSourceToDestination);
+				ConversionFunction->Deprecation_InitializeFromMemberReference(Blueprint, GraphName, Reference, Path);
+				Path = FMVVMBlueprintPropertyPath();
+			}
+			Wrapper = FName();
+			Reference = FMemberReference();
 		}
-		Wrapper = FName();
-		Reference = FMemberReference();
 	};
 
 	PRAGMA_DISABLE_DEPRECATION_WARNINGS
-	Deprecate(DestinationToSourceWrapper_DEPRECATED, DestinationToSourceFunction_DEPRECATED, DestinationToSourceConversion);
-	Deprecate(SourceToDestinationWrapper_DEPRECATED, SourceToDestinationFunction_DEPRECATED, SourceToDestinationConversion);
+	Deprecate(false, DestinationToSourceWrapper_DEPRECATED, DestinationToSourceFunction_DEPRECATED, DestinationToSourceConversion, Owner.DestinationPath);
+	Deprecate(true, SourceToDestinationWrapper_DEPRECATED, SourceToDestinationFunction_DEPRECATED, SourceToDestinationConversion, Owner.SourcePath);
 	PRAGMA_ENABLE_DEPRECATION_WARNINGS
 }
 
@@ -88,7 +104,7 @@ namespace UE::MVVM::Private
 
 		if (UMVVMBlueprintViewConversionFunction* ViewConversionFunction = Binding.Conversion.GetConversionFunction(bIsSource))
 		{
-			TVariant<const UFunction*, TSubclassOf<UK2Node>> ConversionFunction = ViewConversionFunction->GetConversionFunction(WidgetBlueprint->SkeletonGeneratedClass);
+			TVariant<const UFunction*, TSubclassOf<UK2Node>> ConversionFunction = ViewConversionFunction->GetConversionFunction(WidgetBlueprint);
 			if (ConversionFunction.IsType<const UFunction*>())
 			{
 				if (const UFunction* ConversionFunctionPtr = ConversionFunction.Get<const UFunction*>())
@@ -141,32 +157,24 @@ namespace UE::MVVM::Private
 			{
 				NameBuilder << TEXT("(");
 
-				if (ViewConversionFunction->NeedsWrapperGraph())
+				bool bFirst = true;
+				for (const FMVVMBlueprintPin& Pin : ViewConversionFunction->GetPins())
 				{
-					bool bFirst = true;
-					for (const FMVVMBlueprintPin& Pin : ViewConversionFunction->GetPins())
+					if (!bFirst)
 					{
-						if (!bFirst)
-						{
-							NameBuilder << TEXT(", ");
-						}
-
-						if (Pin.UsedPathAsValue())
-						{
-							AddPath(Pin.GetPath());
-						}
-						else
-						{
-							NameBuilder << Pin.GetValueAsString(WidgetBlueprint->SkeletonGeneratedClass);
-						}
+						NameBuilder << TEXT(", ");
+					}
 
-						bFirst = false;
+					if (Pin.UsedPathAsValue())
+					{
+						AddPath(Pin.GetPath());
 					}
-				}
-				else
-				{
-					const FMVVMBlueprintPropertyPath& Path = bIsSource ? Binding.SourcePath : Binding.DestinationPath;
-					AddPath(Path);
+					else
+					{
+						NameBuilder << Pin.GetValueAsString(WidgetBlueprint->SkeletonGeneratedClass);
+					}
+
+					bFirst = false;
 				}
 
 				NameBuilder << TEXT(")");
diff --git a/Engine/Plugins/Runtime/ModelViewViewModel/Source/ModelViewViewModelBlueprint/Private/MVVMBlueprintViewConversionFunction.cpp b/Engine/Plugins/Runtime/ModelViewViewModel/Source/ModelViewViewModelBlueprint/Private/MVVMBlueprintViewConversionFunction.cpp
--- a/Engine/Plugins/Runtime/ModelViewViewModel/Source/ModelViewViewModelBlueprint/Private/MVVMBlueprintViewConversionFunction.cpp
+++ b/Engine/Plugins/Runtime/ModelViewViewModel/Source/ModelViewViewModelBlueprint/Private/MVVMBlueprintViewConversionFunction.cpp
@@ -3,6 +3,7 @@
 #include "MVVMBlueprintViewConversionFunction.h"
 #include "MVVMDeveloperProjectSettings.h"
 
+#include "Bindings/MVVMBindingHelper.h"
 #include "Bindings/MVVMConversionFunctionHelper.h"
 #include "Bindings/MVVMBindingHelper.h"
 #include "Engine/Blueprint.h"
@@ -10,10 +11,13 @@
 
 #include "EdGraphSchema_K2.h"
 #include "EdGraph/EdGraphPin.h"
+#include "GraphEditAction.h"
 #include "K2Node_CallFunction.h"
 #include "K2Node_FunctionEntry.h"
 #include "K2Node_FunctionResult.h"
 #include "KismetCompiler.h"
+#include "Kismet/BlueprintFunctionLibrary.h"
+#include "Templates/ValueOrError.h"
 
 #include UE_INLINE_GENERATED_CPP_BY_NAME(MVVMBlueprintViewConversionFunction)
 
@@ -22,9 +26,92 @@ namespace UE::MVVM::Private
 	static FLazyName DefaultConversionFunctionName = TEXT("__ConversionFunction");
 }
 
+bool UMVVMBlueprintViewConversionFunction::IsValidConversionFunction(const UBlueprint* WidgetBlueprint, const UFunction* Function)
+{
+	// functions in the widget blueprint can do anything they want, other functions have to be static functions in a BlueprintFunctionLibrary
+	const UClass* FunctionClass = Function->GetOuterUClass();
+
+	bool bIsFromWidgetBlueprint = WidgetBlueprint->GeneratedClass && WidgetBlueprint->GeneratedClass->IsChildOf(FunctionClass) && Function->HasAllFunctionFlags(FUNC_BlueprintPure | FUNC_Const);
+	bool bIsFromSkeletonWidgetBlueprint = WidgetBlueprint->SkeletonGeneratedClass && WidgetBlueprint->SkeletonGeneratedClass->IsChildOf(FunctionClass) && Function->HasAllFunctionFlags(FUNC_BlueprintPure | FUNC_Const);
+	bool bFromBlueprintFunctionLibrary = FunctionClass->IsChildOf<UBlueprintFunctionLibrary>() && Function->HasAllFunctionFlags(FUNC_Static | FUNC_BlueprintPure);
+	if (!bIsFromWidgetBlueprint && !bIsFromSkeletonWidgetBlueprint && !bFromBlueprintFunctionLibrary)
+	{
+		return false;
+	}
+
+	TValueOrError<const FProperty*, FText> ReturnResult = UE::MVVM::BindingHelper::TryGetReturnTypeForConversionFunction(Function);
+	if (ReturnResult.HasError())
+	{
+		return false;
+	}
+
+	const FProperty* ReturnProperty = ReturnResult.GetValue();
+	if (ReturnProperty == nullptr)
+	{
+		return false;
+	}
+
+	TValueOrError<TArray<const FProperty*>, FText> ArgumentsResult = UE::MVVM::BindingHelper::TryGetArgumentsForConversionFunction(Function);
+	if (ArgumentsResult.HasError())
+	{
+		return false;
+	}
+
+	return GetDefault<UMVVMDeveloperProjectSettings>()->IsConversionFunctionAllowed(WidgetBlueprint, Function);
+}
+
+bool UMVVMBlueprintViewConversionFunction::IsValid(const UBlueprint* SelfContext) const
+{
+	if (FunctionNode.Get())
+	{
+		return false; // not valid for now
+	}
+	
+	UClass* SelfClass = SelfContext->SkeletonGeneratedClass ? SelfContext->SkeletonGeneratedClass : SelfContext->GeneratedClass;
+	const UFunction* ConversionFunction = FunctionReference.ResolveMember<UFunction>(SelfClass);
+	if (ConversionFunction)
+	{
+		if (!GetDefault<UMVVMDeveloperProjectSettings>()->IsConversionFunctionAllowed(SelfContext, ConversionFunction))
+		{
+			return false;
+		}
+
+		return IsValidConversionFunction(SelfContext, ConversionFunction);
+	}
+	return false;
+}
+
+bool UMVVMBlueprintViewConversionFunction::NeedsWrapperGraph(const UBlueprint* SelfContext) const
+{
+	UClass* SelfClass = SelfContext->SkeletonGeneratedClass ? SelfContext->SkeletonGeneratedClass : SelfContext->GeneratedClass;
+	return NeedsWrapperGraphInternal(SelfClass);
+}
+
+bool UMVVMBlueprintViewConversionFunction::NeedsWrapperGraphInternal(const UClass* SkeletalSelfContext) const
+{
+	if (FunctionNode.Get())
+	{
+		return true; // not valid for now
+	}
+
+
+	const UFunction* ConversionFunction = FunctionReference.ResolveMember<UFunction>(const_cast<UClass*>(SkeletalSelfContext));
+	if (ensure(ConversionFunction))
+	{
+		return SavedPins.Num() > 1 || !UE::MVVM::BindingHelper::IsValidForSimpleRuntimeConversion(ConversionFunction);
+	}
+	return false;
+}
+
+/** The wrapper Graph is generated on domains and is not saved. */
+bool UMVVMBlueprintViewConversionFunction::IsWrapperGraphTransient() const
+{
+	return bWrapperGraphTransient;
+}
+
 const UFunction* UMVVMBlueprintViewConversionFunction::GetCompiledFunction(const UClass* SelfContext) const
 {
-	if (NeedsWrapperGraph())
+	if (NeedsWrapperGraphInternal(SelfContext))
 	{
 		FMemberReference CompiledFunction;
 		CompiledFunction.SetSelfMember(GraphName);
@@ -33,16 +120,17 @@ const UFunction* UMVVMBlueprintViewConversionFunction::GetCompiledFunction(const
 	return FunctionReference.ResolveMember<UFunction>(const_cast<UClass*>(SelfContext));
 }
 
-FName UMVVMBlueprintViewConversionFunction::GetCompiledFunctionName() const
+FName UMVVMBlueprintViewConversionFunction::GetCompiledFunctionName(const UClass* SelfContext) const
 {
-	if (NeedsWrapperGraph())
+	if (NeedsWrapperGraphInternal(SelfContext))
 	{
+		ensure(!GraphName.IsNone());
 		return GraphName;
 	}
 	return FunctionReference.GetMemberName();
 }
 
-TVariant<const UFunction*, TSubclassOf<UK2Node>> UMVVMBlueprintViewConversionFunction::GetConversionFunction(const UClass* SelfContext) const
+TVariant<const UFunction*, TSubclassOf<UK2Node>> UMVVMBlueprintViewConversionFunction::GetConversionFunction(const UBlueprint* SelfContext) const
 {
 	if (FunctionNode.Get())
 	{
@@ -50,15 +138,10 @@ TVariant<const UFunction*, TSubclassOf<UK2Node>> UMVVMBlueprintViewConversionFun
 	}
 	else
 	{
-		const UFunction* ConversionFunction = FunctionReference.ResolveMember<UFunction>(const_cast<UClass*>(SelfContext));
+		UClass* SelfClass = SelfContext->SkeletonGeneratedClass ? SelfContext->SkeletonGeneratedClass : SelfContext->GeneratedClass;
+		const UFunction* ConversionFunction = FunctionReference.ResolveMember<UFunction>(SelfClass);
 		return TVariant<const UFunction*, TSubclassOf<UK2Node>>( TInPlaceType<const UFunction*>(), ConversionFunction);
 	}
-	//return FunctionReference.ResolveMember<UFunction>(WidgetBlueprint->SkeletonGeneratedClass);
-//UK2Node_CallFunction* CallFunctionNode = GetConversionFunctionNode(WidgetBlueprint, Binding, bSourceToDestination);
-//if (CallFunctionNode)
-//{
-//	return CallFunctionNode->GetTargetFunction();
-//}
 }
 
 void UMVVMBlueprintViewConversionFunction::Reset()
@@ -68,20 +151,14 @@ void UMVVMBlueprintViewConversionFunction::Reset()
 	GraphName = FName();
 	bWrapperGraphTransient = false;
 	SavedPins.Reset();
-	CachedWrapperGraph = nullptr;
-	CachedWrapperNode = nullptr;
-}
-
-void UMVVMBlueprintViewConversionFunction::InitFromFunction(UBlueprint* InContext, const UFunction* InFunction)
-{
-	InitFromFunction(InContext, InFunction, UE::MVVM::Private::DefaultConversionFunctionName.Resolve());
+	SetCachedWrapperGraph(nullptr, nullptr, nullptr);
 }
 
-void UMVVMBlueprintViewConversionFunction::InitFromFunction(UBlueprint* InContext, const UFunction* InFunction, FName InGraphName)
+void UMVVMBlueprintViewConversionFunction::InitializeFromFunction(UBlueprint* InContext, FName InGraphName, const UFunction* InFunction)
 {
 	Reset();
 
-	if (InFunction)
+	if (ensure(InFunction && !InGraphName.IsNone()))
 	{
 		UClass* OwnerClass = InFunction->GetOuterUClass();
 		FGuid MemberGuid;
@@ -97,24 +174,21 @@ void UMVVMBlueprintViewConversionFunction::InitFromFunction(UBlueprint* InContex
 		{
 			if (UBlueprint* VariableOwnerBP = Cast<UBlueprint>(OwnerClass->ClassGeneratedBy))
 			{
-				OwnerClass = VariableOwnerBP->SkeletonGeneratedClass;
+				OwnerClass = VariableOwnerBP->SkeletonGeneratedClass ? VariableOwnerBP->SkeletonGeneratedClass : VariableOwnerBP->GeneratedClass;
 			}
 
 			FunctionReference.SetExternalMember(InFunction->GetFName(), OwnerClass, MemberGuid);
 		}
-		// used to be SetFromField<UFunction>(Function, SelfContext->SkeletonGeneratedClass);
 
-		if (UE::MVVM::ConversionFunctionHelper::RequiresWrapper(InFunction))
-		{
-			GraphName = InGraphName;
-			bWrapperGraphTransient = !GetDefault<UMVVMDeveloperProjectSettings>()->bAllowConversionFunctionGeneratedGraphInEditor;
-			GetOrCreateWrapperGraphInternal(InContext, InFunction);
-			SavePinValues(InContext);
-		}
+		check(GraphName.IsNone()); // the name needs to be set before a GetOrCreateWrapperGraph
+		GraphName = InGraphName;
+		bWrapperGraphTransient = !GetDefault<UMVVMDeveloperProjectSettings>()->bAllowConversionFunctionGeneratedGraphInEditor;
+		GetOrCreateWrapperGraphInternal(InContext, InFunction);
+		SavePinValues(InContext);
 	}
 }
 
-void UMVVMBlueprintViewConversionFunction::InitializeFromWrapperGraph(UBlueprint* SelfContext, UEdGraph* Graph)
+void UMVVMBlueprintViewConversionFunction::Deprecation_InitializeFromWrapperGraph(UBlueprint* SelfContext, UEdGraph* Graph)
 {
 	Reset();
 
@@ -128,13 +202,13 @@ void UMVVMBlueprintViewConversionFunction::InitializeFromWrapperGraph(UBlueprint
 		{
 			FunctionNode = WrapperNode->GetClass();
 		}
-		CachedWrapperGraph = WrapperNode->GetGraph();
-		check(CachedWrapperGraph);
+		check(WrapperNode->GetGraph());
+		SetCachedWrapperGraph(SelfContext, WrapperNode->GetGraph(), WrapperNode);
 
+		check(GraphName.IsNone());
 		GraphName = CachedWrapperGraph->GetFName();
 		bWrapperGraphTransient = !GetDefault<UMVVMDeveloperProjectSettings>()->bAllowConversionFunctionGeneratedGraphInEditor;
 
-		CachedWrapperNode = WrapperNode;
 		SavePinValues(SelfContext);
 
 		if (bWrapperGraphTransient && CachedWrapperNode)
@@ -144,88 +218,161 @@ void UMVVMBlueprintViewConversionFunction::InitializeFromWrapperGraph(UBlueprint
 	}
 }
 
-void UMVVMBlueprintViewConversionFunction::InitializeFromMemberReference(UBlueprint* SelfContext, FMemberReference MemberReference)
+void UMVVMBlueprintViewConversionFunction::Deprecation_InitializeFromMemberReference(UBlueprint* SelfContext, FName InGraphName, FMemberReference MemberReference, const FMVVMBlueprintPropertyPath& Source)
 {
 	Reset();
 
 	FunctionReference = MemberReference;
-}
 
-UEdGraph* UMVVMBlueprintViewConversionFunction::GetOrCreateIntermediateWrapperGraph(FKismetCompilerContext& Context) const
-{
-	if (!NeedsWrapperGraph())
+	check(GraphName.IsNone()); // the name needs to be set before a GetOrCreateWrapperGraph
+	GraphName = InGraphName;
+	bWrapperGraphTransient = !GetDefault<UMVVMDeveloperProjectSettings>()->bAllowConversionFunctionGeneratedGraphInEditor;
+
+	// since it is a new object, we can't create a the graph right away
+	UClass* GeneratedClass = SelfContext->SkeletonGeneratedClass ? SelfContext->SkeletonGeneratedClass : SelfContext->GeneratedClass;
+	const UFunction* ConversionFunction = FunctionReference.ResolveMember<UFunction>(GeneratedClass);
+	const FProperty* PinProperty = ConversionFunction ? UE::MVVM::BindingHelper::GetFirstArgumentProperty(ConversionFunction) : nullptr;
+	if (PinProperty)
 	{
-		return nullptr;
+		FMVVMBlueprintPin NewPin = PinProperty->GetFName();
+		NewPin.SetPath(Source);
+		SavedPins.Add(MoveTemp(NewPin));
 	}
-
-	if (CachedWrapperGraph)
+	else
 	{
-		return CachedWrapperGraph;
+		SavedPins.Reset();
 	}
+}
 
-	if (IsWrapperGraphTransient())
+void UMVVMBlueprintViewConversionFunction::Deprecation_SetWrapperGraphName(UBlueprint* SelfContext, FName InGraphName, const FMVVMBlueprintPropertyPath& Source)
+{
+	if (ensure(SavedPins.Num() == 0) && ensure(GraphName.IsNone()))
 	{
-		if (FunctionNode.Get())
+		GraphName = InGraphName;
+		bWrapperGraphTransient = !GetDefault<UMVVMDeveloperProjectSettings>()->bAllowConversionFunctionGeneratedGraphInEditor;
+
+		// since it is a new object, we can't create a the graph right away
+		UClass* GeneratedClass = SelfContext->SkeletonGeneratedClass ? SelfContext->SkeletonGeneratedClass : SelfContext->GeneratedClass;
+		const UFunction* ConversionFunction = FunctionReference.ResolveMember<UFunction>(GeneratedClass);
+		const FProperty* PinProperty = ConversionFunction ? UE::MVVM::BindingHelper::GetFirstArgumentProperty(ConversionFunction) : nullptr;
+		if (PinProperty)
 		{
-			ensure(false); // not supported for now
-			return nullptr;
+			FMVVMBlueprintPin NewPin = PinProperty->GetFName();
+			NewPin.SetPath(Source);
+			SavedPins.Add(MoveTemp(NewPin));
 		}
 		else
 		{
-			const UFunction* ConversionFunction = FunctionReference.ResolveMember<UFunction>(Context.NewClass);
-			return GetOrCreateWrapperGraphInternal(Context, ConversionFunction);
+			SavedPins.Reset();
 		}
 	}
-	else
+}
+
+void UMVVMBlueprintViewConversionFunction::SetCachedWrapperGraph(UBlueprint* Context, UEdGraph* Graph, UK2Node* Node)
+{
+	if (CachedWrapperNode && OnUserDefinedPinRenamedHandle.IsValid())
 	{
-		TObjectPtr<UEdGraph>* Found = Context.Blueprint->FunctionGraphs.FindByPredicate([GraphName = GetWrapperGraphName()](const UEdGraph* Other) { return Other->GetFName() == GraphName; });
-		if (Found)
-		{
-			CachedWrapperGraph = *Found;
-			CachedWrapperNode = UE::MVVM::ConversionFunctionHelper::GetWrapperNode(CachedWrapperGraph);
-		}
-		return CachedWrapperGraph;
+		CachedWrapperNode->OnUserDefinedPinRenamed().Remove(OnUserDefinedPinRenamedHandle);
+	}
+	if (CachedWrapperGraph && OnGraphChangedHandle.IsValid())
+	{
+		CachedWrapperGraph->RemoveOnGraphChangedHandler(OnGraphChangedHandle);
+	}
+
+	CachedWrapperGraph = Graph;
+	CachedWrapperNode = Node;
+	OnGraphChangedHandle.Reset();
+	OnUserDefinedPinRenamedHandle.Reset();
+
+	if (CachedWrapperGraph && Context)
+	{
+		TWeakObjectPtr<UBlueprint> WeakContext = Context;
+		OnGraphChangedHandle = CachedWrapperGraph->AddOnGraphChangedHandler(FOnGraphChanged::FDelegate::CreateUObject(this, &UMVVMBlueprintViewConversionFunction::HandleGraphChanged, WeakContext));
+	}
+	if (CachedWrapperNode && Context)
+	{
+		TWeakObjectPtr<UBlueprint> WeakContext = Context;
+		OnUserDefinedPinRenamedHandle = CachedWrapperNode->OnUserDefinedPinRenamed().AddUObject(this, &UMVVMBlueprintViewConversionFunction::HandleUserDefinedPinRenamed, WeakContext);
 	}
 }
 
-UEdGraph* UMVVMBlueprintViewConversionFunction::GetOrCreateWrapperGraph(UBlueprint* Blueprint) const
+void UMVVMBlueprintViewConversionFunction::CreateWrapperGraphName()
 {
-	if (!NeedsWrapperGraph())
+	ensure(!GraphName.IsNone());
+	if (GraphName.IsNone())
 	{
-		return nullptr;
+		TStringBuilder<256> StringBuilder;
+		StringBuilder << UE::MVVM::Private::DefaultConversionFunctionName.Resolve();
+		StringBuilder << FGuid::NewDeterministicGuid(GetFullName()).ToString(EGuidFormats::DigitsWithHyphensLower);
+		GraphName = StringBuilder.ToString();
 	}
+}
+
+UEdGraph* UMVVMBlueprintViewConversionFunction::GetOrCreateIntermediateWrapperGraph(FKismetCompilerContext& Context)
+{
+	check(Context.NewClass);
 
 	if (CachedWrapperGraph)
 	{
 		return CachedWrapperGraph;
 	}
 
-	if (IsWrapperGraphTransient())
+	TObjectPtr<UEdGraph>* FoundGraph = !GraphName.IsNone() ? Context.Blueprint->FunctionGraphs.FindByPredicate([GraphName = GetWrapperGraphName()](const UEdGraph* Other) { return Other->GetFName() == GraphName; }) : nullptr;
+	if (FoundGraph)
 	{
-		if (FunctionNode.Get())
+		ensureMsgf(!IsWrapperGraphTransient(), TEXT("The graph is transient. It should not be saved in the editor."));
+		UBlueprint* NullContext = nullptr; // do not register the callback
+		const_cast<UMVVMBlueprintViewConversionFunction*>(this)->SetCachedWrapperGraph(NullContext, *FoundGraph, UE::MVVM::ConversionFunctionHelper::GetWrapperNode(*FoundGraph));
+		LoadPinValuesInternal(Context.Blueprint);
+		return CachedWrapperGraph;
+	}
+	else if (IsValid(Context.Blueprint))
+	{
+		TVariant<const UFunction*, TSubclassOf<UK2Node>> ConversionFunction = GetConversionFunction(Context.Blueprint);
+		if (ConversionFunction.IsType<TSubclassOf<UK2Node>>())
 		{
 			ensure(false); // not supported for now
 			return nullptr;
 		}
-		else
-		{
-			const UFunction* ConversionFunction = FunctionReference.ResolveMember<UFunction>(Blueprint->SkeletonGeneratedClass);
-			return GetOrCreateWrapperGraphInternal(Blueprint, ConversionFunction);
-		}
+
+		CreateWrapperGraphName();
+		return GetOrCreateWrapperGraphInternal(Context, ConversionFunction.Get<const UFunction*>());
 	}
-	else
+	return nullptr;
+}
+
+UEdGraph* UMVVMBlueprintViewConversionFunction::GetOrCreateWrapperGraph(UBlueprint* Blueprint)
+{
+	check(Blueprint);
+
+	if (CachedWrapperGraph)
+	{
+		return CachedWrapperGraph;
+	}
+
+	TObjectPtr<UEdGraph>* FoundGraph = Blueprint->FunctionGraphs.FindByPredicate([GraphName = GetWrapperGraphName()](const UEdGraph* Other) { return Other->GetFName() == GraphName; });
+	if (FoundGraph)
+	{
+		ensureMsgf(!IsWrapperGraphTransient(), TEXT("The graph is transient. It should not be saved in the editor."));
+		const_cast<UMVVMBlueprintViewConversionFunction*>(this)->SetCachedWrapperGraph(Blueprint, *FoundGraph, UE::MVVM::ConversionFunctionHelper::GetWrapperNode(*FoundGraph));
+		LoadPinValuesInternal(Blueprint);
+	}
+	else if (IsValid(Blueprint))
 	{
-		TObjectPtr<UEdGraph>* Found = Blueprint->FunctionGraphs.FindByPredicate([GraphName = GetWrapperGraphName()](const UEdGraph* Other) { return Other->GetFName() == GraphName; });
-		if (Found)
+		TVariant<const UFunction*, TSubclassOf<UK2Node>> ConversionFunction = GetConversionFunction(Blueprint);
+		if (ConversionFunction.IsType<TSubclassOf<UK2Node>>())
 		{
-			CachedWrapperGraph = *Found;
-			CachedWrapperNode = UE::MVVM::ConversionFunctionHelper::GetWrapperNode(CachedWrapperGraph);
+			ensure(false); // not supported for now
+			return nullptr;
 		}
-		return CachedWrapperGraph;
+
+		CreateWrapperGraphName();
+		return GetOrCreateWrapperGraphInternal(Blueprint, ConversionFunction.Get<const UFunction*>());
 	}
+	return nullptr;
 }
 
-UEdGraphPin* UMVVMBlueprintViewConversionFunction::GetOrCreateGraphPin(UBlueprint* Blueprint, FName PinName) const
+UEdGraphPin* UMVVMBlueprintViewConversionFunction::GetOrCreateGraphPin(UBlueprint* Blueprint, FName PinName)
 {
 	GetOrCreateWrapperGraph(Blueprint);
 	if (CachedWrapperNode)
@@ -235,17 +382,18 @@ UEdGraphPin* UMVVMBlueprintViewConversionFunction::GetOrCreateGraphPin(UBlueprin
 	return nullptr;
 }
 
-UEdGraph* UMVVMBlueprintViewConversionFunction::GetOrCreateWrapperGraphInternal(FKismetCompilerContext& Context, const UFunction* Function) const
+UEdGraph* UMVVMBlueprintViewConversionFunction::GetOrCreateWrapperGraphInternal(FKismetCompilerContext& Context, const UFunction* Function)
 {
 	return GetOrCreateWrapperGraphInternal(Context.Blueprint, Function);
 }
 
-UEdGraph* UMVVMBlueprintViewConversionFunction::GetOrCreateWrapperGraphInternal(UBlueprint* Blueprint, const UFunction* Function) const
+UEdGraph* UMVVMBlueprintViewConversionFunction::GetOrCreateWrapperGraphInternal(UBlueprint* Blueprint, const UFunction* Function)
 {
+	check(Blueprint);
+	check(Function);
 	bool bConst = true;
 	TPair<UEdGraph*, UK2Node*> Result = UE::MVVM::ConversionFunctionHelper::CreateGraph(Blueprint, GraphName, nullptr, Function, bConst, bWrapperGraphTransient);
-	CachedWrapperGraph = Result.Get<0>();
-	CachedWrapperNode = Result.Get<1>();
+	const_cast<UMVVMBlueprintViewConversionFunction*>(this)->SetCachedWrapperGraph(Blueprint, Result.Get<0>(), Result.Get<1>());
 	LoadPinValuesInternal(Blueprint);
 	return CachedWrapperGraph;
 }
@@ -257,17 +405,11 @@ void UMVVMBlueprintViewConversionFunction::RemoveWrapperGraph(UBlueprint* Bluepr
 	{
 		FBlueprintEditorUtils::RemoveGraph(Blueprint, Result->Get());
 	}
-	CachedWrapperGraph = nullptr;
-	CachedWrapperNode = nullptr;
+	SetCachedWrapperGraph(Blueprint, nullptr, nullptr);
 }
 
 void UMVVMBlueprintViewConversionFunction::SetGraphPin(UBlueprint* Blueprint, FName PinName, const FMVVMBlueprintPropertyPath& Path)
 {
-	if (!NeedsWrapperGraph())
-	{
-		return;
-	}
-
 	UEdGraphPin* GraphPin = GetOrCreateGraphPin(Blueprint, PinName);
 
 	// Set the value and make the blueprint as dirty before creating the pin.
@@ -288,27 +430,53 @@ void UMVVMBlueprintViewConversionFunction::SavePinValues(UBlueprint* Blueprint)
 {
 	if (CachedWrapperNode)
 	{
-		SavedPins.Reset();
-		for (UEdGraphPin* Pin : CachedWrapperNode->Pins)
-		{
-			if (Pin->PinName != UEdGraphSchema_K2::PN_Self && Pin->PinName != UEdGraphSchema_K2::PN_Execute && Pin->Direction == EGPD_Input)
-			{
-				SavedPins.Add(FMVVMBlueprintPin::CreateFromPin(Blueprint, Pin));
-			}
-		}
+		SavedPins = FMVVMBlueprintPin::CreateFromNode(Blueprint, CachedWrapperNode);
 	}
 }
 
-void UMVVMBlueprintViewConversionFunction::LoadPinValuesInternal(UBlueprint* Blueprint) const
+void UMVVMBlueprintViewConversionFunction::LoadPinValuesInternal(UBlueprint* Blueprint)
 {
 	if (CachedWrapperNode)
 	{
-		for (const FMVVMBlueprintPin& Pin : SavedPins)
+		TArray<FMVVMBlueprintPin> MissingPins = FMVVMBlueprintPin::CopyAndReturnMissingPins(Blueprint, CachedWrapperNode, SavedPins);
+		SavedPins.Append(MissingPins);
+	}
+}
+
+void UMVVMBlueprintViewConversionFunction::HandleGraphChanged(const FEdGraphEditAction& EditAction, TWeakObjectPtr<UBlueprint> WeakBlueprint)
+{
+	UBlueprint* Blueprint = WeakBlueprint.Get();
+	if (Blueprint && EditAction.Graph == CachedWrapperGraph && CachedWrapperGraph)
+	{
+		if (CachedWrapperNode && EditAction.Nodes.Contains(CachedWrapperNode))
 		{
-			if (UEdGraphPin* GraphPin = CachedWrapperNode->FindPin(Pin.GetName()))
+			if (EditAction.Action == EEdGraphActionType::GRAPHACTION_RemoveNode)
+			{
+				CachedWrapperNode = UE::MVVM::ConversionFunctionHelper::GetWrapperNode(CachedWrapperGraph);
+				SavePinValues(Blueprint);
+				OnWrapperGraphModified.Broadcast();
+			}
+			else if (EditAction.Action == EEdGraphActionType::GRAPHACTION_EditNode)
 			{
-				Pin.CopyTo(Blueprint, GraphPin);
+				SavePinValues(Blueprint);
+				OnWrapperGraphModified.Broadcast();
 			}
 		}
+		else if (CachedWrapperNode == nullptr && EditAction.Action == EEdGraphActionType::GRAPHACTION_AddNode)
+		{
+			CachedWrapperNode = UE::MVVM::ConversionFunctionHelper::GetWrapperNode(CachedWrapperGraph);
+			SavePinValues(Blueprint);
+			OnWrapperGraphModified.Broadcast();
+		}
+	}
+}
+
+void UMVVMBlueprintViewConversionFunction::HandleUserDefinedPinRenamed(UK2Node* InNode, FName OldPinName, FName NewPinName, TWeakObjectPtr<UBlueprint> WeakBlueprint)
+{
+	UBlueprint* Blueprint = WeakBlueprint.Get();
+	if (Blueprint && InNode == CachedWrapperNode)
+	{
+		SavePinValues(Blueprint);
+		OnWrapperGraphModified.Broadcast();
 	}
-}
\ No newline at end of file
+}
diff --git a/Engine/Plugins/Runtime/ModelViewViewModel/Source/ModelViewViewModelBlueprint/Private/MVVMBlueprintViewEvent.cpp b/Engine/Plugins/Runtime/ModelViewViewModel/Source/ModelViewViewModelBlueprint/Private/MVVMBlueprintViewEvent.cpp
--- a/Engine/Plugins/Runtime/ModelViewViewModel/Source/ModelViewViewModelBlueprint/Private/MVVMBlueprintViewEvent.cpp
+++ b/Engine/Plugins/Runtime/ModelViewViewModel/Source/ModelViewViewModelBlueprint/Private/MVVMBlueprintViewEvent.cpp
@@ -13,6 +13,7 @@
 
 #include "EdGraphSchema_K2.h"
 #include "EdGraph/EdGraphPin.h"
+#include "GraphEditAction.h"
 #include "K2Node_CallFunction.h"
 #include "K2Node_FunctionEntry.h"
 #include "K2Node_FunctionResult.h"
@@ -61,6 +62,32 @@ void UMVVMBlueprintViewEvent::SetDestinationPath(FMVVMBlueprintPropertyPath InDe
 	CreateWrapperGraphInternal();
 	SavePinValues();
 }
+
+void UMVVMBlueprintViewEvent::SetCachedWrapperGraphInternal(UEdGraph* Graph, UK2Node* Node)
+{
+	if (CachedWrapperNode && OnUserDefinedPinRenamedHandle.IsValid())
+	{
+		CachedWrapperNode->OnUserDefinedPinRenamed().Remove(OnUserDefinedPinRenamedHandle);
+	}
+	if (CachedWrapperGraph && OnGraphChangedHandle.IsValid())
+	{
+		CachedWrapperGraph->RemoveOnGraphChangedHandler(OnGraphChangedHandle);
+	}
+
+	CachedWrapperGraph = Graph;
+	CachedWrapperNode = Node;
+	OnGraphChangedHandle.Reset();
+	OnUserDefinedPinRenamedHandle.Reset();
+
+	if (CachedWrapperGraph)
+	{
+		OnGraphChangedHandle = CachedWrapperGraph->AddOnGraphChangedHandler(FOnGraphChanged::FDelegate::CreateUObject(this, &UMVVMBlueprintViewEvent::HandleGraphChanged));
+	}
+	if (CachedWrapperNode)
+	{
+		OnUserDefinedPinRenamedHandle = CachedWrapperNode->OnUserDefinedPinRenamed().AddUObject(this, &UMVVMBlueprintViewEvent::HandleUserDefinedPinRenamed);
+	}
+}
 	
 UEdGraph* UMVVMBlueprintViewEvent::GetOrCreateWrapperGraph()
 {
@@ -78,8 +105,7 @@ void UMVVMBlueprintViewEvent::RemoveWrapperGraph()
 	if (CachedWrapperGraph)
 	{
 		FBlueprintEditorUtils::RemoveGraph(GetWidgetBlueprintInternal(), CachedWrapperGraph);
-		CachedWrapperGraph = nullptr;
-		CachedWrapperNode = nullptr;
+		SetCachedWrapperGraphInternal(nullptr, nullptr);
 	}
 
 	Messages.Empty();
@@ -100,15 +126,8 @@ void UMVVMBlueprintViewEvent::SavePinValues()
 {
 	if (CachedWrapperNode)
 	{
-		SavedPins.Reset();
 		UWidgetBlueprint* Blueprint = GetWidgetBlueprintInternal();
-		for (UEdGraphPin* Pin : CachedWrapperNode->Pins)
-		{
-			if (Pin->PinName != UEdGraphSchema_K2::PN_Self && Pin->PinName != UEdGraphSchema_K2::PN_Execute && Pin->Direction == EGPD_Input)
-			{
-				SavedPins.Add(FMVVMBlueprintPin::CreateFromPin(Blueprint, Pin));
-			}
-		}
+		SavedPins = FMVVMBlueprintPin::CreateFromNode(Blueprint, CachedWrapperNode);
 	}
 }
 
@@ -231,8 +250,7 @@ UEdGraph* UMVVMBlueprintViewEvent::CreateWrapperGraphInternal(const UFunction* D
 	bool bIsConst = false;
 	bool bTransient = true;
 	TPair<UEdGraph*, UK2Node*> Result = UE::MVVM::ConversionFunctionHelper::CreateGraph(WidgetBlueprint, GraphName, DelegateSignature, Function, bIsConst, bTransient);
-	CachedWrapperGraph = Result.Get<0>();
-	CachedWrapperNode = Result.Get<1>();
+	SetCachedWrapperGraphInternal(Result.Get<0>(), Result.Get<1>());
 
 	// Add SelfNode if needed
 	if (CachedWrapperNode)
@@ -266,8 +284,7 @@ UEdGraph* UMVVMBlueprintViewEvent::CreateWrapperGraphInternal(const UFunction* D
 			UK2Node_VariableSet* VariableNode = CastChecked<UK2Node_VariableSet>(NewNode);
 			UEdGraphSchema_K2::ConfigureVarNode(VariableNode, Property->GetFName(), Property->GetOwnerStruct(), WidgetBlueprint);
 		});
-	CachedWrapperGraph = Result.Get<0>();
-	CachedWrapperNode = Result.Get<1>();
+	SetCachedWrapperGraphInternal(Result.Get<0>(), Result.Get<1>());
 
 	// Add SelfNode
 	if (CachedWrapperNode)
@@ -291,37 +308,8 @@ void UMVVMBlueprintViewEvent::LoadPinValuesInternal()
 {
 	if (CachedWrapperNode)
 	{
-		TArray<UEdGraphPin*> AllPins;
-		AllPins.Reserve(CachedWrapperNode->Pins.Num());
-		for (UEdGraphPin* Pin : CachedWrapperNode->Pins)
-		{
-			if (Pin->PinName != UEdGraphSchema_K2::PN_Self && Pin->PinName != UEdGraphSchema_K2::PN_Execute && Pin->Direction == EGPD_Input)
-			{
-				AllPins.Add(Pin);
-			}
-		}
-
-		UBlueprint* Blueprint = GetWidgetBlueprintInternal();
-		for (int32 Index = SavedPins.Num() - 1; Index >= 0; --Index)
-		{
-			const FMVVMBlueprintPin& Pin = SavedPins[Index];
-			if (UEdGraphPin* GraphPin = CachedWrapperNode->FindPin(Pin.GetName()))
-			{
-				AllPins.RemoveSingleSwap(GraphPin);
-				Pin.CopyTo(Blueprint, GraphPin);
-			}
-			else
-			{
-				// pin doesn't exist anymore
-				SavedPins.RemoveAt(Index);
-			}
-		}
-
-		// Create the reminding pin.
-		for (UEdGraphPin* Pin : AllPins)
-		{
-			SavedPins.Add(FMVVMBlueprintPin::CreateFromPin(Blueprint, Pin));
-		}
+		TArray<FMVVMBlueprintPin> MissingPins = FMVVMBlueprintPin::CopyAndReturnMissingPins(GetWidgetBlueprintInternal(), CachedWrapperNode, SavedPins);
+		SavedPins.Append(MissingPins);
 	}
 }
 
@@ -347,7 +335,7 @@ bool UMVVMBlueprintViewEvent::HasCompilationMessage(EMessageType InMessageType)
 	});
 }
 
-void UMVVMBlueprintViewEvent::AddCompilationToBinding(FMessage MessageToAdd)
+void UMVVMBlueprintViewEvent::AddCompilationToBinding(FMessage MessageToAdd) const
 {
 	Messages.Add(MoveTemp(MessageToAdd));
 }
@@ -399,6 +387,42 @@ FString UMVVMBlueprintViewEvent::GetSearchableString() const
 	return Builder.ToString();
 }
 
+void UMVVMBlueprintViewEvent::HandleGraphChanged(const FEdGraphEditAction& EditAction)
+{
+	if (EditAction.Graph == CachedWrapperGraph && CachedWrapperGraph)
+	{
+		if (CachedWrapperNode && EditAction.Nodes.Contains(CachedWrapperNode))
+		{
+			if (EditAction.Action == EEdGraphActionType::GRAPHACTION_RemoveNode)
+			{
+				CachedWrapperNode = UE::MVVM::ConversionFunctionHelper::GetWrapperNode(CachedWrapperGraph);
+				SavePinValues();
+				OnWrapperGraphModified.Broadcast();
+			}
+			else if (EditAction.Action == EEdGraphActionType::GRAPHACTION_EditNode)
+			{
+				SavePinValues();
+				OnWrapperGraphModified.Broadcast();
+			}
+		}
+		else if (CachedWrapperNode == nullptr && EditAction.Action == EEdGraphActionType::GRAPHACTION_AddNode)
+		{
+			CachedWrapperNode = UE::MVVM::ConversionFunctionHelper::GetWrapperNode(CachedWrapperGraph);
+			SavePinValues();
+			OnWrapperGraphModified.Broadcast();
+		}
+	}
+}
+
+void UMVVMBlueprintViewEvent::HandleUserDefinedPinRenamed(UK2Node* InNode, FName OldPinName, FName NewPinName)
+{
+	if (InNode == CachedWrapperNode)
+	{
+		SavePinValues();
+		OnWrapperGraphModified.Broadcast();
+	}
+}
+
 void UMVVMBlueprintViewEvent::PostEditChangeChainProperty(FPropertyChangedChainEvent& PropertyChainEvent)
 {
 	Super::PostEditChangeChainProperty(PropertyChainEvent);
