MVVM: Big rework of the MVVM view compiler. It add supports for conversion function to go from complex to simple. Supports pin split on simple conversion function. Add ability to to support setter fun

diff --git a/Engine/Plugins/Runtime/ModelViewViewModel/Source/ModelViewViewModelBlueprint/Private/MVVMViewBlueprintCompiler.cpp b/Engine/Plugins/Runtime/ModelViewViewModel/Source/ModelViewViewModelBlueprint/Private/MVVMViewBlueprintCompiler.cpp
--- a/Engine/Plugins/Runtime/ModelViewViewModel/Source/ModelViewViewModelBlueprint/Private/MVVMViewBlueprintCompiler.cpp
+++ b/Engine/Plugins/Runtime/ModelViewViewModel/Source/ModelViewViewModelBlueprint/Private/MVVMViewBlueprintCompiler.cpp
@@ -26,8 +26,8 @@
 
 namespace UE::MVVM::Private
 {
-FAutoConsoleVariable CVarLogViewCompliedResult(
-	TEXT("MVVM.LogViewCompliedResult"),
+FAutoConsoleVariable CVarLogViewCompiledResult(
+	TEXT("MVVM.LogViewCompiledResult"),
 	false,
 	TEXT("After the view is compiled log the compiled bindings and sources.")
 );
@@ -48,8 +48,19 @@ FAutoConsoleCommand CVarTestGenerateSetter(
 #endif
 
 static const FText CouldNotCreateSourceFieldPathFormat = LOCTEXT("CouldNotCreateSourceFieldPath", "Couldn't create the source field path '{0}'. {1}");
+static const FText CouldNotCreateDestinationFieldPathFormat = LOCTEXT("CouldNotCreateDestinationFieldPath", "Couldn't create the destination field path '{0}'. {1}");
 static const FText PropertyPathIsInvalidFormat = LOCTEXT("PropertyPathIsInvalid", "The property path '{0}' is invalid.");
 
+void RenameObjectToTransientPackage(UObject* ObjectToRename)
+{
+	const ERenameFlags RenFlags = REN_DoNotDirty | REN_ForceNoResetLoaders | REN_DontCreateRedirectors;
+
+	ObjectToRename->Rename(nullptr, GetTransientPackage(), RenFlags);
+	ObjectToRename->SetFlags(RF_Transient);
+	ObjectToRename->ClearFlags(RF_Public | RF_Standalone | RF_ArchetypeObject);
+	FLinkerLoad::InvalidateExport(ObjectToRename);
+}
+
 FString PropertyPathToString(const UClass* InSelfContext, const UMVVMBlueprintView* BlueprintView, const FMVVMBlueprintPropertyPath& PropertyPath)
 {
 	if (!PropertyPath.IsValid())
@@ -100,7 +111,7 @@ FText GetViewModelIdText(const FMVVMBlueprintPropertyPath& PropertyPath)
 	return FText::FromString(PropertyPath.GetViewModelId().ToString(EGuidFormats::DigitsWithHyphensInBraces));
 }
 
-TValueOrError<FCompiledBindingLibraryCompiler::FFieldPathHandle, FText> AddObjectFieldPath(FCompiledBindingLibraryCompiler& BindingLibraryCompiler, const UWidgetBlueprintGeneratedClass* Class, FStringView ObjectPath, UClass* ExpectedType)
+TValueOrError<FCompiledBindingLibraryCompiler::FFieldPathHandle, FText> AddObjectFieldPath(FCompiledBindingLibraryCompiler& BindingLibraryCompiler, const UWidgetBlueprintGeneratedClass* Class, FStringView ObjectPath, const UClass* ExpectedType)
 {
 	// Generate a path to read the value at runtime
 	static const FText InvalidGetterFormat = LOCTEXT("ViewModelInvalidGetterWithReason", "Viewmodel has an invalid Getter. {0}");
@@ -120,12 +131,66 @@ TValueOrError<FCompiledBindingLibraryCompiler::FFieldPathHandle, FText> AddObjec
 	return MakeValue(ReadFieldPathResult.StealValue());
 }
 
-FMVVMViewBlueprintCompiler::FMVVMViewBlueprintCompiler(FWidgetBlueprintCompilerContext& InCreationContext)
+
+FMVVMViewBlueprintCompiler::FMVVMViewBlueprintCompiler(FWidgetBlueprintCompilerContext& InCreationContext, UMVVMBlueprintView* InBlueprintView)
 	: WidgetBlueprintCompilerContext(InCreationContext)
+	, BlueprintView(InBlueprintView)
 	, BindingLibraryCompiler(InCreationContext.WidgetBlueprint())
-{}
+{
+	check(BlueprintView.IsValid());
+}
+
+
+void FMVVMViewBlueprintCompiler::AddMessage(const FText& MessageText, EMessageType MessageType) const
+{
+	switch (MessageType)
+	{
+	case EMessageType::Info:
+		WidgetBlueprintCompilerContext.MessageLog.Note(*MessageText.ToString());
+		break;
+	case EMessageType::Warning:
+		WidgetBlueprintCompilerContext.MessageLog.Warning(*MessageText.ToString());
+		break;
+	case EMessageType::Error:
+		WidgetBlueprintCompilerContext.MessageLog.Error(*MessageText.ToString());
+		break;
+	default:
+		break;
+	}
+}
+
+
+void FMVVMViewBlueprintCompiler::AddMessages(TArrayView<TWeakPtr<FCompilerBinding>> Bindings, TArrayView<TWeakPtr<FCompilerEvent>> Events, const FText& MessageText, EMessageType MessageType) const
+{
+	for (TWeakPtr<FCompilerBinding>& Binding : Bindings)
+	{
+		AddMessageForBinding(Binding.Pin(), MessageText, MessageType, FName());
+	}
+	for (TWeakPtr<FCompilerEvent>& Event : Events)
+	{
+		AddMessageForEvent(Event.Pin(), MessageText, MessageType, FName());
+	}
+
+	if (Bindings.Num() == 0 && Events.Num() == 0)
+	{
+		AddMessage(MessageText, MessageType);
+	}
+}
+
+void FMVVMViewBlueprintCompiler::AddMessageForBinding(const TSharedPtr<FCompilerBinding>& Binding, const FText& MessageText, EMessageType MessageType, FName ArgumentName) const
+{
+	const FMVVMBlueprintViewBinding* BindingPtr = Binding ? BlueprintView->GetBindingAt(Binding->Key.ViewBindingIndex) : nullptr;
+	if (BindingPtr)
+	{
+		AddMessageForBinding(*BindingPtr, MessageText, MessageType, ArgumentName);
+	}
+	else
+	{
+		AddMessage(MessageText, MessageType);
+	}
+}
 
-void FMVVMViewBlueprintCompiler::AddMessageForBinding(FMVVMBlueprintViewBinding& Binding, UMVVMBlueprintView* BlueprintView, const FText& MessageText, EBindingMessageType MessageType, FName ArgumentName) const
+void FMVVMViewBlueprintCompiler::AddMessageForBinding(const FMVVMBlueprintViewBinding& Binding, const FText& MessageText, EMessageType MessageType, FName ArgumentName) const
 {
 	const FText BindingName = FText::FromString(Binding.GetDisplayNameString(WidgetBlueprintCompilerContext.WidgetBlueprint()));
 
@@ -138,26 +203,29 @@ void FMVVMViewBlueprintCompiler::AddMessageForBinding(FMVVMBlueprintViewBinding&
 	{
 		FormattedError = FText::Format(LOCTEXT("BindingFormat", "Binding '{0}': {1}"), BindingName, MessageText);
 	}
+	AddMessage(FormattedError, MessageType);
 
-	switch (MessageType)
+	static EBindingMessageType BindingMessageTypes[] = { EBindingMessageType ::Info, EBindingMessageType ::Warning, EBindingMessageType ::Error};
+	FBindingMessage NewMessage = { FormattedError, BindingMessageTypes[static_cast<int32>(MessageType)] };
+	BlueprintView->AddMessageToBinding(Binding.BindingId, NewMessage);
+}
+
+
+void FMVVMViewBlueprintCompiler::AddMessageForEvent(const TSharedPtr<FCompilerEvent>& Event, const FText& MessageText, EMessageType MessageType, FName ArgumentName) const
+{
+	const UMVVMBlueprintViewEvent* EventPtr = Event ? Event->Event.Get() : nullptr;
+	if (EventPtr)
 	{
-	case EBindingMessageType::Info:
-		WidgetBlueprintCompilerContext.MessageLog.Note(*FormattedError.ToString());
-		break;
-	case EBindingMessageType::Warning:
-		WidgetBlueprintCompilerContext.MessageLog.Warning(*FormattedError.ToString());
-		break;
-	case EBindingMessageType::Error:
-		WidgetBlueprintCompilerContext.MessageLog.Error(*FormattedError.ToString());
-		break;
-	default:
-		break;
+		AddMessageForEvent(EventPtr, MessageText, MessageType, ArgumentName);
+	}
+	else
+	{
+		AddMessage(MessageText, MessageType);
 	}
-	FBindingMessage NewMessage = { FormattedError, MessageType };
-	BlueprintView->AddMessageToBinding(Binding.BindingId, NewMessage);
 }
 
-void FMVVMViewBlueprintCompiler::AddMessageForEvent(UMVVMBlueprintViewEvent* Event, const FText& MessageText, UMVVMBlueprintViewEvent::EMessageType MessageType, FName ArgumentName) const
+
+void FMVVMViewBlueprintCompiler::AddMessageForEvent(const UMVVMBlueprintViewEvent* Event, const FText& MessageText, EMessageType MessageType, FName ArgumentName) const
 {
 	const FText EventName = Event->GetDisplayName(true);
 
@@ -170,30 +238,25 @@ void FMVVMViewBlueprintCompiler::AddMessageForEvent(UMVVMBlueprintViewEvent* Eve
 	{
 		FormattedError = FText::Format(LOCTEXT("EventFormat", "Event '{0}': {1}"), EventName, MessageText);
 	}
+	AddMessage(FormattedError, MessageType);
 
-	switch (MessageType)
-	{
-	case UMVVMBlueprintViewEvent::EMessageType::Info:
-		WidgetBlueprintCompilerContext.MessageLog.Note(*FormattedError.ToString());
-		break;
-	case UMVVMBlueprintViewEvent::EMessageType::Warning:
-		WidgetBlueprintCompilerContext.MessageLog.Warning(*FormattedError.ToString());
-		break;
-	case UMVVMBlueprintViewEvent::EMessageType::Error:
-		WidgetBlueprintCompilerContext.MessageLog.Error(*FormattedError.ToString());
-		break;
-	default:
-		break;
-	}
-	UMVVMBlueprintViewEvent::FMessage NewMessage = { FormattedError, MessageType };
+	static UMVVMBlueprintViewEvent::EMessageType BindingMessageTypes[] = { UMVVMBlueprintViewEvent::EMessageType::Info, UMVVMBlueprintViewEvent::EMessageType::Warning, UMVVMBlueprintViewEvent::EMessageType::Error };
+	UMVVMBlueprintViewEvent::FMessage NewMessage = { FormattedError, BindingMessageTypes[static_cast<int32>(MessageType)] };
 	Event->AddCompilationToBinding(NewMessage);
 }
 
 
-void FMVVMViewBlueprintCompiler::AddErrorForViewModel(const FMVVMBlueprintViewModelContext& ViewModel, const FText& Message) const
+void FMVVMViewBlueprintCompiler::AddMessageForViewModel(const FMVVMBlueprintViewModelContext& ViewModel, const FText& Message, EMessageType MessageType) const
 {
 	const FText FormattedError = FText::Format(LOCTEXT("ViewModelFormat", "Viewodel '{0}': {1}"), ViewModel.GetDisplayName(), Message);
-	WidgetBlueprintCompilerContext.MessageLog.Error(*FormattedError.ToString());
+	AddMessage(FormattedError, MessageType);
+}
+
+
+void FMVVMViewBlueprintCompiler::AddMessageForViewModel(const FText& ViewModelDisplayName, const FText& Message, EMessageType MessageType) const
+{
+	const FText FormattedError = FText::Format(LOCTEXT("ViewModelFormat", "Viewodel '{0}': {1}"), ViewModelDisplayName, Message);
+	AddMessage(FormattedError, MessageType);
 }
 
 
@@ -208,16 +271,6 @@ void FMVVMViewBlueprintCompiler::CleanOldData(UWidgetBlueprintGeneratedClass* Cl
 	// Clean old View
 	if (!WidgetBlueprintCompilerContext.Blueprint->bIsRegeneratingOnLoad && WidgetBlueprintCompilerContext.bIsFullCompile)
 	{
-		auto RenameObjectToTransientPackage = [](UObject* ObjectToRename)
-		{
-			const ERenameFlags RenFlags = REN_DoNotDirty | REN_ForceNoResetLoaders | REN_DontCreateRedirectors;
-
-			ObjectToRename->Rename(nullptr, GetTransientPackage(), RenFlags);
-			ObjectToRename->SetFlags(RF_Transient);
-			ObjectToRename->ClearFlags(RF_Public | RF_Standalone | RF_ArchetypeObject);
-			FLinkerLoad::InvalidateExport(ObjectToRename);
-		};
-
 		TArray<UObject*> Children;
 		const bool bIncludeNestedObjects = false;
 		ForEachObjectWithOuter(ClassToClean, [&Children](UObject* Child)
@@ -236,93 +289,58 @@ void FMVVMViewBlueprintCompiler::CleanOldData(UWidgetBlueprintGeneratedClass* Cl
 }
 
 
-void FMVVMViewBlueprintCompiler::CreateFunctions(UMVVMBlueprintView* BlueprintView)
-{
-	if (!bAreSourcesCreatorValid || !bAreBindingsValid || !bAreEventsValid)
-	{
-		return;
-	}
-	
-	// Build the list of BP destinations
-	CreateBindingDestinationContexts(BlueprintView);
-
-	if (!bAreSourcesCreatorValid || !bAreBindingsValid || !bAreEventsValid)
-	{
-		return;
-	}
-
-	// Generate the setter code
-	if (GetDefault<UMVVMDeveloperProjectSettings>()->bAllowGeneratedViewModelSetter)
-	{
-		for (const FCompilerSourceCreatorContext& SourceCreator : CompilerSourceCreatorContexts)
-		{
-			if (SourceCreator.SetterGraph)
-			{
-				if (!UE::MVVM::FunctionGraphHelper::GenerateViewModelSetter(WidgetBlueprintCompilerContext, SourceCreator.SetterGraph, SourceCreator.ViewModelContext.GetViewModelName()))
-				{
-					AddErrorForViewModel(SourceCreator.ViewModelContext, LOCTEXT("SetterFunctionCouldNotBeGenerated", "The setter function could not be generated."));
-					continue;
-				}
-			}
-		}
-	}
-}
-
-
-void FMVVMViewBlueprintCompiler::CreateVariables(const FWidgetBlueprintCompilerContext::FCreateVariableContext& Context, UMVVMBlueprintView* BlueprintView)
+void FMVVMViewBlueprintCompiler::CreateVariables(const FWidgetBlueprintCompilerContext::FCreateVariableContext& Context)
 {
 	if (!BlueprintView)
 	{
 		return;
 	}
 
-	if (!bAreSourcesCreatorValid || !bAreSourceContextsValid || !bAreBindingsValid || !bAreEventsValid)
+	if (!AreStepsValid())
 	{
 		return;
 	}
 
 	if (Context.GetCompileType() == EKismetCompileType::SkeletonOnly)
 	{
-		CreateWidgetMap(Context, BlueprintView);
-		CreateSourceLists(Context, BlueprintView);
-		CreateFunctionsDeclaration(Context, BlueprintView);
-	}
-	else
-	{
-		CreateIntermediateGraphFunctions(Context, BlueprintView);
+		CreateWidgetMap(Context);
+		CreateBindingList(Context);
+		CreateEventList(Context);
+		CreateRequiredProperties(Context);
+		CreatePublicFunctionsDeclaration(Context);
 	}
 
-	auto CreateVariable = [&Context](const FCompilerUserWidgetPropertyContext& SourceContext) -> FProperty*
+	auto CreateVariable = [&Context](const FCompilerUserWidgetProperty& UserWidgetProperty) -> FProperty*
 	{
-		FEdGraphPinType NewPropertyPinType(UEdGraphSchema_K2::PC_Object, NAME_None, SourceContext.Class, EPinContainerType::None, false, FEdGraphTerminalType());
-		FProperty* NewProperty = Context.CreateVariable(SourceContext.PropertyName, NewPropertyPinType);
+		FEdGraphPinType NewPropertyPinType(UEdGraphSchema_K2::PC_Object, NAME_None, UserWidgetProperty.AuthoritativeClass, EPinContainerType::None, false, FEdGraphTerminalType());
+		FProperty* NewProperty = Context.CreateVariable(UserWidgetProperty.Name, NewPropertyPinType);
 		if (NewProperty != nullptr)
 		{
 			NewProperty->SetPropertyFlags(CPF_BlueprintVisible | CPF_RepSkip | CPF_Transient | CPF_DuplicateTransient);
-			if (SourceContext.BlueprintSetter.IsEmpty())
+			if (UserWidgetProperty.bReadOnly)
 			{
 				NewProperty->SetPropertyFlags(CPF_BlueprintReadOnly);
 			}
-			NewProperty->SetPropertyFlags(SourceContext.bExposeOnSpawn ? CPF_ExposeOnSpawn : CPF_DisableEditOnInstance);
+			NewProperty->SetPropertyFlags(UserWidgetProperty.bExposeOnSpawn ? CPF_ExposeOnSpawn : CPF_DisableEditOnInstance);
 
 #if WITH_EDITOR
-			if (!SourceContext.BlueprintSetter.IsEmpty())
+			if (!UserWidgetProperty.BlueprintSetter.IsEmpty())
 			{
-				NewProperty->SetMetaData(FBlueprintMetadata::MD_PropertySetFunction, *SourceContext.BlueprintSetter);
+				NewProperty->SetMetaData(FBlueprintMetadata::MD_PropertySetFunction, *UserWidgetProperty.BlueprintSetter);
 			}
-			if (!SourceContext.DisplayName.IsEmpty())
+			if (!UserWidgetProperty.DisplayName.IsEmpty())
 			{
-				NewProperty->SetMetaData(FBlueprintMetadata::MD_DisplayName, *SourceContext.DisplayName.ToString());
+				NewProperty->SetMetaData(FBlueprintMetadata::MD_DisplayName, *UserWidgetProperty.DisplayName.ToString());
 			}
-			if (!SourceContext.CategoryName.IsEmpty())
+			if (!UserWidgetProperty.CategoryName.IsEmpty())
 			{
-				NewProperty->SetMetaData(FBlueprintMetadata::MD_FunctionCategory, *SourceContext.CategoryName);
+				NewProperty->SetMetaData(FBlueprintMetadata::MD_FunctionCategory, *UserWidgetProperty.CategoryName);
 			}
-			if (SourceContext.bExposeOnSpawn)
+			if (UserWidgetProperty.bExposeOnSpawn)
 			{
 				NewProperty->SetMetaData(FBlueprintMetadata::MD_ExposeOnSpawn, TEXT("true"));
 			}
-			if (SourceContext.bPrivate)
+			if (UserWidgetProperty.bPrivate)
 			{
 				NewProperty->SetMetaData(FBlueprintMetadata::MD_Private, TEXT("true"));
 			}
@@ -331,82 +349,89 @@ void FMVVMViewBlueprintCompiler::CreateVariables(const FWidgetBlueprintCompilerC
 		return NewProperty;
 	};
 
-	for (FCompilerUserWidgetPropertyContext& SourceContext : CompilerUserWidgetPropertyContexts)
+	for (FCompilerUserWidgetProperty& UserWidgetProperty : NeededUserWidgetProperties)
 	{
-		SourceContext.Field = BindingHelper::FindFieldByName(Context.GetGeneratedClass(), FMVVMBindingName(SourceContext.PropertyName));
+		check(UserWidgetProperty.AuthoritativeClass);
+		check(!UserWidgetProperty.Name.IsNone());
+		UserWidgetProperty.Property = nullptr; // Skeletal set the property, Full needs the new property
+
+		FMVVMConstFieldVariant UserWidgetPropertyField = BindingHelper::FindFieldByName(Context.GetGeneratedClass(), FMVVMBindingName(UserWidgetProperty.Name));
 
 		// The class is not linked yet. It may not be available yet.
-		if (SourceContext.Field.IsEmpty())
+		if (UserWidgetPropertyField.IsEmpty())
 		{
 			for (FField* Field = Context.GetGeneratedClass()->ChildProperties; Field != nullptr; Field = Field->Next)
 			{
-				if (Field->GetFName() == SourceContext.PropertyName)
+				if (Field->GetFName() == UserWidgetProperty.Name)
 				{
 					if (CastField<FProperty>(Field))
 					{
-						SourceContext.Field = FMVVMFieldVariant(CastField<FProperty>(Field));
+						UserWidgetPropertyField = FMVVMFieldVariant(CastField<FProperty>(Field));
 					}
 					else
 					{
-						WidgetBlueprintCompilerContext.MessageLog.Error(*FText::Format(LOCTEXT("FieldIsNotProperty", "The field for source '{0}' exists but is not a property."), SourceContext.DisplayName).ToString());
-						bAreSourcesCreatorValid = false;
+						WidgetBlueprintCompilerContext.MessageLog.Error(*FText::Format(LOCTEXT("FieldIsNotProperty", "The field for source '{0}' exists but is not a property."), UserWidgetProperty.DisplayName).ToString());
+						bIsCreateVariableStepValid = false;
 					}
 					break;
 				}
 			}
 			for (UField* Field = Context.GetGeneratedClass()->Children; Field != nullptr; Field = Field->Next)
 			{
-				if (Field->GetFName() == SourceContext.PropertyName)
+				if (Field->GetFName() == UserWidgetProperty.Name)
 				{
-					WidgetBlueprintCompilerContext.MessageLog.Error(*FText::Format(LOCTEXT("FieldIsNotProperty", "The field for source '{0}' exists but is not a property."), SourceContext.DisplayName).ToString());
-					bAreSourcesCreatorValid = false;
+					WidgetBlueprintCompilerContext.MessageLog.Error(*FText::Format(LOCTEXT("FieldIsNotProperty", "The field for source '{0}' exists but is not a property."), UserWidgetProperty.DisplayName).ToString());
+					bIsCreateVariableStepValid = false;
 					break;
 				}
 			}
 		}
 
-		if (SourceContext.Field.IsEmpty())
+		if (UserWidgetPropertyField.IsEmpty())
 		{
 			UClass* ParentClass = Context.GetGeneratedClass()->GetSuperClass();
-			if (const FProperty* Property = ParentClass->FindPropertyByName(SourceContext.PropertyName))
+			if (const FProperty* Property = ParentClass->FindPropertyByName(UserWidgetProperty.Name))
 			{
-				SourceContext.Field = FMVVMFieldVariant(Property);
+				UserWidgetPropertyField = FMVVMFieldVariant(Property);
 			}
 		}
 
-
 		// Will always create viewmodel properties.
 		// Will never create properties for animation or other Self.Object
 		// Will create properties for widget when they are not already created.
-
-		const bool bIsWidgetInsideWidgetTree = WidgetNameToWidgetPointerMap.Find(SourceContext.PropertyName) != nullptr;
-		const bool bIsViewmodel = SourceContext.ViewModelId.IsValid();
-
-		if (bIsViewmodel && !SourceContext.Field.IsEmpty())
+		if (UserWidgetProperty.CreationType == FCompilerUserWidgetProperty::ECreationType::CreateOnlyIfDoesntExist && !UserWidgetPropertyField.IsEmpty())
 		{
 			// Viewmodel property cannot already exist. It will creates issue with initialization and with View::SetViewModel.
-			const UClass* OwnerClass = Cast<UClass>(SourceContext.Field.GetOwner());
-			WidgetBlueprintCompilerContext.MessageLog.Error(*FText::Format(LOCTEXT("ViewmodelPropertyAlreadyExistInParent", "There is already a property named '{0}' in scope '{1}' for the viewmodel."), SourceContext.DisplayName, (OwnerClass ? OwnerClass->GetDisplayNameText() : FText::GetEmpty())).ToString());
-			bAreSourceContextsValid = false;
+			const UClass* OwnerClass = Cast<UClass>(UserWidgetPropertyField.GetOwner());
+			WidgetBlueprintCompilerContext.MessageLog.Error(*FText::Format(LOCTEXT("PropertyAlreadyExistInParent", "There is already a property named '{0}' in scope '{1}'."), UserWidgetProperty.DisplayName, (OwnerClass ? OwnerClass->GetDisplayNameText() : FText::GetEmpty())).ToString());
+			bIsCreateVariableStepValid = false;
 			continue;
 		}
 
-		if (!SourceContext.Field.IsEmpty())
+		if (!UserWidgetPropertyField.IsEmpty())
 		{
-			if (!BindingHelper::IsValidForSourceBinding(SourceContext.Field))
+			if (UserWidgetPropertyField.IsFunction())
+			{
+				WidgetBlueprintCompilerContext.MessageLog.Error(*FText::Format(LOCTEXT("FunctionCanBeSource", "Function can't be source. '{0}'."), UserWidgetProperty.DisplayName).ToString());
+				bIsCreateVariableStepValid = false;
+				continue;
+			}
+
+			if (!BindingHelper::IsValidForSourceBinding(UserWidgetPropertyField))
 			{
-				WidgetBlueprintCompilerContext.MessageLog.Error(*FText::Format(LOCTEXT("FieldNotAccessibleAtRuntime", "The field for source '{0}' exists but is not accessible at runtime."), SourceContext.DisplayName).ToString());
-				bAreSourcesCreatorValid = false;
+				WidgetBlueprintCompilerContext.MessageLog.Error(*FText::Format(LOCTEXT("FieldNotAccessibleAtRuntime", "The field for source '{0}' exists but is not accessible at runtime."), UserWidgetProperty.DisplayName).ToString());
+				bIsCreateVariableStepValid = false;
 				continue;
 			}
 
-			const FProperty* Property = SourceContext.Field.IsProperty() ? SourceContext.Field.GetProperty() : BindingHelper::GetReturnProperty(SourceContext.Field.GetFunction());
+			ensure(UserWidgetPropertyField.IsProperty());
+			const FProperty* Property = UserWidgetPropertyField.IsProperty() ? UserWidgetPropertyField.GetProperty() : nullptr;
 			const FObjectProperty* ObjectProperty = CastField<const FObjectProperty>(Property);
-			const bool bIsCompatible = ObjectProperty && SourceContext.Class->IsChildOf(ObjectProperty->PropertyClass);
+			const bool bIsCompatible = ObjectProperty && UserWidgetProperty.AuthoritativeClass->IsChildOf(ObjectProperty->PropertyClass);
 			if (!bIsCompatible)
 			{
-				WidgetBlueprintCompilerContext.MessageLog.Error(*FText::Format(LOCTEXT("PropertyExistsAndNotCompatible", "There is already a property named '{0}' that is not compatible with the source of the same name."), SourceContext.DisplayName).ToString());
-				bAreSourceContextsValid = false;
+				WidgetBlueprintCompilerContext.MessageLog.Error(*FText::Format(LOCTEXT("PropertyExistsAndNotCompatible", "There is already a property named '{0}' that is not compatible with the source of the same name."), UserWidgetProperty.DisplayName).ToString());
+				bIsCreateVariableStepValid = false;
 				continue;
 			}
 
@@ -415,30 +440,35 @@ void FMVVMViewBlueprintCompiler::CreateVariables(const FWidgetBlueprintCompilerC
 			{
 				// Widget needs to be BindWidget to be reused as a property.
 				const UClass* OwnerClass = Cast<UClass>(ObjectProperty->GetOwnerStruct());
-				WidgetBlueprintCompilerContext.MessageLog.Error(*FText::Format(LOCTEXT("WidgetPropertyAlreadyExist", "There is already a property named '{0}' in scope '{1}' for the widget. Are you missing a BindWidget?."), SourceContext.DisplayName, (OwnerClass ? OwnerClass->GetDisplayNameText() : FText::GetEmpty())).ToString());
-				bAreSourceContextsValid = false;
+				WidgetBlueprintCompilerContext.MessageLog.Error(*FText::Format(LOCTEXT("WidgetPropertyAlreadyExist", "There is already a property named '{0}' in scope '{1}' for the widget. Are you missing a BindWidget?."), UserWidgetProperty.DisplayName, (OwnerClass ? OwnerClass->GetDisplayNameText() : FText::GetEmpty())).ToString());
+				bIsCreateVariableStepValid = false;
 				continue;
 			}
 		}
 
 		// Can we reused the property or we need to create a new one.
-		bool bCreateVariable = bIsViewmodel || (bIsWidgetInsideWidgetTree && SourceContext.Field.IsEmpty());
+		bool bCreateVariable = UserWidgetProperty.CreationType == FCompilerUserWidgetProperty::ECreationType::CreateOnlyIfDoesntExist
+			|| (UserWidgetProperty.CreationType == FCompilerUserWidgetProperty::ECreationType::CreateIfDoesntExist && UserWidgetPropertyField.IsEmpty());
 		if (bCreateVariable)
 		{
-			SourceContext.Field = FMVVMConstFieldVariant(CreateVariable(SourceContext));
+			UserWidgetProperty.Property = CreateVariable(UserWidgetProperty);
+		}
+		else if (UserWidgetPropertyField.IsProperty())
+		{
+			UserWidgetProperty.Property = UserWidgetPropertyField.GetProperty();
 		}
 
-		if (SourceContext.Field.IsEmpty())
+		if (UserWidgetProperty.Property == nullptr)
 		{
-			WidgetBlueprintCompilerContext.MessageLog.Error(*FText::Format(LOCTEXT("VariableCouldNotBeCreated", "The variable for '{0}' could not be created."), SourceContext.DisplayName).ToString());
-			bAreSourceContextsValid = false;
+			WidgetBlueprintCompilerContext.MessageLog.Error(*FText::Format(LOCTEXT("VariableCouldNotBeCreated", "The variable for '{0}' could not be created."), UserWidgetProperty.DisplayName).ToString());
+			bIsCreateVariableStepValid = false;
 			continue;
 		}
 	}
 }
 
 
-void FMVVMViewBlueprintCompiler::CreateWidgetMap(const FWidgetBlueprintCompilerContext::FCreateVariableContext& Context, UMVVMBlueprintView* BlueprintView)
+void FMVVMViewBlueprintCompiler::CreateWidgetMap(const FWidgetBlueprintCompilerContext::FCreateVariableContext& Context)
 {
 	// The widget tree is not created yet for SKEL class.
 	//Context.GetGeneratedClass()->GetWidgetTreeArchetype()
@@ -463,26 +493,121 @@ void FMVVMViewBlueprintCompiler::CreateWidgetMap(const FWidgetBlueprintCompilerC
 }
 
 
-void FMVVMViewBlueprintCompiler::CreateSourceLists(const FWidgetBlueprintCompilerContext::FCreateVariableContext& Context, UMVVMBlueprintView* BlueprintView)
+void FMVVMViewBlueprintCompiler::CreateBindingList(const FWidgetBlueprintCompilerContext::FCreateVariableContext& Context)
+{
+	ensure(Context.GetCompileType() == EKismetCompileType::SkeletonOnly);
+	if (Context.GetCompileType() != EKismetCompileType::SkeletonOnly)
+	{
+		return;
+	}
+
+	ValidBindings.Reset();
+
+	// Build the list of bindings that we should compile.
+	for (int32 Index = 0; Index < BlueprintView->GetNumBindings(); ++Index)
+	{
+		FMVVMBlueprintViewBinding* BindingPtr = BlueprintView->GetBindingAt(Index);
+		if (BindingPtr == nullptr)
+		{
+			WidgetBlueprintCompilerContext.MessageLog.Error(*FText::Format(LOCTEXT("BindingInvalidIndex", "Internal error: Tried to fetch binding for invalid index {0}."), Index).ToString());
+			bIsCreateVariableStepValid = false;
+			continue;
+		}
+
+		FMVVMBlueprintViewBinding& Binding = *BindingPtr;
+		if (!Binding.bCompile)
+		{
+			continue;
+		}
+
+		const bool bHasConversionFunction = Binding.Conversion.GetConversionFunction(true) || Binding.Conversion.GetConversionFunction(false);
+		if (bHasConversionFunction && Binding.BindingType == EMVVMBindingMode::TwoWay)
+		{
+			AddMessageForBinding(Binding, LOCTEXT("TwoWayBindingsWithConversion", "Two-way bindings are not allowed to use conversion functions."), EMessageType::Error, FName());
+			bIsCreateVariableStepValid = false;
+			continue;
+		}
+
+		if (IsForwardBinding(Binding.BindingType))
+		{
+			TSharedRef<FCompilerBinding> ValidBinding = MakeShared<FCompilerBinding>();
+			ValidBinding->Key.ViewBindingIndex = Index;
+			ValidBinding->Key.bIsForwardBinding = true;
+			ValidBinding->bIsOneTimeBinding = IsOneTimeBinding(Binding.BindingType);
+
+			ValidBindings.Add(ValidBinding);
+		}
+		if (IsBackwardBinding(Binding.BindingType))
+		{
+			TSharedRef<FCompilerBinding> ValidBinding = MakeShared<FCompilerBinding>();
+			ValidBinding->Key.ViewBindingIndex = Index;
+			ValidBinding->Key.bIsForwardBinding = false;
+			ValidBinding->bIsOneTimeBinding = IsOneTimeBinding(Binding.BindingType);
+
+			ValidBindings.Add(ValidBinding);
+		}
+	}
+}
+
+
+void FMVVMViewBlueprintCompiler::CreateEventList(const FWidgetBlueprintCompilerContext::FCreateVariableContext& Context)
+{
+	for (UMVVMBlueprintViewEvent* EventPtr : BlueprintView->GetEvents())
+	{
+		if (EventPtr == nullptr)
+		{
+			WidgetBlueprintCompilerContext.MessageLog.Error(*LOCTEXT("EventInvalid", "Internal error: An event is invalid.").ToString());
+			bIsCreateVariableStepValid = false;
+			continue;
+		}
+
+		if (!EventPtr->bCompile)
+		{
+			continue;
+		}
+
+		if (!EventPtr->GetEventPath().HasPaths())
+		{
+			AddMessageForEvent(EventPtr, LOCTEXT("EventInvalidEventPath", "The event path is invalid."), EMessageType::Error, FName());
+			bIsCreateVariableStepValid = false;
+			continue;
+		}
+
+		TSharedRef<FCompilerEvent> ValidEvent = MakeShared<FCompilerEvent>();
+		ValidEvent->Event = EventPtr;
+
+		ValidEvents.Add(ValidEvent);
+	}
+}
+
+
+void FMVVMViewBlueprintCompiler::CreateRequiredProperties(const FWidgetBlueprintCompilerContext::FCreateVariableContext& Context)
 {
-	CompilerUserWidgetPropertyContexts.Reset();
-	CompilerSourceCreatorContexts.Reset();
+	ensure(Context.GetCompileType() == EKismetCompileType::SkeletonOnly);
+	if (Context.GetCompileType() != EKismetCompileType::SkeletonOnly)
+	{
+		return;
+	}
+
+	NeededBindingSources.Reset();
+	NeededUserWidgetProperties.Reset();
+	ViewModelCreatorContexts.Reset();
+	ViewModelSettersToGenerate.Reset();
 
 	TSet<FGuid> ViewModelGuids;
-	TSet<FName> WidgetSources;
 	for (const FMVVMBlueprintViewModelContext& ViewModelContext : BlueprintView->GetViewModels())
 	{
 		if (!ViewModelContext.GetViewModelId().IsValid())
 		{
-			AddErrorForViewModel(ViewModelContext, LOCTEXT("ViewmodelInvalidGuid", "GUID is invalid."));
-			bAreSourcesCreatorValid = false;
+			AddMessageForViewModel(ViewModelContext, LOCTEXT("ViewmodelInvalidGuid", "GUID is invalid."), EMessageType::Error);
+			bIsCreateVariableStepValid = false;
 			continue;
 		}
 
 		if (ViewModelGuids.Contains(ViewModelContext.GetViewModelId()))
 		{
-			AddErrorForViewModel(ViewModelContext, LOCTEXT("ViewmodelAlreadyAdded", "Identical viewmodel has already been added."));
-			bAreSourcesCreatorValid = false;
+			AddMessageForViewModel(ViewModelContext, LOCTEXT("ViewmodelAlreadyAdded", "Identical viewmodel has already been added."), EMessageType::Error);
+			bIsCreateVariableStepValid = false;
 			continue;
 		}
 
@@ -490,86 +615,151 @@ void FMVVMViewBlueprintCompiler::CreateSourceLists(const FWidgetBlueprintCompile
 
 		if (ViewModelContext.GetViewModelClass() == nullptr || !ViewModelContext.IsValid())
 		{
-			AddErrorForViewModel(ViewModelContext, LOCTEXT("ViewmodelInvalidClass", "Invalid class."));
-			bAreSourcesCreatorValid = false;
+			AddMessageForViewModel(ViewModelContext, LOCTEXT("ViewmodelInvalidClass", "Invalid class."), EMessageType::Error);
+			bIsCreateVariableStepValid = false;
 			continue;
 		}
 
+		FName PropertyName = ViewModelContext.GetViewModelName();
+
 		const bool bCreateSetterFunction = GetDefault<UMVVMDeveloperProjectSettings>()->bAllowGeneratedViewModelSetter
 			&& (ViewModelContext.bCreateSetterFunction || ViewModelContext.CreationType == EMVVMBlueprintViewModelContextCreationType::Manual);
+		FString SetterFunctionName;
+		if (bCreateSetterFunction)
+		{
+			SetterFunctionName = TEXT("Set") + ViewModelContext.GetViewModelName().ToString();
 
-		int32 FoundSourceCreatorContextIndex = INDEX_NONE;
-		if (Context.GetCompileType() == EKismetCompileType::SkeletonOnly)
+			FCompilerViewModelSetter& ViewModelSetter = ViewModelSettersToGenerate.AddDefaulted_GetRef();
+			ViewModelSetter.Class = ViewModelContext.GetViewModelClass();
+			ViewModelSetter.PropertyName = PropertyName;
+			ViewModelSetter.BlueprintSetter = SetterFunctionName;
+			ViewModelSetter.DisplayName = ViewModelContext.GetDisplayName();
+		}
+
+		TSharedRef<FCompilerBindingSource> SourceContext = MakeShared<FCompilerBindingSource>();
 		{
-			FCompilerSourceCreatorContext SourceContext;
-			SourceContext.ViewModelContext = ViewModelContext;
-			SourceContext.Type = ECompilerSourceCreatorType::ViewModel;
-			if (bCreateSetterFunction)
+			SourceContext->AuthoritativeClass = ViewModelContext.GetViewModelClass();
+			SourceContext->Name = PropertyName;
+			SourceContext->Type = FCompilerBindingSource::EType::ViewModel;
+			SourceContext->bIsOptional = ViewModelContext.bOptional;
+
+			if (ViewModelContext.CreationType == EMVVMBlueprintViewModelContextCreationType::Manual)
 			{
-				SourceContext.SetterFunctionName = TEXT("Set") + ViewModelContext.GetViewModelName().ToString();
+				SourceContext->bIsOptional = true;
 			}
-			FoundSourceCreatorContextIndex = CompilerSourceCreatorContexts.Emplace(MoveTemp(SourceContext));
+			else if (ViewModelContext.CreationType == EMVVMBlueprintViewModelContextCreationType::CreateInstance)
+			{
+				SourceContext->bIsOptional = false;
+			}
+
+			NeededBindingSources.Add(SourceContext);
 		}
-		else
+
 		{
-			FGuid ViewModelId = ViewModelContext.GetViewModelId();
-			FoundSourceCreatorContextIndex = CompilerSourceCreatorContexts.IndexOfByPredicate([ViewModelId](const FCompilerSourceCreatorContext& Other)
-				{
-					return Other.ViewModelContext.GetViewModelId() == ViewModelId;
-				});
+			FCompilerUserWidgetProperty& SourceVariable = NeededUserWidgetProperties.AddDefaulted_GetRef();
+			SourceVariable.AuthoritativeClass = ViewModelContext.GetViewModelClass();
+			SourceVariable.Name = PropertyName;
+			SourceVariable.DisplayName = ViewModelContext.GetDisplayName();
+			SourceVariable.CategoryName = TEXT("Viewmodel");
+			SourceVariable.BlueprintSetter = SetterFunctionName;
+			SourceVariable.CreationType = FCompilerUserWidgetProperty::ECreationType::CreateOnlyIfDoesntExist;
+			SourceVariable.bExposeOnSpawn = bCreateSetterFunction;
+			SourceVariable.bPrivate = ViewModelContext.InstancedViewModel != nullptr;
+			SourceVariable.bReadOnly = !bCreateSetterFunction;
 		}
-		checkf(FoundSourceCreatorContextIndex != INDEX_NONE, TEXT("The viewmodel was added after the skeleton was created?"));
 
-		FCompilerUserWidgetPropertyContext SourceVariable;
-		SourceVariable.Class = ViewModelContext.GetViewModelClass();
-		SourceVariable.PropertyName = ViewModelContext.GetViewModelName();
-		SourceVariable.DisplayName = ViewModelContext.GetDisplayName();
-		SourceVariable.CategoryName = TEXT("Viewmodel");
-		SourceVariable.bExposeOnSpawn = bCreateSetterFunction;
-		SourceVariable.bPrivate = ViewModelContext.InstancedViewModel != nullptr;
-		SourceVariable.BlueprintSetter = CompilerSourceCreatorContexts[FoundSourceCreatorContextIndex].SetterFunctionName;
-		SourceVariable.ViewModelId = ViewModelContext.GetViewModelId();
-		CompilerUserWidgetPropertyContexts.Emplace(MoveTemp(SourceVariable));
+		{
+			FCompilerViewModelCreatorContext& CreatorContext = ViewModelCreatorContexts.AddDefaulted_GetRef();
+			CreatorContext.ViewModelContext = ViewModelContext;
+			CreatorContext.Source = SourceContext;
+		}
 	}
 
-	bAreSourceContextsValid = bAreSourcesCreatorValid;
-
+	TSet<FName> WidgetSourcesCreated;
+	TSet<FName> WidgetUserPropertyCreated;
+	bool bSelfBindingSourceCreated = NeededBindingSources.ContainsByPredicate([](const TSharedRef<FCompilerBindingSource>& Other)
+		{
+			return Other->Type == FCompilerBindingSource::EType::Self;
+		});
 	const FName DefaultWidgetCategory = Context.GetWidgetBlueprint()->GetFName();
-
-	auto GenerateCompilerSourceContext = [Self = this, BlueprintView, DefaultWidgetCategory, Class = Context.GetGeneratedClass(), &ViewModelGuids, &WidgetSources](const FMVVMBlueprintPropertyPath& PropertyPath) -> TValueOrError<void, FText>
+	auto GenerateCompilerContext = [Self = this, DefaultWidgetCategory, Class = Context.GetGeneratedClass(), &ViewModelGuids, &WidgetSourcesCreated, &WidgetUserPropertyCreated, &bSelfBindingSourceCreated](bool bInCreateSource, const FMVVMBlueprintPropertyPath& PropertyPath) -> TValueOrError<void, FText>
 	{
 		switch (PropertyPath.GetSource(Self->WidgetBlueprintCompilerContext.WidgetBlueprint()))
 		{
 		case EMVVMBlueprintFieldPathSource::SelfContext:
+		{
+			if (!bSelfBindingSourceCreated)
+			{
+				bSelfBindingSourceCreated = true;
+
+				TSharedRef<FCompilerBindingSource> SourceContext = MakeShared<FCompilerBindingSource>();
+				SourceContext->AuthoritativeClass = nullptr;
+				SourceContext->Name = Self->WidgetBlueprintCompilerContext.WidgetBlueprint()->GetFName();
+				SourceContext->Type = FCompilerBindingSource::EType::Self;
+				SourceContext->bIsOptional = false;
+				Self->NeededBindingSources.Add(SourceContext);
+			}
 			return MakeValue();
+		}
 		case EMVVMBlueprintFieldPathSource::Widget:
 		{
-			// If the widget doesn't have a property, add one automatically.
-			if (!WidgetSources.Contains(PropertyPath.GetWidgetName()))
+			// Only do this once
+			bool bNewCreateSource = !WidgetSourcesCreated.Contains(PropertyPath.GetWidgetName()) && bInCreateSource;
+			bool bNewAddWidgetProperty = !WidgetUserPropertyCreated.Contains(PropertyPath.GetWidgetName());
+			if (bNewCreateSource || bNewAddWidgetProperty)
 			{
-				WidgetSources.Add(PropertyPath.GetWidgetName());
-
 				UWidget** WidgetPtr = Self->WidgetNameToWidgetPointerMap.Find(PropertyPath.GetWidgetName());
 				if (WidgetPtr == nullptr || *WidgetPtr == nullptr)
 				{
 					return MakeError(FText::Format(LOCTEXT("InvalidWidgetFormat", "Could not find the targeted widget: {0}"), FText::FromName(PropertyPath.GetWidgetName())));
 				}
-
 				UWidget* Widget = *WidgetPtr;
-				FCompilerUserWidgetPropertyContext SourceVariable;
-				SourceVariable.Class = Widget->GetClass();
-				SourceVariable.PropertyName = PropertyPath.GetWidgetName();
-				SourceVariable.DisplayName = FText::FromString(Widget->GetDisplayLabel());
-				SourceVariable.CategoryName = TEXT("Widget");
-				SourceVariable.ViewModelId = FGuid();
-				SourceVariable.bPrivate = false;
-				Self->CompilerUserWidgetPropertyContexts.Emplace(MoveTemp(SourceVariable));
+
+				if (bNewCreateSource)
+				{
+					FName PropertyName = PropertyPath.GetWidgetName();
+					{
+						TSharedRef<FCompilerBindingSource>* FoundCompilerSource = Self->NeededBindingSources.FindByPredicate([PropertyName](const TSharedRef<FCompilerBindingSource>& Other)
+							{
+								return Other->Name == PropertyName;
+							});
+						if (FoundCompilerSource != nullptr)
+						{
+							// It should be in the TSet<FName> WidgetSources
+							return MakeError(FText::Format(LOCTEXT("ExistingWidgetSourceFormat", "Internal error. A widget source already exist: {0}"), FText::FromName(PropertyPath.GetWidgetName())));
+						}
+					}
+
+					TSharedRef<FCompilerBindingSource> SourceContext = MakeShared<FCompilerBindingSource>();
+					SourceContext->AuthoritativeClass = Widget->GetClass();
+					SourceContext->Name = PropertyName;
+					SourceContext->Type = FCompilerBindingSource::EType::Widget;
+					SourceContext->bIsOptional = false;
+					Self->NeededBindingSources.Add(SourceContext);
+
+					WidgetSourcesCreated.Add(Widget->GetFName());
+				}
+
+				if (bNewAddWidgetProperty)
+				{
+					FCompilerUserWidgetProperty& SourceVariable = Self->NeededUserWidgetProperties.AddDefaulted_GetRef();
+					SourceVariable.AuthoritativeClass = Widget->GetClass();
+					SourceVariable.Name = PropertyPath.GetWidgetName();
+					SourceVariable.DisplayName = FText::FromString(Widget->GetDisplayLabel());
+					SourceVariable.CategoryName = TEXT("Widget");
+					SourceVariable.CreationType = FCompilerUserWidgetProperty::ECreationType::CreateIfDoesntExist;
+					SourceVariable.bExposeOnSpawn = false;
+					SourceVariable.bPrivate = true;
+					SourceVariable.bReadOnly = true;
+
+					WidgetUserPropertyCreated.Add(Widget->GetFName());
+				}
 			}
 			break;
 		}
 		case EMVVMBlueprintFieldPathSource::ViewModel:
 		{
-			const FMVVMBlueprintViewModelContext* SourceViewModelContext = BlueprintView->FindViewModel(PropertyPath.GetViewModelId());
+			const FMVVMBlueprintViewModelContext* SourceViewModelContext = Self->BlueprintView->FindViewModel(PropertyPath.GetViewModelId());
 			if (SourceViewModelContext == nullptr)
 			{
 				return MakeError(FText::Format(LOCTEXT("BindingViewModelNotFound", "Could not find viewmodel with GUID {0}."), GetViewModelIdText(PropertyPath)));
@@ -587,142 +777,95 @@ void FMVVMViewBlueprintCompiler::CreateSourceLists(const FWidgetBlueprintCompile
 		return MakeValue();
 	};
 
-	// Only find the source first property and destination first property.
-	//The full path will be tested later. We want to build the list of property needed.
-	for (int32 Index = 0; Index < BlueprintView->GetNumBindings(); ++Index)
+	// Find the start property for each path.
+	//The full path will be tested later. We want to build the list of property needed before generating the graphs.
+	for (const TSharedRef<FCompilerBinding>& ValidBinding : ValidBindings)
 	{
-		FMVVMBlueprintViewBinding* BindingPtr = BlueprintView->GetBindingAt(Index);
-		if (BindingPtr == nullptr)
-		{
-			WidgetBlueprintCompilerContext.MessageLog.Error(*FText::Format(LOCTEXT("BindingInvalidIndex", "Internal error: Tried to fetch binding for invalid index {0}."), Index).ToString());
-			bAreSourceContextsValid = false;
-			continue;
-		}
-
-		FMVVMBlueprintViewBinding& Binding = *BindingPtr;
-		if (!Binding.bCompile)
-		{
-			continue;
-		}
-
-		const bool bIsForwardBinding = IsForwardBinding(Binding.BindingType);
-		const bool bIsBackwardBinding = IsBackwardBinding(Binding.BindingType);
+		const FMVVMBlueprintViewBinding& Binding = *(BlueprintView->GetBindingAt(ValidBinding->Key.ViewBindingIndex));
 
-
-		auto RunGenerateCompilerSourceContext = [Self = this, &GenerateCompilerSourceContext, &Binding, BlueprintView](const FMVVMBlueprintPropertyPath& PropertyPath, FName ArgumentName)
-		{
-			TValueOrError<void, FText> SourceContextResult = GenerateCompilerSourceContext(PropertyPath);
-			if (SourceContextResult.HasError())
+		auto RunGenerateCompilerSourceContext = [Self = this, &GenerateCompilerContext, &Binding](bool bCreateSource, const FMVVMBlueprintPropertyPath& PropertyPath, FName ArgumentName)
 			{
-				Self->AddMessageForBinding(Binding, BlueprintView, SourceContextResult.StealError(), EBindingMessageType::Error, ArgumentName);
-				Self->bAreSourceContextsValid = false;
-			}
-		};
+				TValueOrError<void, FText> SourceContextResult = GenerateCompilerContext(bCreateSource, PropertyPath);
+				if (SourceContextResult.HasError())
+				{
+					Self->AddMessageForBinding(Binding, SourceContextResult.StealError(), EMessageType::Error, ArgumentName);
+					Self->bIsCreateVariableStepValid = false;
+				}
+			};
 
-		if (bIsForwardBinding || bIsBackwardBinding)
+		UMVVMBlueprintViewConversionFunction* ConversionFunction = Binding.Conversion.GetConversionFunction(ValidBinding->Key.bIsForwardBinding);
+		if (ConversionFunction)
 		{
-			UMVVMBlueprintViewConversionFunction* ConversionFunction = Binding.Conversion.GetConversionFunction(bIsForwardBinding);
-			if (ConversionFunction && ConversionFunction->NeedsWrapperGraph())
+			// validate the sources
+			for (const FMVVMBlueprintPin& Pin : ConversionFunction->GetPins())
 			{
-				ConversionFunction->SavePinValues(Context.GetWidgetBlueprint());
-				if (ConversionFunction->GetPins().Num() > 0)
+				if (Pin.UsedPathAsValue())
 				{
-					if (Binding.BindingType == EMVVMBindingMode::TwoWay)
+					if (Pin.GetName().IsNone())
 					{
-						AddMessageForBinding(Binding, BlueprintView, LOCTEXT("TwoWayBindingsWithConversion", "Two-way bindings are not allowed to use conversion functions."), EBindingMessageType::Error);
-						bAreSourceContextsValid = false;
+						AddMessageForBinding(Binding
+							, FText::Format(LOCTEXT("InvalidFunctionArgumentPathName", "The conversion function {0} has an invalid argument."), FText::FromString(Binding.GetDisplayNameString(WidgetBlueprintCompilerContext.WidgetBlueprint())))
+							, EMessageType::Error
+							, FName()
+						);
+						bIsCreateVariableStepValid = false;
 						continue;
 					}
-
-					// generate sources for conversion function arguments
-					for (const FMVVMBlueprintPin& Pin: ConversionFunction->GetPins())
-					{
-						if (Pin.UsedPathAsValue())
-						{
-							RunGenerateCompilerSourceContext(Pin.GetPath(), Pin.GetName());
-						}
-					}
-
-					// generate destination source
-					if (bIsForwardBinding)
-					{
-						RunGenerateCompilerSourceContext(Binding.DestinationPath, FName());
-					}
-					else
-					{
-
-						RunGenerateCompilerSourceContext(Binding.SourcePath, FName());
-					}
+					RunGenerateCompilerSourceContext(true, Pin.GetPath(), Pin.GetName());
 				}
 			}
-			else
-			{
-				// if we aren't using a conversion function, just validate the widget and viewmodel paths
-				RunGenerateCompilerSourceContext(Binding.DestinationPath, FName());
-				RunGenerateCompilerSourceContext(Binding.SourcePath, FName());
-			}
+
+			// validate the destination
+			const FMVVMBlueprintPropertyPath& BindingDestinationPath = ValidBinding->Key.bIsForwardBinding ? Binding.DestinationPath : Binding.SourcePath;
+			RunGenerateCompilerSourceContext(false, BindingDestinationPath, FName());
+		}
+		else
+		{
+			// if we aren't using a conversion function, validate the source and destination paths
+			RunGenerateCompilerSourceContext(ValidBinding->Key.bIsForwardBinding, Binding.SourcePath, FName());
+			RunGenerateCompilerSourceContext(!ValidBinding->Key.bIsForwardBinding, Binding.DestinationPath, FName());
 		}
 	}
 
 	// Find the event property.
 	//All inputs must also exist (for the BP to compile).
-	for (UMVVMBlueprintViewEvent* EventPtr : BlueprintView->GetEvents())
+	for (TSharedRef<FCompilerEvent>& ValidEvent : ValidEvents)
 	{
-		if (EventPtr == nullptr)
+		UMVVMBlueprintViewEvent* EventPtr = ValidEvent->Event.Get();
+		check(EventPtr);
+		auto RunGenerateCompilerSourceContext = [Self = this, &GenerateCompilerContext, EventPtr](bool bCreateSource, const FMVVMBlueprintPropertyPath& PropertyPath, FName ArgumentName)
 		{
-			WidgetBlueprintCompilerContext.MessageLog.Error(*LOCTEXT("EventInvalid", "Internal error: An event is invalid.").ToString());
-			bAreSourceContextsValid = false;
-			continue;
-		}
-		
-		if (!EventPtr->bCompile)
-		{
-			continue;
-		}
-
-		if (!EventPtr->GetEventPath().HasPaths())
-		{
-			AddMessageForEvent(EventPtr, LOCTEXT("EventInvalidEventPath", "The event path is invalid."), UMVVMBlueprintViewEvent::EMessageType::Error, FName());
-			bAreSourceContextsValid = false;
-			continue;
-		}
-
-		auto RunGenerateCompilerSourceContext = [Self = this, &GenerateCompilerSourceContext, EventPtr, BlueprintView](const FMVVMBlueprintPropertyPath& PropertyPath, FName ArgumentName)
-		{
-			TValueOrError<void, FText> SourceContextResult = GenerateCompilerSourceContext(PropertyPath);
+			TValueOrError<void, FText> SourceContextResult = GenerateCompilerContext(bCreateSource, PropertyPath);
 			if (SourceContextResult.HasError())
 			{
-				Self->AddMessageForEvent(EventPtr, SourceContextResult.StealError(), UMVVMBlueprintViewEvent::EMessageType::Error, ArgumentName);
-				Self->bAreSourceContextsValid = false;
+				Self->AddMessageForEvent(EventPtr, SourceContextResult.StealError(), EMessageType::Error, ArgumentName);
+				Self->bIsCreateVariableStepValid = false;
 			}
 		};
 
-		RunGenerateCompilerSourceContext(EventPtr->GetEventPath(), FName());
-		RunGenerateCompilerSourceContext(EventPtr->GetDestinationPath(), FName());
+		RunGenerateCompilerSourceContext(false, EventPtr->GetEventPath(), FName());
+		RunGenerateCompilerSourceContext(false, EventPtr->GetDestinationPath(), FName());
 		for (const FMVVMBlueprintPin& Pin : EventPtr->GetPins())
 		{
 			if (Pin.UsedPathAsValue())
 			{
-				RunGenerateCompilerSourceContext(Pin.GetPath(), Pin.GetName());
+				if (Pin.GetName().IsNone())
+				{
+					AddMessageForEvent(EventPtr, LOCTEXT("InvalidEventArgumentPathName", "The event has an invalid argument."), EMessageType::Error, Pin.GetName());
+					bIsCreateVariableStepValid = false;
+					continue;
+				}
+				RunGenerateCompilerSourceContext(false, Pin.GetPath(), Pin.GetName());
 			}
 		}
 	}
 }
 
 
-void FMVVMViewBlueprintCompiler::CreateFunctionsDeclaration(const FWidgetBlueprintCompilerContext::FCreateVariableContext& Context, UMVVMBlueprintView* BlueprintView)
+void FMVVMViewBlueprintCompiler::CreatePublicFunctionsDeclaration(const FWidgetBlueprintCompilerContext::FCreateVariableContext& Context)
 {
 	// Clean all previous intermediate function graph. It should stay alive. The graph lives on the Blueprint not on the class and it's used to generate the UFunction.
 	{
-		auto RenameObjectToTransientPackage = [](UObject* ObjectToRename)
-		{
-			const ERenameFlags RenFlags = REN_DontCreateRedirectors | REN_NonTransactional | REN_DoNotDirty;
-			ObjectToRename->Rename(nullptr, GetTransientPackage(), RenFlags);
-			ObjectToRename->SetFlags(RF_Transient);
-			ObjectToRename->ClearFlags(RF_Public | RF_Standalone | RF_ArchetypeObject);
-			FLinkerLoad::InvalidateExport(ObjectToRename);
-		};
-
 		for (UEdGraph* OldGraph : BlueprintView->TemporaryGraph)
 		{
 			if (OldGraph)
@@ -735,267 +878,392 @@ void FMVVMViewBlueprintCompiler::CreateFunctionsDeclaration(const FWidgetBluepri
 
 	if (GetDefault<UMVVMDeveloperProjectSettings>()->bAllowGeneratedViewModelSetter)
 	{
-		for (FCompilerSourceCreatorContext& SourceCreator : CompilerSourceCreatorContexts)
+		for (FCompilerViewModelSetter& Setter : ViewModelSettersToGenerate)
 		{
-			if (!SourceCreator.SetterFunctionName.IsEmpty() && SourceCreator.Type == ECompilerSourceCreatorType::ViewModel)
-			{
-				ensure(SourceCreator.SetterGraph == nullptr);
+			ensure(Setter.SetterGraph == nullptr);
 
-				SourceCreator.SetterGraph = UE::MVVM::FunctionGraphHelper::CreateIntermediateFunctionGraph(
-					WidgetBlueprintCompilerContext
-					, SourceCreator.SetterFunctionName
-					, (FUNC_BlueprintCallable | FUNC_Public | FUNC_Final)
-					, TEXT("Viewmodel")
-					, false);
-				BlueprintView->TemporaryGraph.Add(SourceCreator.SetterGraph);
-
-				if (SourceCreator.SetterGraph == nullptr || SourceCreator.SetterGraph->GetFName() != FName(*SourceCreator.SetterFunctionName))
-				{
-					WidgetBlueprintCompilerContext.MessageLog.Error(*FText::Format(LOCTEXT("SetterNameAlreadyExists", "The setter name {0} already exists and could not be autogenerated."),
-						FText::FromString(SourceCreator.SetterFunctionName)
-					).ToString()
-					);
-				}
+			Setter.SetterGraph = UE::MVVM::FunctionGraphHelper::CreateIntermediateFunctionGraph(
+				WidgetBlueprintCompilerContext
+				, Setter.BlueprintSetter
+				, (FUNC_BlueprintCallable | FUNC_Public | FUNC_Final)
+				, TEXT("Viewmodel")
+				, false);
+			BlueprintView->TemporaryGraph.Add(Setter.SetterGraph);
 
-				UE::MVVM::FunctionGraphHelper::AddFunctionArgument(SourceCreator.SetterGraph, SourceCreator.ViewModelContext.GetViewModelClass(), "Viewmodel");
+			if (Setter.SetterGraph == nullptr || Setter.SetterGraph->GetFName() != FName(*Setter.BlueprintSetter))
+			{
+				WidgetBlueprintCompilerContext.MessageLog.Error(*FText::Format(LOCTEXT("SetterNameAlreadyExists", "The setter name {0} already exists and could not be autogenerated."),
+					FText::FromString(Setter.BlueprintSetter)
+				).ToString()
+				);
+				bIsCreateVariableStepValid = false;
+				continue;
 			}
+
+			UE::MVVM::FunctionGraphHelper::AddFunctionArgument(Setter.SetterGraph, const_cast<UClass*>(Setter.Class), "Viewmodel");
 		}
 	}
 }
 
 
-void FMVVMViewBlueprintCompiler::CreateIntermediateGraphFunctions(const FWidgetBlueprintCompilerContext::FCreateVariableContext & Context, UMVVMBlueprintView * BlueprintView)
+void FMVVMViewBlueprintCompiler::CreateFunctions(const FWidgetBlueprintCompilerContext::FCreateFunctionContext& Context)
 {
-	// Add Generated Conversion functions to the blueprint
-	const int32 NumBindings = BlueprintView->GetNumBindings();
-	for (int32 Index = 0; Index < NumBindings; ++Index)
+	if (!AreStepsValid())
 	{
-		FMVVMBlueprintViewBinding* BindingPtr = BlueprintView->GetBindingAt(Index);
-		if (BindingPtr == nullptr)
-		{
-			WidgetBlueprintCompilerContext.MessageLog.Error(*FText::Format(LOCTEXT("InvalidBindingIndex", "Internal error. Invalid binding index given."), Index).ToString());
-			bAreBindingsValid = false;
-			continue;
-		}
+		return;
+	}
 
-		FMVVMBlueprintViewBinding& Binding = *BindingPtr;
-		if (!Binding.bCompile)
-		{
-			continue;
-		}
+	SourceViewModelDynamicCreatorContexts.Reset();
+
+	CategorizeBindings(Context);
+	CategorizeEvents(Context);
+	CreateWriteFieldContexts(Context);
+	CreateViewModelSetters(Context);
+	CreateIntermediateGraphFunctions(Context);
+}
+
+
+void FMVVMViewBlueprintCompiler::CategorizeBindings(const FWidgetBlueprintCompilerContext::FCreateFunctionContext& Context)
+{
+	// Find the type of the bindings
+	for (TSharedRef<FCompilerBinding>& ValidBinding : ValidBindings)
+	{
+		FMVVMBlueprintViewBinding& Binding = *BlueprintView->GetBindingAt(ValidBinding->Key.ViewBindingIndex);
+
+		check(ValidBinding->Type == FCompilerBinding::EType::Unknown);
+		ValidBinding->Type = FCompilerBinding::EType::Invalid;
 
-		auto AddGeneratedConversionFunction = [&Context, BindingPtr, this](bool bIsForwardBinding)
+		UMVVMBlueprintViewConversionFunction* ConversionFunction = Binding.Conversion.GetConversionFunction(ValidBinding->Key.bIsForwardBinding);
+		if (ConversionFunction)
 		{
-			UMVVMBlueprintViewConversionFunction* ConversionFunction = BindingPtr->Conversion.GetConversionFunction(bIsForwardBinding);
-			if (ConversionFunction && ConversionFunction->IsWrapperGraphTransient())
+			UClass* NewClass = WidgetBlueprintCompilerContext.NewClass;
+			if (!ConversionFunction->IsValid(WidgetBlueprintCompilerContext.WidgetBlueprint()))
 			{
-				UEdGraph* WrapperGraph = ConversionFunction->GetOrCreateIntermediateWrapperGraph(WidgetBlueprintCompilerContext);
-				if (WrapperGraph)
-				{
-					bool bAlreadyContained = WidgetBlueprintCompilerContext.Blueprint->FunctionGraphs.Contains(WrapperGraph);
-					if (ensure(!bAlreadyContained))
-					{
-						Context.AddGeneratedFunctionGraph(WrapperGraph);
-					}
-				}
+				AddMessageForBinding(Binding, LOCTEXT("InvalidConversionFunction", "The conversion function is invalid."), EMessageType::Error, FName());
+				bIsCreateFunctionsStepValid = false;
+				continue;
 			}
-		};
 
-		if (IsForwardBinding(Binding.BindingType))
-		{
-			AddGeneratedConversionFunction(true);
+			// Make sure the graph is up to date
+			UEdGraph* WrapperGraph = ConversionFunction->GetOrCreateIntermediateWrapperGraph(WidgetBlueprintCompilerContext);
+			if (WrapperGraph == nullptr)
+			{
+				AddMessageForBinding(Binding, LOCTEXT("InvalidConversionFunctionGraph", "The conversion function graph could not be generated."), EMessageType::Error, FName());
+				bIsCreateFunctionsStepValid = false;
+				continue;
+			}
+			ConversionFunction->SavePinValues(WidgetBlueprintCompilerContext.WidgetBlueprint());
+
+			if (ConversionFunction->GetPins().Num() == 0)
+			{
+				AddMessageForBinding(Binding
+					, FText::Format(LOCTEXT("InvalidNumberOfFunctionPin", "The conversion function {0} has no source."), FText::FromString(Binding.GetDisplayNameString(WidgetBlueprintCompilerContext.WidgetBlueprint())))
+					, EMessageType::Error
+					, FName()
+				);
+				bIsCreateFunctionsStepValid = false;
+				continue;
+			}
+
+			if (ConversionFunction->NeedsWrapperGraph(WidgetBlueprintCompilerContext.WidgetBlueprint()))
+			{
+				ValidBinding->Type = FCompilerBinding::EType::ComplexConversionFunction;
+			}
+			else
+			{
+				ValidBinding->Type = FCompilerBinding::EType::SimpleConversionFunction;
+			}
+			ValidBinding->ConversionFunction = ConversionFunction;
+
+			// Because the editor use the destination to order the bindings, the destination path can be "valid". Pointing only to the WidgetName.
+			//const FMVVMBlueprintPropertyPath& BindingSourcePath = ValidBinding->Key.bIsForwardBinding ? Binding.SourcePath : Binding.DestinationPath;
+			//if (BindingSourcePath.IsValid())
+			//{
+			//	AddMessageForBinding(Binding, LOCTEXT("ShouldNotHaveSourceWarning", "Internal Error. The binding should not have a source."), EMessageType::Warning, FName());
+			//}
 		}
-		if (IsBackwardBinding(Binding.BindingType))
+		else
 		{
-			AddGeneratedConversionFunction(false);
+			ValidBinding->Type = FCompilerBinding::EType::Assignment;
 		}
 	}
+}
 
-	// Add Generated event to the blueprint
-	for (UMVVMBlueprintViewEvent* EventPtr : BlueprintView->GetEvents())
-	{
-		if (EventPtr == nullptr)
-		{
-			WidgetBlueprintCompilerContext.MessageLog.Error(*LOCTEXT("InvalidEvent", "Internal error. Invalid event.").ToString());
-			bAreEventsValid = false;
-			continue;
-		}
 
-		if (!EventPtr->bCompile)
-		{
-			continue;
-		}
+void FMVVMViewBlueprintCompiler::CategorizeEvents(const FWidgetBlueprintCompilerContext::FCreateFunctionContext& Context)
+{
+	for (TSharedRef<FCompilerEvent>& Event : ValidEvents)
+	{
+		Event->Type = FCompilerEvent::EType::Invalid;
 
+		UMVVMBlueprintViewEvent* EventPtr = Event->Event.Get();
 		UEdGraph* WrapperGraph = EventPtr->GetOrCreateWrapperGraph();
 		if (WrapperGraph == nullptr)
 		{
-			AddMessageForEvent(EventPtr, LOCTEXT("InvalidEventGraph", "The event could not be generated"), UMVVMBlueprintViewEvent::EMessageType::Warning, FName());
+			AddMessageForEvent(EventPtr, LOCTEXT("InvalidEventGraph", "The event could not be generated."), EMessageType::Warning, FName());
+			bIsCreateFunctionsStepValid = false;
 			continue;
 		}
+		EventPtr->SavePinValues();
 
-		bool bAlreadyContained = WidgetBlueprintCompilerContext.Blueprint->FunctionGraphs.Contains(WrapperGraph);
-		if (ensure(!bAlreadyContained))
-		{
-			Context.AddGeneratedFunctionGraph(WrapperGraph);
-		}
+		Event->Type = FCompilerEvent::EType::Valid;
 	}
 }
 
 
-void FMVVMViewBlueprintCompiler::CreateBindingDestinationContexts(UMVVMBlueprintView* BlueprintView)
+void FMVVMViewBlueprintCompiler::CreateWriteFieldContexts(const FWidgetBlueprintCompilerContext::FCreateFunctionContext& Context)
 {
-	if (!bAreSourcesCreatorValid || !bAreSourceContextsValid || !bAreBindingsValid || !bAreEventsValid)
-	{
-		return;
-	}
-
-	BindingDestinationContexts.Empty();
-
 	// Use the Skeleton class. The class bind and not all functions are generated yet
 	UWidgetBlueprintGeneratedClass* NewSkeletonClass = Cast<UWidgetBlueprintGeneratedClass>(WidgetBlueprintCompilerContext.Blueprint->SkeletonGeneratedClass);
 	if (NewSkeletonClass == nullptr)
 	{
 		WidgetBlueprintCompilerContext.MessageLog.Error(*LOCTEXT("InvalidNewClass", "Internal error. The skeleton class is not valid.").ToString());
+		bIsCreateFunctionsStepValid = false;
 		return;
 	}
 
-	const int32 NumBindings = BlueprintView->GetNumBindings();
-	for (int32 Index = 0; Index < NumBindings; ++Index)
+	for (TSharedRef<FCompilerBinding>& ValidBinding : ValidBindings)
 	{
-		FMVVMBlueprintViewBinding* BindingPtr = BlueprintView->GetBindingAt(Index);
-		if (BindingPtr == nullptr)
-		{
-			WidgetBlueprintCompilerContext.MessageLog.Error(*FText::Format(LOCTEXT("InvalidBindingIndex", "Internal error. Invalid binding index given."), Index).ToString());
-			bAreBindingsValid = false;
-			continue;
-		}
+		const FMVVMBlueprintViewBinding& Binding = *BlueprintView->GetBindingAt(ValidBinding->Key.ViewBindingIndex);
 
-		FMVVMBlueprintViewBinding& Binding = *BindingPtr;
-		if (!Binding.bCompile)
+		if (ValidBinding->Type == FCompilerBinding::EType::Assignment
+			|| ValidBinding->Type == FCompilerBinding::EType::SimpleConversionFunction
+			|| ValidBinding->Type == FCompilerBinding::EType::ComplexConversionFunction)
 		{
-			continue;
+			const FMVVMBlueprintPropertyPath& DestinationPropertyPath = ValidBinding->Key.bIsForwardBinding ? Binding.DestinationPath : Binding.SourcePath;
+			TValueOrError<FCreateFieldsResult, FText> FieldContextResult = CreateFieldContext(NewSkeletonClass, DestinationPropertyPath, false);
+			if (FieldContextResult.HasError())
+			{
+				AddMessageForBinding(Binding, FieldContextResult.StealError(), EMessageType::Error, FName());
+				bIsCreateFunctionsStepValid = false;
+				continue;
+			}
+
+			// Test if it already exist
+			const TArray<UE::MVVM::FMVVMConstFieldVariant>& SkeletalGeneratedFieldsResult = FieldContextResult.GetValue().SkeletalGeneratedFields;
+			{
+				TSharedRef<FGeneratedWriteFieldPathContext>* Found = GeneratedWriteFieldPaths.FindByPredicate([&SkeletalGeneratedFieldsResult](const TSharedRef<FGeneratedWriteFieldPathContext>& Other)
+					{
+						return Other->SkeletalGeneratedFields == SkeletalGeneratedFieldsResult;
+					});
+				if (Found)
+				{
+					// Temporary removing this message until the assets are fixed.
+					//AddMessageForBinding(Binding
+					//	, FText::Format(LOCTEXT("PropertyPathAlreadyUsed", "The property path '{0}' is already used by another binding."), PropertyPathToText(NewSkeletonClass, BlueprintView.Get(), DestinationPropertyPath))
+					//	, EMessageType::Warning
+					//	, FName()
+					//);
+					ValidBinding->WritePath = *Found;
+					continue;
+				}
+			}
+
+			TSharedRef<FGeneratedWriteFieldPathContext> WriteFieldPath = MakeShared<FGeneratedWriteFieldPathContext>();
+			GeneratedWriteFieldPaths.Add(WriteFieldPath);
+			WriteFieldPath->UsedByBindings.AddUnique(ValidBinding);
+			WriteFieldPath->GeneratedFields = MoveTemp(FieldContextResult.GetValue().GeneratedFields);
+			WriteFieldPath->SkeletalGeneratedFields = MoveTemp(FieldContextResult.GetValue().SkeletalGeneratedFields);
+			WriteFieldPath->GeneratedFrom = DestinationPropertyPath.GetSource(WidgetBlueprintCompilerContext.WidgetBlueprint());
+			WriteFieldPath->bCanBeSetInNative = CanBeSetInNative(WriteFieldPath->SkeletalGeneratedFields);
+			WriteFieldPath->bUseByNativeBinding = true; // the setter function can be in BP but the destination will be set in native
+
+			// Assign the Destination to the binding
+			ValidBinding->WritePath = WriteFieldPath;
 		}
+	}
 
-		auto AddDestination = [Index, BindingPtr, BlueprintView, NewSkeletonClass, this](bool bIsForwardBinding)
-		{
-			FBindingDestinationContext& NewDestinationContext = BindingDestinationContexts.AddDefaulted_GetRef();
-			NewDestinationContext.BindingIndex = Index;
-			NewDestinationContext.bIsForwardBinding = bIsForwardBinding;
+	//The destination for event must also exist (for the BP to compile).
+	for (TSharedRef<FCompilerEvent>& ValidEvent : ValidEvents)
+	{
+		UMVVMBlueprintViewEvent* EventPtr = ValidEvent->Event.Get();
+		check(EventPtr);
 
-			const FMVVMBlueprintPropertyPath& DestinationPath = bIsForwardBinding ? BindingPtr->DestinationPath : BindingPtr->SourcePath;
-			TArray<UE::MVVM::FMVVMConstFieldVariant> SkeletalSetterPath = CreateBindingDestinationPath(BlueprintView, NewSkeletonClass, DestinationPath);
-			if (!IsPropertyPathValid(WidgetBlueprintCompilerContext.WidgetBlueprint(), SkeletalSetterPath))
+		if (ValidEvent->Type == FCompilerEvent::EType::Valid)
+		{
+			TValueOrError<FCreateFieldsResult, FText> FieldContextResult = CreateFieldContext(NewSkeletonClass, EventPtr->GetDestinationPath(), false);
+			if (FieldContextResult.HasError())
 			{
-				AddMessageForBinding(*BindingPtr, BlueprintView, FText::Format(Private::PropertyPathIsInvalidFormat,
-					PropertyPathToText(NewSkeletonClass, BlueprintView, DestinationPath)),
-					EBindingMessageType::Error
+				AddMessageForEvent(EventPtr
+					, FText::Format(Private::PropertyPathIsInvalidFormat, PropertyPathToText(NewSkeletonClass, BlueprintView.Get(), EventPtr->GetDestinationPath()))
+					, EMessageType::Error
+					, FName()
 				);
-				bAreBindingsValid = false;
-				return;
-			}
-			
-			// Generate the path with property converted to BP function
-			TValueOrError<TArray<FMVVMConstFieldVariant>, FText> SkeletalGeneratedField = FieldPathHelper::GenerateFieldPathList(SkeletalSetterPath, false);
-			if (SkeletalGeneratedField.HasError() || !IsPropertyPathValid(WidgetBlueprintCompilerContext.WidgetBlueprint(), SkeletalGeneratedField.GetValue()))
-			{
-				AddMessageForBinding(*BindingPtr, BlueprintView, FText::Format(Private::CouldNotCreateSourceFieldPathFormat
-					, ::UE::MVVM::FieldPathHelper::ToText(SkeletalSetterPath)
-					, SkeletalGeneratedField.GetError())
-					, EBindingMessageType::Error
-				);
-				bAreBindingsValid = false;
-				return;
+				bIsCreateFunctionsStepValid = false;
+				continue;
 			}
 
-			NewDestinationContext.bCanBeSetInNative = CanBeSetInNative(SkeletalGeneratedField.GetValue());
-
-			// If the destination can't be set in cpp, we need to generate a BP function to set the value.
-			if (!NewDestinationContext.bCanBeSetInNative)
+			// Test if it already exist
+			const TArray<UE::MVVM::FMVVMConstFieldVariant>& SkeletalGeneratedFieldsResult = FieldContextResult.GetValue().SkeletalGeneratedFields;
 			{
-				const FProperty* SetterProperty = nullptr;
-				if (SkeletalGeneratedField.GetValue().Num() > 0 && ensure(SkeletalGeneratedField.GetValue().Last().IsProperty()))
+				TSharedRef<FGeneratedWriteFieldPathContext>* Found = GeneratedWriteFieldPaths.FindByPredicate([&SkeletalGeneratedFieldsResult](const TSharedRef<FGeneratedWriteFieldPathContext>& Other) { return Other->SkeletalGeneratedFields == SkeletalGeneratedFieldsResult; });
+				if (Found)
 				{
-					SetterProperty = SkeletalGeneratedField.GetValue().Last().GetProperty();
+					(*Found)->UsedByEvents.AddUnique(ValidEvent);
+					ValidEvent->WritePath = *Found;
 				}
-
-				if (SetterProperty == nullptr)
+				else
 				{
-					AddMessageForBinding(*BindingPtr, BlueprintView, FText::Format(LOCTEXT("CantGetSetter", "Internal Error. The setter function was not created. {0}"),
-						PropertyPathToText(NewSkeletonClass, BlueprintView, DestinationPath)),
-						EBindingMessageType::Error
-					);
-					bAreBindingsValid = false;
-					return;
+					TSharedRef<FGeneratedWriteFieldPathContext> WriteFieldPath = MakeShared<FGeneratedWriteFieldPathContext>();
+					GeneratedWriteFieldPaths.Add(WriteFieldPath);
+					WriteFieldPath->UsedByEvents.AddUnique(ValidEvent);
+					WriteFieldPath->GeneratedFields = MoveTemp(FieldContextResult.GetValue().GeneratedFields);
+					WriteFieldPath->SkeletalGeneratedFields = MoveTemp(FieldContextResult.GetValue().SkeletalGeneratedFields);
+					WriteFieldPath->GeneratedFrom = EventPtr->GetDestinationPath().GetSource(WidgetBlueprintCompilerContext.WidgetBlueprint());
+					WriteFieldPath->bCanBeSetInNative = CanBeSetInNative(WriteFieldPath->SkeletalGeneratedFields);
+					WriteFieldPath->bUseByNativeBinding = false;
+
+					// Assign the Destination to the events
+					ValidEvent->WritePath = WriteFieldPath;
 				}
+			}
+		}
+	}
+}
 
-				// create a setter function to be called from native. For now we follow the convention of Setter(Conversion(Getter))
-				UEdGraph* GeneratedSetterGraph = UE::MVVM::FunctionGraphHelper::CreateIntermediateFunctionGraph(WidgetBlueprintCompilerContext, FString::Printf(TEXT("__Setter_%s"), *SetterProperty->GetName()), EFunctionFlags::FUNC_None, TEXT("AutogeneratedSetter"), false);
-				if (GeneratedSetterGraph == nullptr)
+
+void FMVVMViewBlueprintCompiler::CreateViewModelSetters(const FWidgetBlueprintCompilerContext::FCreateFunctionContext& Context)
+{
+	if (GetDefault<UMVVMDeveloperProjectSettings>()->bAllowGeneratedViewModelSetter)
+	{
+		for (FCompilerViewModelSetter& Setter : ViewModelSettersToGenerate)
+		{
+			if (ensure(Setter.SetterGraph != nullptr))
+			{
+				if (!UE::MVVM::FunctionGraphHelper::GenerateViewModelSetter(WidgetBlueprintCompilerContext, Setter.SetterGraph, Setter.PropertyName))
 				{
-					AddMessageForBinding(*BindingPtr, BlueprintView, FText::Format(LOCTEXT("CantCreateSetter", "Internal Error. The setter function was not created. {0}"),
-						PropertyPathToText(NewSkeletonClass, BlueprintView, DestinationPath)),
-						EBindingMessageType::Error
-					);
-					bAreBindingsValid = false;
-					return;
+					AddMessageForViewModel(Setter.DisplayName, LOCTEXT("SetterFunctionCouldNotBeGenerated", "The setter function could not be generated."), EMessageType::Warning);
+					bIsCreateFunctionsStepValid = false;
+					continue;
 				}
+			}
+		}
+	}
+}
 
-				UE::MVVM::FunctionGraphHelper::AddFunctionArgument(GeneratedSetterGraph, SetterProperty, "NewValue");
 
-				// set GeneratedSetterFunction. Use the SkeletalSetterPath here to use the setter will be generated when the function is generated.
-				if (!UE::MVVM::FunctionGraphHelper::GenerateIntermediateSetter(WidgetBlueprintCompilerContext, GeneratedSetterGraph, SkeletalSetterPath))
-				{
-					AddMessageForBinding(*BindingPtr, BlueprintView, FText::Format(LOCTEXT("CantGeneratedSetter", "Internal Error. The setter function was not generated. {0}"),
-						PropertyPathToText(NewSkeletonClass, BlueprintView, DestinationPath)),
-						EBindingMessageType::Error
-					);
-					bAreBindingsValid = false;
-					return;
-				}
+void FMVVMViewBlueprintCompiler::CreateIntermediateGraphFunctions(const FWidgetBlueprintCompilerContext::FCreateFunctionContext& Context)
+{
+	// Add function to set destination if needed
+	for (TSharedRef<FGeneratedWriteFieldPathContext>& GeneratedDestination : GeneratedWriteFieldPaths)
+	{
+		// If the destination can't be set in cpp, we need to generate a BP function to set the value.
+		if (!GeneratedDestination->bCanBeSetInNative && GeneratedDestination->bUseByNativeBinding)
+		{
+			auto AddErrorMessage = [Self = this, &GeneratedDestination](const FText& ErrorMsg)
+			{
+				Self->AddMessages(GeneratedDestination->UsedByBindings, GeneratedDestination->UsedByEvents, ErrorMsg, EMessageType::Error);
+			};
 
-				// the new path can only be set later, once the function is compiled.
-				NewDestinationContext.GeneratedFunctionName = GeneratedSetterGraph->GetFName();
+			const FProperty* SetterProperty = nullptr;
+			if (ensure(GeneratedDestination->SkeletalGeneratedFields.Num() > 0 && GeneratedDestination->SkeletalGeneratedFields.Last().IsProperty()))
+			{
+				SetterProperty = GeneratedDestination->SkeletalGeneratedFields.Last().GetProperty();
 			}
-		};
 
-		if (IsForwardBinding(Binding.BindingType))
+			if (SetterProperty == nullptr)
+			{
+				AddErrorMessage(FText::Format(LOCTEXT("CantGetSetter", "Internal Error. The setter function was not created. {0}"), ::UE::MVVM::FieldPathHelper::ToText(GeneratedDestination->GeneratedFields)));
+				bIsCreateFunctionsStepValid = false;
+				continue;
+			}
+
+			// create a setter function to be called from native. For now we follow the convention of Setter(Conversion(Getter))
+			UEdGraph* GeneratedSetterGraph = UE::MVVM::FunctionGraphHelper::CreateIntermediateFunctionGraph(WidgetBlueprintCompilerContext, FString::Printf(TEXT("__Setter_%s"), *SetterProperty->GetName()), EFunctionFlags::FUNC_None, TEXT("AutogeneratedSetter"), false);
+			if (GeneratedSetterGraph == nullptr)
+			{
+				AddErrorMessage(FText::Format(LOCTEXT("CantCreateSetter", "Internal Error. The setter function was not created. {0}"), ::UE::MVVM::FieldPathHelper::ToText(GeneratedDestination->GeneratedFields)));
+				bIsCreateFunctionsStepValid = false;
+				continue;
+			}
+
+			UE::MVVM::FunctionGraphHelper::AddFunctionArgument(GeneratedSetterGraph, SetterProperty, "NewValue");
+
+			// set GeneratedSetterFunction. Use the SkeletalSetterPath here to use the setter will be generated when the function is generated.
+			if (!UE::MVVM::FunctionGraphHelper::GenerateIntermediateSetter(WidgetBlueprintCompilerContext, GeneratedSetterGraph, GeneratedDestination->SkeletalGeneratedFields))
+			{
+				AddErrorMessage(FText::Format(LOCTEXT("CantGeneratedSetter", "Internal Error. The setter function was not generated. {0}"), ::UE::MVVM::FieldPathHelper::ToText(GeneratedDestination->GeneratedFields)));
+				bIsCreateFunctionsStepValid = false;
+				continue;
+			}
+
+			// the new path can only be set later, once the function is compiled.
+			GeneratedDestination->GeneratedFunctionName = GeneratedSetterGraph->GetFName();
+		}
+	}
+
+	// Add Generated Conversion functions to the blueprint
+	for (const TSharedRef<FCompilerBinding>& ValidBinding : ValidBindings)
+	{
+		FMVVMBlueprintViewBinding& Binding = *BlueprintView->GetBindingAt(ValidBinding->Key.ViewBindingIndex);
+
+		if (ValidBinding->Type == FCompilerBinding::EType::ComplexConversionFunction)
 		{
-			AddDestination(true);
+			UMVVMBlueprintViewConversionFunction* ConversionFunction = Binding.Conversion.GetConversionFunction(ValidBinding->Key.bIsForwardBinding);
+			check(ConversionFunction);
+
+			UEdGraph* WrapperGraph = ConversionFunction->GetOrCreateIntermediateWrapperGraph(WidgetBlueprintCompilerContext);
+			if (ensure(WrapperGraph) && ConversionFunction->IsWrapperGraphTransient())
+			{
+				bool bAlreadyContained = WidgetBlueprintCompilerContext.Blueprint->FunctionGraphs.Contains(WrapperGraph);
+				if (ensure(!bAlreadyContained))
+				{
+					Context.AddGeneratedFunctionGraph(WrapperGraph);
+				}
+			}
 		}
-		if (IsBackwardBinding(Binding.BindingType))
+	}
+
+	// Add Generated event to the blueprint
+	for (TSharedRef<FCompilerEvent>& Event : ValidEvents)
+	{
+		UMVVMBlueprintViewEvent* EventPtr = Event->Event.Get();
+		if (Event->Type == FCompilerEvent::EType::Valid)
 		{
-			AddDestination(false);
+			UEdGraph* WrapperGraph = EventPtr->GetOrCreateWrapperGraph();
+			ensure(WrapperGraph);
+
+			bool bAlreadyContained = WidgetBlueprintCompilerContext.Blueprint->FunctionGraphs.Contains(WrapperGraph);
+			if (ensure(!bAlreadyContained))
+			{
+				Context.AddGeneratedFunctionGraph(WrapperGraph);
+			}
 		}
 	}
 }
 
 
-bool FMVVMViewBlueprintCompiler::PreCompile(UWidgetBlueprintGeneratedClass* Class, UMVVMBlueprintView* BlueprintView)
+bool FMVVMViewBlueprintCompiler::PreCompile(UWidgetBlueprintGeneratedClass* Class)
 {
-	if (!bAreSourcesCreatorValid || !bAreSourceContextsValid || !bAreBindingsValid || !bAreEventsValid)
+	if (!AreStepsValid())
 	{
 		return false;
 	}
 
-	const int32 NumBindings = BlueprintView->GetNumBindings();
-	CompilerBindings.Reset(NumBindings*2);
-	BindingSourceContexts.Reset(NumBindings*2);
-	SimpleBindingContexts.Reset();
-	ComplexConversionFunctionContexts.Reset();
-	EventDelegateContexts.Reset();
+	FixCompilerBindingSelfSource(Class);
+	AddWarningForPropertyWithMVVMAndLegacyBinding(Class);
+
+	GeneratedReadFieldPaths.Reset();
+	FixWriteFieldPathContext(Class);
+	CreateReadFieldContexts(Class);
+	CreateCreatorContentFromBindingSource(Class);
+
+	if (!AreStepsValid())
+	{
+		return false;
+	}
 
-	PreCompileBindingSources(Class, BlueprintView);
-	PreCompileSourceCreators(Class, BlueprintView);
-	PreCompileBindings(Class, BlueprintView);
-	PreCompileEvents(Class, BlueprintView);
+	PreCompileViewModelCreatorContexts(Class);
+	PreCompileBindings(Class);
+	PreCompileEvents(Class);
 
-	return bAreSourcesCreatorValid && bAreSourceContextsValid && bAreBindingsValid && bAreEventsValid;
+	return AreStepsValid();
 }
 
 
-bool FMVVMViewBlueprintCompiler::Compile(UWidgetBlueprintGeneratedClass* Class, UMVVMBlueprintView* BlueprintView, UMVVMViewClass* ViewExtension)
+bool FMVVMViewBlueprintCompiler::Compile(UWidgetBlueprintGeneratedClass* Class, UMVVMViewClass* ViewExtension)
 {
-	if (!bAreSourcesCreatorValid || !bAreSourceContextsValid || !bAreBindingsValid || !bAreEventsValid)
+	if (!AreStepsValid())
 	{
 		return false;
 	}
@@ -1006,22 +1274,22 @@ bool FMVVMViewBlueprintCompiler::Compile(UWidgetBlueprintGeneratedClass* Class,
 		WidgetBlueprintCompilerContext.MessageLog.Error(*FText::Format(LOCTEXT("BindingCompilationFailed", "The binding compilation failed. {1}"), CompileResult.GetError()).ToString());
 		return false;
 	}
-	CompileSourceCreators(CompileResult.GetValue(), Class, BlueprintView, ViewExtension);
-	CompileBindings(CompileResult.GetValue(), Class, BlueprintView, ViewExtension);
-	CompileEvents(CompileResult.GetValue(), Class, BlueprintView, ViewExtension);
+	CompileViewModelCreatorContexts(CompileResult.GetValue(), Class, ViewExtension);
+	CompileBindings(CompileResult.GetValue(), Class, ViewExtension);
+	CompileEvents(CompileResult.GetValue(), Class, ViewExtension);
 
 	{
 		ViewExtension->bInitializeSourcesOnConstruct = BlueprintView->GetSettings()->bInitializeSourcesOnConstruct;
 		ViewExtension->bInitializeBindingsOnConstruct = ViewExtension->bInitializeSourcesOnConstruct ? BlueprintView->GetSettings()->bInitializeBindingsOnConstruct : false;
 	}
 
-	bool bResult = bAreSourcesCreatorValid && bAreSourceContextsValid && bAreBindingsValid && bAreEventsValid;
+	bool bResult = AreStepsValid();
 	if (bResult)
 	{
 		ViewExtension->BindingLibrary = MoveTemp(CompileResult.GetValue().Library);
 
 #if UE_WITH_MVVM_DEBUGGING
-		if (CVarLogViewCompliedResult->GetBool())
+		if (CVarLogViewCompiledResult->GetBool())
 		{
 			FMVVMViewClass_SourceCreator::FToStringArgs CreatorsToStringArgs = FMVVMViewClass_SourceCreator::FToStringArgs::All();
 			CreatorsToStringArgs.bUseDisplayName = false;
@@ -1032,365 +1300,454 @@ bool FMVVMViewBlueprintCompiler::Compile(UWidgetBlueprintGeneratedClass* Class,
 #endif
 	}
 
-	return bResult;
+	return AreStepsValid();
 }
 
 
-bool FMVVMViewBlueprintCompiler::PreCompileBindingSources(UWidgetBlueprintGeneratedClass* Class, UMVVMBlueprintView* BlueprintView)
+void FMVVMViewBlueprintCompiler::FixWriteFieldPathContext(UWidgetBlueprintGeneratedClass* Class)
 {
-	const int32 NumBindings = BlueprintView->GetNumBindings(); // NB Binding can be added when creating a dynamic vm
-	for (int32 Index = 0; Index < NumBindings; ++Index)
+	for (TSharedRef<FGeneratedWriteFieldPathContext>& GeneratedDestination : GeneratedWriteFieldPaths)
 	{
-		FMVVMBlueprintViewBinding* BindingPtr = BlueprintView->GetBindingAt(Index);
-		if (BindingPtr == nullptr)
+		// If a graph was generated. Use it instead for the SkeletalGeneratedFields
+		if (!GeneratedDestination->GeneratedFunctionName.IsNone())
 		{
-			WidgetBlueprintCompilerContext.MessageLog.Error(*FText::Format(LOCTEXT("InvalidBindingIndex", "Internal error. Invalid binding index given."), Index).ToString());
-			bAreBindingsValid = false;
-			continue;
-		}
+			UFunction* GeneratedFunction = Class->FindFunctionByName(GeneratedDestination->GeneratedFunctionName);
+			if (GeneratedFunction == nullptr)
+			{
+				WidgetBlueprintCompilerContext.MessageLog.Error(*LOCTEXT("CantFindGeneratedBindingSetterFunction", "Internal Error. The setter function was not generated.").ToString());
+				bIsPreCompileStepValid = false;
+				continue;
+			}
 
-		FMVVMBlueprintViewBinding& Binding = *BindingPtr;
-		if (!Binding.bCompile)
-		{
-			continue;
+			GeneratedDestination->GeneratedFields.Reset();
+			GeneratedDestination->GeneratedFields.Add(UE::MVVM::FMVVMConstFieldVariant(GeneratedFunction));
+			// note the function is not added on the skeletal class
+			GeneratedDestination->SkeletalGeneratedFields.Reset();
+			GeneratedDestination->SkeletalGeneratedFields.Add(UE::MVVM::FMVVMConstFieldVariant(GeneratedFunction));
 		}
+	}
+}
 
-		bool bIsOneTimeBinding = IsOneTimeBinding(Binding.BindingType);
 
-		auto CreateSourceContextForPropertyPath = [Self=this, &Binding, BlueprintView, Class, Index, bIsOneTimeBinding](const FMVVMBlueprintPropertyPath& Path, bool bForwardBinding, int32 ComplexConversionFunctionContextIndex, FName ArgumentName) -> bool
+void FMVVMViewBlueprintCompiler::CreateReadFieldContexts(UWidgetBlueprintGeneratedClass* Class)
+{
+	auto AlreadyExist = [Self = this](const TArray<UE::MVVM::FMVVMConstFieldVariant>& SkeletalGeneratedFields) -> TSharedPtr<FGeneratedReadFieldPathContext>
 		{
-			const TValueOrError<FBindingSourceContext, FText> CreatedBindingSourceContext = Self->CreateBindingSourceContext(BlueprintView, Class, Path, bIsOneTimeBinding);
-			if (CreatedBindingSourceContext.HasError())
-			{
-				Self->AddMessageForBinding(Binding, BlueprintView,
-					FText::Format(LOCTEXT("PropertyPathInvalidWithReason", "The property path '{0}' is invalid. {1}"),
-						PropertyPathToText(Class, BlueprintView, Binding.SourcePath),
-						CreatedBindingSourceContext.GetError()
-					),
-					EBindingMessageType::Error,
-					ArgumentName
-				);
-				return false;
-			}
-
-			FBindingSourceContext BindingSourceContext = CreatedBindingSourceContext.GetValue();
-			if (!IsPropertyPathValid(Self->WidgetBlueprintCompilerContext.WidgetBlueprint(), BindingSourceContext.PropertyPath))
-			{
-				Self->AddMessageForBinding(Binding,
-					BlueprintView,
-					FText::Format(Private::PropertyPathIsInvalidFormat, PropertyPathToText(Class, BlueprintView, Binding.SourcePath)),
-					EBindingMessageType::Error,
-					ArgumentName
-				);
-				return false;
-			}
+			// Test if it already exist
+			TSharedRef<FGeneratedReadFieldPathContext>* Found = Self->GeneratedReadFieldPaths.FindByPredicate([&SkeletalGeneratedFields](const TSharedRef<FGeneratedReadFieldPathContext>& Other) { return Other->SkeletalGeneratedFields == SkeletalGeneratedFields; });
+			return Found != nullptr ? *Found : TSharedPtr<FGeneratedReadFieldPathContext>();
+		};
 
-			if (BindingSourceContext.SourceClass == nullptr)
-			{
-				Self->AddMessageForBinding(Binding, BlueprintView, LOCTEXT("BindingInvalidSourceClass", "Internal error. The binding could not find its source class."), EBindingMessageType::Error, ArgumentName);
-				return false;
-			}
+	for (TSharedRef<FCompilerBinding>& ValidBinding : ValidBindings)
+	{
+		const FMVVMBlueprintViewBinding& Binding = *BlueprintView->GetBindingAt(ValidBinding->Key.ViewBindingIndex);
 
-			if (!BindingSourceContext.bIsRootWidget && BindingSourceContext.UserWidgetPropertyContextIndex == INDEX_NONE && BindingSourceContext.SourceCreatorContextIndex == INDEX_NONE)
+		auto CreateBindingSourceContext = [Self = this, Class, &ValidBinding, &Binding](const FMVVMBlueprintPropertyPath& PropertyPath, FName ArgumentName) -> TValueOrError<FCreateFieldsResult, void>
+		{
+			TValueOrError<FCreateFieldsResult, FText> FieldContextResult = Self->CreateFieldContext(Class, PropertyPath, true);
+			if (FieldContextResult.HasError())
 			{
-				Self->AddMessageForBinding(Binding, BlueprintView, LOCTEXT("BindingInvalidSource", "Internal error. The binding could not find its source."), EBindingMessageType::Error, ArgumentName);
-				return false;
+				Self->AddMessageForBinding(Binding
+					, FText::Format(Private::PropertyPathIsInvalidFormat, PropertyPathToText(Class, Self->BlueprintView.Get(), PropertyPath))
+					, EMessageType::Error
+					, ArgumentName
+				);
+				Self->bIsPreCompileStepValid = false;
+				return MakeError();
 			}
 
-			BindingSourceContext.BindingIndex = Index;
-			BindingSourceContext.bIsForwardBinding = bForwardBinding;
-			BindingSourceContext.ComplexConversionFunctionContextIndex = ComplexConversionFunctionContextIndex;
-
-			Self->BindingSourceContexts.Add(MoveTemp(BindingSourceContext));
-			return true;
+			return MakeValue(FieldContextResult.StealValue());
 		};
 
-		enum class ECreateSourcesForConversionFunctionResult : uint8 { Valid, Failed, Continue };
-		auto CreateSourcesForConversionFunction = [this, &Binding, BlueprintView, Index, &CreateSourceContextForPropertyPath](bool bForwardBinding)
+		auto CreateFieldId = [Self = this, Class, &Binding, bIsOneTimeBinding = ValidBinding->bIsOneTimeBinding](const FMVVMBlueprintPropertyPath& PropertyPath, TSharedPtr<FGeneratedReadFieldPathContext>& ReadFieldContext, FName ArgumentName) -> TValueOrError<void, void>
 		{
-			ECreateSourcesForConversionFunctionResult Result = ECreateSourcesForConversionFunctionResult::Continue;
-			if (const UMVVMBlueprintViewConversionFunction* ConversionFunction = Binding.Conversion.GetConversionFunction(bForwardBinding))
+			if (!bIsOneTimeBinding && !ReadFieldContext->NotificationField.IsValid())
 			{
-				if (ConversionFunction->NeedsWrapperGraph())
+				TValueOrError<FCompilerNotifyFieldId, FText> CreateFieldResult = Self->CreateNotifyFieldId(Class, ReadFieldContext, Binding);
+				if (CreateFieldResult.HasError())
 				{
-					const int32 ComplexConversionFunctionContextIndex = ComplexConversionFunctionContexts.AddDefaulted();
-					{
-						FComplexConversionFunctionContext& NewConversionFunctionContext = ComplexConversionFunctionContexts[ComplexConversionFunctionContextIndex];
-						NewConversionFunctionContext.BindingIndex = Index;
-						NewConversionFunctionContext.bIsForwardBinding = bForwardBinding;
-					}
+					Self->AddMessageForBinding(Binding
+						, FText::Format(LOCTEXT("CreateNotifyFieldIdFailedInvalidSelfContext", "The property path '{0}' is invalid. {1}"), PropertyPathToText(Class, Self->BlueprintView.Get(), PropertyPath), CreateFieldResult.StealError())
+						, EMessageType::Error
+						, ArgumentName
+					);
+					return MakeError();
+				}
 
-					for (const FMVVMBlueprintPin& Pin : ConversionFunction->GetPins())
+				// Sanity check
+				{
+					// if there is a FieldId associated with the read property
+					if (CreateFieldResult.GetValue().Source)
 					{
-						if (Pin.UsedPathAsValue())
+						EMVVMBlueprintFieldPathSource PathSource = PropertyPath.GetSource(Self->WidgetBlueprintCompilerContext.WidgetBlueprint());
+						bool bValidViewModel = CreateFieldResult.GetValue().Source->Type == FCompilerBindingSource::EType::ViewModel && PathSource == EMVVMBlueprintFieldPathSource::ViewModel;
+						bool bValidWidget = CreateFieldResult.GetValue().Source->Type == FCompilerBindingSource::EType::Widget && PathSource == EMVVMBlueprintFieldPathSource::Widget;
+						bool bDynamic = CreateFieldResult.GetValue().Source->Type == FCompilerBindingSource::EType::DynamicViewmodel;
+						bool bSelf = CreateFieldResult.GetValue().Source->Type == FCompilerBindingSource::EType::Self && PathSource == EMVVMBlueprintFieldPathSource::SelfContext;
+						if (!(bValidViewModel || bValidWidget || bDynamic || bSelf))
 						{
-							if (Pin.GetName().IsNone())
-							{
-								AddMessageForBinding(Binding, BlueprintView,
-									FText::Format(LOCTEXT("InvalidArgumentPathName", "The conversion function {0} has an invalid argument"), FText::FromString(Binding.GetDisplayNameString(WidgetBlueprintCompilerContext.WidgetBlueprint()))),
-									EBindingMessageType::Error
-								);
-								Result = ECreateSourcesForConversionFunctionResult::Failed;
-							}
-							else
-							{
-								const FMVVMBlueprintPropertyPath& Path = Pin.GetPath();
-								if (CreateSourceContextForPropertyPath(Path, bForwardBinding, ComplexConversionFunctionContextIndex, Pin.GetName()))
-								{
-									Result = ECreateSourcesForConversionFunctionResult::Valid;
-								}
-								else
-								{
-									Result = ECreateSourcesForConversionFunctionResult::Failed;
-									break;
-								}
-							}
+							Self->AddMessageForBinding(Binding
+								, FText::Format(LOCTEXT("CreateNotifyFieldIdFailedInvalidInvalidContext", "Internal error. The property path '{0}' is invalid. The context is invalid."), PropertyPathToText(Class, Self->BlueprintView.Get(), PropertyPath))
+								, EMessageType::Error
+								, ArgumentName
+							);
+							return MakeError();
 						}
 					}
-
-					if (Result == ECreateSourcesForConversionFunctionResult::Continue)
-					{
-						// The bindings doesn't have a path but could be onetime with hard codded value.
-						FSimpleBindingContext& NewBindingContext = SimpleBindingContexts.AddDefaulted_GetRef();
-						NewBindingContext.BindingIndex = Index;
-						NewBindingContext.bIsForwardBinding = bForwardBinding;
-						NewBindingContext.ComplexConversionFunctionContextIndex = ComplexConversionFunctionContextIndex;
-						ComplexConversionFunctionContexts[ComplexConversionFunctionContextIndex].bNeedsValidSource = false;
-
-						Result = ECreateSourcesForConversionFunctionResult::Valid;
-					}
 				}
+
+				ReadFieldContext->NotificationField = MakeShared<FCompilerNotifyFieldId>(CreateFieldResult.GetValue());
+				ReadFieldContext->OptionalSource = ReadFieldContext->NotificationField->Source;
 			}
-			return Result;
+			return MakeValue();
 		};
 
-		auto AddWarningForPropertyWithMVVMAndLegacyBinding = [this, &Binding, &BlueprintView, Class](const FMVVMBlueprintPropertyPath& Path)
+		if (ValidBinding->Type == FCompilerBinding::EType::Assignment)
 		{
-			if (!Path.HasPaths())
-			{
-				return;
-			}
-
-			// There can't be a legacy binding in the local scope, so we can skip this if the MVVM binding refers to a property in local scope.
-			if (Path.HasFieldInLocalScope())
-			{
-				return;
-			}
-
-			TArrayView<FMVVMBlueprintFieldPath const> MVVMBindingPath = Path.GetFieldPaths();
-			TArray< FDelegateRuntimeBinding > LegacyBindings = Class->Bindings;
-			FName MVVMFieldName = Path.GetFieldNames(Class).Last();
-			FName MVVMObjectName = Path.GetWidgetName();
-
-			if (Path.GetFieldPaths().Last().GetBindingKind() == EBindingKind::Function)
+			const FMVVMBlueprintPropertyPath& BindingSourcePath = ValidBinding->Key.bIsForwardBinding ? Binding.SourcePath : Binding.DestinationPath;
+			TValueOrError<FCreateFieldsResult, void> CreateSourceResult = CreateBindingSourceContext(BindingSourcePath, FName());
+			if (CreateSourceResult.HasError())
 			{
-				return;
+				continue;
 			}
 
-			// If the first field is a UserWidget, we know this property resides in a nested UserWidget.
-			if (MVVMBindingPath[0].GetParentClass(Class) && MVVMBindingPath[0].GetParentClass(Class)->IsChildOf(UUserWidget::StaticClass()) && MVVMBindingPath.Num() > 1)
+			TSharedPtr<FGeneratedReadFieldPathContext> Found = AlreadyExist(CreateSourceResult.GetValue().SkeletalGeneratedFields);
+			if (!Found)
 			{
-				if (UWidgetBlueprintGeneratedClass* NestedBPGClass = Cast<UWidgetBlueprintGeneratedClass>(MVVMBindingPath[MVVMBindingPath.Num() - 2].GetParentClass(Class)))
-				{
-					LegacyBindings = NestedBPGClass->Bindings;
+				Found = MakeShared<FGeneratedReadFieldPathContext>();
+				Found->OptionalSource = MoveTemp(CreateSourceResult.GetValue().Source);
+				Found->GeneratedFields = MoveTemp(CreateSourceResult.GetValue().GeneratedFields);
+				Found->SkeletalGeneratedFields = MoveTemp(CreateSourceResult.GetValue().SkeletalGeneratedFields);
 
-					// We can't use Path.GetWidgetName() when we are dealing with nested UserWidgets, because it refers to the topmost UserWidget.
-					MVVMObjectName = MVVMBindingPath[MVVMBindingPath.Num() - 2].GetFieldName(Class);
-				}
-				else
+				if (CreateFieldId(BindingSourcePath, Found, FName()).HasError())
 				{
-					return;
+					bIsPreCompileStepValid = false;
+					continue;
 				}
+
+				GeneratedReadFieldPaths.Add(Found.ToSharedRef());
 			}
 
-			for (const FDelegateRuntimeBinding& LegacyBinding : LegacyBindings)
+			Found->UsedByBindings.AddUnique(ValidBinding);
+			ValidBinding->ReadPaths.Add(Found.ToSharedRef());
+		}
+		else if (ValidBinding->Type == FCompilerBinding::EType::SimpleConversionFunction
+			|| ValidBinding->Type == FCompilerBinding::EType::ComplexConversionFunction)
+		{
+			UMVVMBlueprintViewConversionFunction* ConversionFunction = Binding.Conversion.GetConversionFunction(ValidBinding->Key.bIsForwardBinding);
+			if (ensure(ConversionFunction))
 			{
-				if (LegacyBinding.ObjectName == MVVMObjectName) 
+				for (const FMVVMBlueprintPin& Pin : ConversionFunction->GetPins())
 				{
-					if (LegacyBinding.PropertyName == MVVMFieldName)
+					if (Pin.UsedPathAsValue())
 					{
-						AddMessageForBinding(Binding, BlueprintView, LOCTEXT("BindingConflictWithLegacy", "The binding is set on a property with legacy binding."), EBindingMessageType::Warning);
-						break;
-					}
-				}
-			}
-		};
-
-		// Add the forward binding. If the binding has a conversion function, use it instead of the regular binding.
-		if (IsForwardBinding(Binding.BindingType))
-		{
-			ECreateSourcesForConversionFunctionResult ConversionFunctionResult = CreateSourcesForConversionFunction(true);
-			if (ConversionFunctionResult == ECreateSourcesForConversionFunctionResult::Continue)
-			{
-				if (Binding.SourcePath.IsValid())
-				{
-					if (Binding.DestinationPath.IsValid())
-					{
-						AddWarningForPropertyWithMVVMAndLegacyBinding(Binding.DestinationPath);
-					}
+						TValueOrError<FCreateFieldsResult, void> CreateSourceResult = CreateBindingSourceContext(Pin.GetPath(), Pin.GetName());
+						if (CreateSourceResult.HasError())
+						{
+							continue;
+						}
 
-					if (!CreateSourceContextForPropertyPath(Binding.SourcePath, true, INDEX_NONE, FName()))
-					{
-						bAreBindingsValid = false;
-						continue;
+						TSharedPtr<FGeneratedReadFieldPathContext> Found = AlreadyExist(CreateSourceResult.GetValue().SkeletalGeneratedFields);
+						if (!Found)
+						{
+							Found = MakeShared<FGeneratedReadFieldPathContext>();
+							Found->OptionalSource = MoveTemp(CreateSourceResult.GetValue().Source);
+							Found->GeneratedFields = MoveTemp(CreateSourceResult.GetValue().GeneratedFields);
+							Found->SkeletalGeneratedFields = MoveTemp(CreateSourceResult.GetValue().SkeletalGeneratedFields);
+
+							if (CreateFieldId(Pin.GetPath(), Found, Pin.GetName()).HasError())
+							{
+								bIsPreCompileStepValid = false;
+								continue;
+							}
+
+							GeneratedReadFieldPaths.Add(Found.ToSharedRef());
+						}
+
+						Found->UsedByBindings.AddUnique(TWeakPtr<FCompilerBinding>(ValidBinding));
+						ValidBinding->ReadPaths.Add(Found.ToSharedRef());
 					}
 				}
-				else
-				{
-					AddMessageForBinding(Binding, BlueprintView, LOCTEXT("BindingEmptySourcePath", "The binding doesn't have a Source or a Conversion function."), EBindingMessageType::Error);
-					bAreBindingsValid = false;
-					continue;
-				}
-			}
-			else if (ConversionFunctionResult == ECreateSourcesForConversionFunctionResult::Failed)
-			{
-				bAreBindingsValid = false;
-				continue;
 			}
 		}
+	}
 
-		// Add the backward binding. If the binding has a conversion function, use it instead of the regular binding.
-		if (IsBackwardBinding(Binding.BindingType))
+	//The pins for event
+	for (TSharedRef<FCompilerEvent>& ValidEvent : ValidEvents)
+	{
+		UMVVMBlueprintViewEvent* EventPtr = ValidEvent->Event.Get();
+		check(EventPtr);
+		if (ValidEvent->Type == FCompilerEvent::EType::Valid)
 		{
-			ECreateSourcesForConversionFunctionResult ConversionFunctionResult = CreateSourcesForConversionFunction(false);
-			if (ConversionFunctionResult == ECreateSourcesForConversionFunctionResult::Continue)
+			for (const FMVVMBlueprintPin& Pin : EventPtr->GetPins())
 			{
-				if (Binding.DestinationPath.IsValid())
+				if (Pin.UsedPathAsValue())
 				{
-					if (Binding.SourcePath.IsValid())
+					TValueOrError<FCreateFieldsResult, FText> CreateSourceResult = CreateFieldContext(Class, Pin.GetPath(), true);
+					if (CreateSourceResult.HasError())
 					{
-						AddWarningForPropertyWithMVVMAndLegacyBinding(Binding.SourcePath);
+						AddMessageForEvent(ValidEvent
+							, FText::Format(Private::PropertyPathIsInvalidFormat, PropertyPathToText(Class, BlueprintView.Get(), Pin.GetPath()))
+							, EMessageType::Error
+							, Pin.GetName()
+						);
+						bIsPreCompileStepValid = false;
 					}
-					if (!CreateSourceContextForPropertyPath(Binding.DestinationPath, false, INDEX_NONE, FName()))
+
+					TSharedPtr<FGeneratedReadFieldPathContext> Found = AlreadyExist(CreateSourceResult.GetValue().SkeletalGeneratedFields);
+					if (!Found)
 					{
-						bAreBindingsValid = false;
-						continue;
+						Found = MakeShared<FGeneratedReadFieldPathContext>();
+						Found->OptionalSource = MoveTemp(CreateSourceResult.GetValue().Source);
+						Found->GeneratedFields = MoveTemp(CreateSourceResult.GetValue().GeneratedFields);
+						Found->SkeletalGeneratedFields = MoveTemp(CreateSourceResult.GetValue().SkeletalGeneratedFields);
+						GeneratedReadFieldPaths.Add(Found.ToSharedRef());
 					}
-				}
-				else
-				{
-					AddMessageForBinding(Binding, BlueprintView, LOCTEXT("BindingEmptyDestinationPath", "The binding doesn't have a Destination or a Conversion function."), EBindingMessageType::Error);
-					bAreBindingsValid = false;
-					continue;
+
+					Found->UsedByEvents.AddUnique(ValidEvent);
+					ValidEvent->ReadPaths.Add(Found.ToSharedRef());
 				}
 			}
-			else if (ConversionFunctionResult == ECreateSourcesForConversionFunctionResult::Failed)
+		}
+	}
+}
+
+
+void FMVVMViewBlueprintCompiler::CreateCreatorContentFromBindingSource(UWidgetBlueprintGeneratedClass* Class)
+{
+	// Add all the needed sources that are not viewmodel (so not in the SourceCreators)
+	for (const TSharedRef<FCompilerBindingSource>& Source : NeededBindingSources)
+	{
+		{
+			const bool bSourceIsViewModel = (Source->Type == FCompilerBindingSource::EType::ViewModel || Source->Type == FCompilerBindingSource::EType::DynamicViewmodel);
+			const FCompilerViewModelCreatorContext* FoundSourceCreator = ViewModelCreatorContexts.FindByPredicate([Source](const FCompilerViewModelCreatorContext& Other) { return Other.Source == Source; });
+			if (bSourceIsViewModel && FoundSourceCreator == nullptr)
 			{
-				bAreBindingsValid = false;
-				continue;
+				WidgetBlueprintCompilerContext.MessageLog.Error(*LOCTEXT("ViewmodelSourceNotAdded", "Internal error. A viewmodel was not added to the compiled list.").ToString());
+				bIsPreCompileStepValid = false;
 			}
 		}
-	}
 
-	return bAreBindingsValid;
+		const bool bSourceIsWidget = Source->Type == FCompilerBindingSource::EType::Widget || Source->Type == FCompilerBindingSource::EType::Self;
+		if (bSourceIsWidget)
+		{
+			FCompilerWidgetCreatorContext& CompiledSourceCreator = WidgetCreatorContexts.AddDefaulted_GetRef();
+			CompiledSourceCreator.Source = Source;
+			CompiledSourceCreator.bSelfReference = Source->Type == FCompilerBindingSource::EType::Self;
+		}
+	}
 }
 
 
-bool FMVVMViewBlueprintCompiler::PreCompileSourceCreators(UWidgetBlueprintGeneratedClass* Class, UMVVMBlueprintView* BlueprintView)
+void FMVVMViewBlueprintCompiler::FixCompilerBindingSelfSource(UWidgetBlueprintGeneratedClass* Class)
 {
-	if (!bAreSourcesCreatorValid)
+	int32 Counter = 0;
+	for (TSharedRef<FCompilerBindingSource>& Source : NeededBindingSources)
 	{
-		return false;
+		if (Source->Type == FCompilerBindingSource::EType::Self)
+		{
+			ensure(Source->AuthoritativeClass == nullptr);
+			Source->AuthoritativeClass = Class;
+			++Counter;
+		}
 	}
 
-	for (FCompilerSourceCreatorContext& SourceCreatorContext : CompilerSourceCreatorContexts)
+	if (Counter > 1)
 	{
-		FMVVMViewClass_SourceCreator CompiledSourceCreator;
+		AddMessage(LOCTEXT("MoreThanSeldContext", "Internal error. There is more than self context.")
+			, EMessageType::Warning
+		);
+	}
+}
+
 
-		if (SourceCreatorContext.Type == ECompilerSourceCreatorType::ViewModel)
+void FMVVMViewBlueprintCompiler::AddWarningForPropertyWithMVVMAndLegacyBinding(UWidgetBlueprintGeneratedClass* Class)
+{
+	const TArray<FDelegateRuntimeBinding>& LegacyBindings = Class->Bindings;
+	for (const TSharedRef<FGeneratedWriteFieldPathContext>& WriteFieldPath : GeneratedWriteFieldPaths)
+	{
+		FName MVVMObjectName;
+		FName MVVMFieldName;
+		if (WriteFieldPath->GeneratedFrom == EMVVMBlueprintFieldPathSource::SelfContext)
 		{
-			const FMVVMBlueprintViewModelContext& ViewModelContext = SourceCreatorContext.ViewModelContext;
-			checkf(ViewModelContext.GetViewModelClass(), TEXT("The viewmodel class is invalid. It was checked in CreateSourceList"));
+			MVVMObjectName = Class->ClassGeneratedBy ? Class->ClassGeneratedBy->GetFName() : FName();
+			MVVMFieldName = WriteFieldPath->SkeletalGeneratedFields.Num() > 0 ? WriteFieldPath->SkeletalGeneratedFields[0].GetName() : FName();
+		}
 
-			if (ViewModelContext.GetViewModelClass()->HasAllClassFlags(CLASS_Deprecated))
+		for (int32 Index = 0; Index < WriteFieldPath->SkeletalGeneratedFields.Num(); ++Index)
+		{
+			const UE::MVVM::FMVVMConstFieldVariant& Field = WriteFieldPath->SkeletalGeneratedFields[Index];
+			const FObjectPropertyBase* ObjectProperty = Field.IsProperty() ? CastField<const FObjectPropertyBase>(Field.GetProperty()) : nullptr;
+			if (ObjectProperty && ObjectProperty->PropertyClass->IsChildOf(UWidget::StaticClass()))
 			{
-				AddErrorForViewModel(ViewModelContext, FText::Format(LOCTEXT("ViewModelTypeDeprecated", "Viewmodel class '{0}' is deprecated and should not be used. Please update it in the View Models panel."),
-					ViewModelContext.GetViewModelClass()->GetDisplayNameText()
-				));
+				MVVMObjectName = ObjectProperty->GetFName();
+				MVVMFieldName = WriteFieldPath->SkeletalGeneratedFields.IsValidIndex(Index+1) ? WriteFieldPath->SkeletalGeneratedFields[Index+1].GetName() : FName();
 			}
+		}
 
-			if (!GetAllowedContextCreationType(ViewModelContext.GetViewModelClass()).Contains(ViewModelContext.CreationType))
+		for (const FDelegateRuntimeBinding& LegacyBinding : LegacyBindings)
+		{
+			if (LegacyBinding.ObjectName == MVVMObjectName && LegacyBinding.PropertyName == MVVMFieldName)
 			{
-				AddErrorForViewModel(ViewModelContext, FText::Format(LOCTEXT("ViewModelContextCreationTypeInvalid", "Viewmodel '{0}' has an invalidate creation type. You can change it in the View Models panel."),
-					ViewModelContext.GetViewModelClass()->GetDisplayNameText()
-				));
-				bAreSourcesCreatorValid = false;
-				continue;
+				if (WriteFieldPath->UsedByBindings.Num())
+				{
+					AddMessages(WriteFieldPath->UsedByBindings
+						, TArrayView<TWeakPtr<FCompilerEvent>>()
+						, LOCTEXT("BindingConflictWithLegacy", "The binding is set on a property with legacy binding.")
+						, EMessageType::Warning
+					);
+				}
 			}
+		}
+	}
+}
 
-			if (ViewModelContext.CreationType == EMVVMBlueprintViewModelContextCreationType::Manual)
+
+void FMVVMViewBlueprintCompiler::PreCompileViewModelCreatorContexts(UWidgetBlueprintGeneratedClass* Class)
+{
+	for (FCompilerViewModelCreatorContext& SourceCreatorContext : ViewModelCreatorContexts)
+	{
+		const FMVVMBlueprintViewModelContext& ViewModelContext = SourceCreatorContext.ViewModelContext;
+		checkf(ViewModelContext.GetViewModelClass(), TEXT("The viewmodel class is invalid. It was checked in CreateSourceList"));
+
+		if (ViewModelContext.GetViewModelClass()->HasAllClassFlags(CLASS_Deprecated))
+		{
+			AddMessageForViewModel(ViewModelContext
+				, FText::Format(LOCTEXT("ViewModelTypeDeprecated", "Viewmodel class '{0}' is deprecated and should not be used. Please update it in the View Models panel."), ViewModelContext.GetViewModelClass()->GetDisplayNameText())
+				, EMessageType::Warning
+			);
+		}
+
+		if (!ViewModelContext.GetViewModelClass()->ImplementsInterface(UNotifyFieldValueChanged::StaticClass()))
+		{
+			AddMessageForViewModel(ViewModelContext
+				, LOCTEXT("ViewmodelInvalidInterface", "The class doesn't implement the interface NotifyFieldValueChanged.")
+				, EMessageType::Error
+			);
+			bIsPreCompileStepValid = false;
+			continue;
+		}
+
+		if (!GetAllowedContextCreationType(ViewModelContext.GetViewModelClass()).Contains(ViewModelContext.CreationType))
+		{
+			AddMessageForViewModel(ViewModelContext
+				, LOCTEXT("ViewModelContextCreationTypeInvalid", "It has an invalidate creation type. You can change it in the View Models panel.")
+				, EMessageType::Error
+			);
+			bIsPreCompileStepValid = false;
+			continue;
+		}
+
+		if (ViewModelContext.CreationType == EMVVMBlueprintViewModelContextCreationType::Manual)
+		{
+		}
+		else if (ViewModelContext.CreationType == EMVVMBlueprintViewModelContextCreationType::CreateInstance)
+		{
+			if (ViewModelContext.GetViewModelClass()->HasAllClassFlags(CLASS_Abstract))
 			{
+				AddMessageForViewModel(ViewModelContext
+					, FText::Format(LOCTEXT("ViewModelTypeAbstract", "Viewmodel class '{0}' is abstract and can't be created. You can change it in the View Models panel."), ViewModelContext.GetViewModelClass()->GetDisplayNameText())
+					, EMessageType::Error
+				);
+				bIsPreCompileStepValid = false;
+				continue;
 			}
-			else if (ViewModelContext.CreationType == EMVVMBlueprintViewModelContextCreationType::CreateInstance)
+		}
+		else if (ViewModelContext.CreationType == EMVVMBlueprintViewModelContextCreationType::PropertyPath)
+		{
+			if (ViewModelContext.ViewModelPropertyPath.IsEmpty())
 			{
-				if (ViewModelContext.GetViewModelClass()->HasAllClassFlags(CLASS_Abstract))
-				{
-					AddErrorForViewModel(ViewModelContext, FText::Format(LOCTEXT("ViewModelTypeAbstract", "Viewmodel class '{0}' is abstract and can't be created. You can change it in the View Models panel."),
-						ViewModelContext.GetViewModelClass()->GetDisplayNameText()
-					));
-					bAreSourcesCreatorValid = false;
-					continue;
-				}
+				AddMessageForViewModel(ViewModelContext
+					, LOCTEXT("ViewModelInvalidGetter", "Viewmodel has an invalid Getter. You can select a new one in the View Models panel.")
+					, EMessageType::Error
+				);
+				bIsPreCompileStepValid = true;
+				continue;
 			}
-			else if (ViewModelContext.CreationType == EMVVMBlueprintViewModelContextCreationType::PropertyPath)
-			{
-				if (ViewModelContext.ViewModelPropertyPath.IsEmpty())
-				{
-					AddErrorForViewModel(ViewModelContext, LOCTEXT("ViewModelInvalidGetter", "Viewmodel has an invalid Getter. You can select a new one in the View Models panel."));
-					bAreSourcesCreatorValid = true;
-					continue;
-				}
-
-				TValueOrError<FCompiledBindingLibraryCompiler::FFieldPathHandle, FText> ReadFieldPathResult = AddObjectFieldPath(BindingLibraryCompiler, Class, ViewModelContext.ViewModelPropertyPath, ViewModelContext.GetViewModelClass());
-				if (ReadFieldPathResult.HasError())
-				{
-					AddErrorForViewModel(ViewModelContext, ReadFieldPathResult.GetError());
-					bAreSourcesCreatorValid = false;
-					continue;
-				}
 
-				SourceCreatorContext.ReadPropertyPath = ReadFieldPathResult.StealValue();
-			}
-			else if (ViewModelContext.CreationType == EMVVMBlueprintViewModelContextCreationType::GlobalViewModelCollection)
+			TValueOrError<FCompiledBindingLibraryCompiler::FFieldPathHandle, FText> ReadFieldPathResult = AddObjectFieldPath(BindingLibraryCompiler, Class, ViewModelContext.ViewModelPropertyPath, ViewModelContext.GetViewModelClass());
+			if (ReadFieldPathResult.HasError())
 			{
-				if (ViewModelContext.GlobalViewModelIdentifier.IsNone())
-				{
-					AddErrorForViewModel(ViewModelContext, LOCTEXT("ViewmodelInvalidGlobalIdentifier", "Viewmodel doesn't have a valid Global identifier. You can specify a new one in the Viewmodels panel."));
-					bAreSourcesCreatorValid = false;
-					continue;
-				}
+				AddMessageForViewModel(ViewModelContext
+					, ReadFieldPathResult.GetError()
+					, EMessageType::Error
+				);
+				bIsPreCompileStepValid = false;
+				continue;
 			}
-			else if (ViewModelContext.CreationType == EMVVMBlueprintViewModelContextCreationType::Resolver)
+
+			SourceCreatorContext.ReadPropertyPathHandle = ReadFieldPathResult.StealValue();
+		}
+		else if (ViewModelContext.CreationType == EMVVMBlueprintViewModelContextCreationType::GlobalViewModelCollection)
+		{
+			if (ViewModelContext.GlobalViewModelIdentifier.IsNone())
 			{
-				if (!ViewModelContext.Resolver)
-				{
-					AddErrorForViewModel(ViewModelContext, LOCTEXT("ViewmodelInvalidResolver", "Viewmodel doesn't have a valid Resolver. You can specify a new one in the Viewmodels panel."));
-					bAreSourcesCreatorValid = false;
-					continue;
-				}
+				AddMessageForViewModel(ViewModelContext
+					, LOCTEXT("ViewmodelInvalidGlobalIdentifier", "Viewmodel doesn't have a valid Global identifier. You can specify a new one in the Viewmodels panel.")
+					, EMessageType::Error
+				);
+				bIsPreCompileStepValid = false;
+				continue;
 			}
-			else
+		}
+		else if (ViewModelContext.CreationType == EMVVMBlueprintViewModelContextCreationType::Resolver)
+		{
+			if (!ViewModelContext.Resolver)
 			{
-				AddErrorForViewModel(ViewModelContext, LOCTEXT("ViewmodelInvalidCreationType", "Viewmodel doesn't have a valid creation type. You can select one in the Viewmodels panel."));
-				bAreSourcesCreatorValid = false;
+				AddMessageForViewModel(ViewModelContext
+					, LOCTEXT("ViewmodelInvalidResolver", "Viewmodel doesn't have a valid Resolver. You can specify a new one in the Viewmodels panel.")
+					, EMessageType::Error
+				);
+				bIsPreCompileStepValid = false;
 				continue;
 			}
 		}
+		else
+		{
+			AddMessageForViewModel(ViewModelContext
+				, LOCTEXT("ViewmodelInvalidCreationType", "Viewmodel doesn't have a valid creation type. You can select one in the Viewmodels panel.")
+				, EMessageType::Error
+			);
+			bIsPreCompileStepValid = false;
+			continue;
+		}
 	}
 
-	return bAreSourcesCreatorValid;
+	for (FCompilerWidgetCreatorContext& WidgetCreator : WidgetCreatorContexts)
+	{
+		if (!WidgetCreator.Source->AuthoritativeClass->ImplementsInterface(UNotifyFieldValueChanged::StaticClass()))
+		{
+			AddMessage(FText::Format(LOCTEXT("WidgetInvalidInterface", "The widget {0} class doesn't implement the interface NotifyFieldValueChanged."), FText::FromName(WidgetCreator.Source->Name))
+				, EMessageType::Error
+			);
+			bIsPreCompileStepValid = false;
+			continue;
+		}
+
+		if (!WidgetCreator.bSelfReference)
+		{
+			TValueOrError<FCompiledBindingLibraryCompiler::FFieldPathHandle, FText> ReadFieldPathResult = AddObjectFieldPath(BindingLibraryCompiler, Class, WidgetCreator.Source->Name.ToString(), WidgetCreator.Source->AuthoritativeClass);
+			if (ReadFieldPathResult.HasError())
+			{
+				AddMessage(FText::Format(LOCTEXT("WidgetAddObjectFieldPathFailFormat", "The widget {0} creator failed. {1}")
+					, FText::FromName(WidgetCreator.Source->Name)
+					, ReadFieldPathResult.GetError()
+				)
+					, EMessageType::Error
+				);
+				bIsPreCompileStepValid = false;
+				continue;
+			}
+			WidgetCreator.ReadPropertyPathHandle = ReadFieldPathResult.StealValue();
+		}
+	}
 }
 
 
-bool FMVVMViewBlueprintCompiler::CompileSourceCreators(const FCompiledBindingLibraryCompiler::FCompileResult& CompileResult, UWidgetBlueprintGeneratedClass* Class, UMVVMBlueprintView* BlueprintView, UMVVMViewClass* ViewExtension)
+void FMVVMViewBlueprintCompiler::CompileViewModelCreatorContexts(const FCompiledBindingLibraryCompiler::FCompileResult& CompileResult, UWidgetBlueprintGeneratedClass* Class, UMVVMViewClass* ViewExtension)
 {
-	if (!bAreSourcesCreatorValid)
-	{
-		return false;
-	}
-
 	TArray<FMVVMViewClass_SourceCreator> UnsortedSourceCreators;
-	for (const FCompilerSourceCreatorContext& SourceCreatorContext : CompilerSourceCreatorContexts)
+	for (FCompilerViewModelCreatorContext& SourceCreatorContext : ViewModelCreatorContexts)
 	{
 		const FMVVMBlueprintViewModelContext& ViewModelContext = SourceCreatorContext.ViewModelContext;
 		FMVVMViewClass_SourceCreator CompiledSourceCreator;
@@ -1399,103 +1756,132 @@ bool FMVVMViewBlueprintCompiler::CompileSourceCreators(const FCompiledBindingLib
 		CompiledSourceCreator.ExpectedSourceType = ViewModelContext.GetViewModelClass();
 		CompiledSourceCreator.PropertyName = ViewModelContext.GetViewModelName();
 
-		bool bCanBeSet = false;
-		bool CanBeEvaluated = false;
-		bool bIsOptional = false;
+		bool bCanBeSet = ViewModelContext.bCreateSetterFunction;
+		bool bCanBeEvaluated = SourceCreatorContext.DynamicContext.IsValid();
+		bool bIsOptional = SourceCreatorContext.Source->bIsOptional;
 		bool bCreateInstance = false;
-		bool IsUserWidgetProperty = false;
+		bool bIsUserWidgetProperty = !SourceCreatorContext.DynamicContext.IsValid();
 
-		if (SourceCreatorContext.Type == ECompilerSourceCreatorType::ViewModel)
+		if (ViewModelContext.CreationType == EMVVMBlueprintViewModelContextCreationType::Manual)
 		{
-			IsUserWidgetProperty = true;
-			bCanBeSet = ViewModelContext.bCreateSetterFunction;
-
-			if (ViewModelContext.CreationType == EMVVMBlueprintViewModelContextCreationType::Manual)
+			bCanBeSet = true;
+			ensure(bIsOptional == true);
+		}
+		else if (ViewModelContext.CreationType == EMVVMBlueprintViewModelContextCreationType::CreateInstance)
+		{
+			bCreateInstance = true;
+		}
+		else if (ViewModelContext.CreationType == EMVVMBlueprintViewModelContextCreationType::PropertyPath)
+		{
+			const FMVVMVCompiledFieldPath* CompiledFieldPath = CompileResult.FieldPaths.Find(SourceCreatorContext.ReadPropertyPathHandle);
+			if (CompiledFieldPath == nullptr)
 			{
-				bCanBeSet = true;
-				bIsOptional = true;
+				AddMessageForViewModel(ViewModelContext
+					, LOCTEXT("ViewModelInvalidInitializationBindingNotGenerated", "The viewmodel initialization binding was not generated.")
+					, EMessageType::Error
+				);
+				bIsCompileStepValid = false;
+				continue;
 			}
-			else if (ViewModelContext.CreationType == EMVVMBlueprintViewModelContextCreationType::CreateInstance)
+
+			CompiledSourceCreator.FieldPath = *CompiledFieldPath;
+			ensure(bIsOptional == ViewModelContext.bOptional);
+		}
+		else if (ViewModelContext.CreationType == EMVVMBlueprintViewModelContextCreationType::GlobalViewModelCollection)
+		{
+			ensure(!ViewModelContext.GlobalViewModelIdentifier.IsNone());
+
+			FMVVMViewModelContext GlobalViewModelInstance;
+			GlobalViewModelInstance.ContextClass = ViewModelContext.GetViewModelClass();
+			GlobalViewModelInstance.ContextName = ViewModelContext.GlobalViewModelIdentifier;
+			if (!GlobalViewModelInstance.IsValid())
 			{
-				bCreateInstance = true;
+				AddMessageForViewModel(ViewModelContext
+					, LOCTEXT("ViewmodelGlobalContextIdentifier", "The context for viewmodel could not be created. Change the identifier.")
+					, EMessageType::Warning
+				);
 			}
-			else if (ViewModelContext.CreationType == EMVVMBlueprintViewModelContextCreationType::PropertyPath)
-			{
-				const FMVVMVCompiledFieldPath* CompiledFieldPath = CompileResult.FieldPaths.Find(SourceCreatorContext.ReadPropertyPath);
-				if (CompiledFieldPath == nullptr)
-				{
-					AddErrorForViewModel(ViewModelContext, LOCTEXT("ViewModelInvalidInitializationBindingNotGenerated", "The viewmodel initialization binding was not generated."));
-					bAreSourcesCreatorValid = false;
-					continue;
-				}
 
-				CompiledSourceCreator.FieldPath = *CompiledFieldPath;
-				bIsOptional = ViewModelContext.bOptional;
-			}
-			else if (ViewModelContext.CreationType == EMVVMBlueprintViewModelContextCreationType::GlobalViewModelCollection)
+			CompiledSourceCreator.GlobalViewModelInstance = MoveTemp(GlobalViewModelInstance);
+			ensure(bIsOptional == ViewModelContext.bOptional);
+		}
+		else if (ViewModelContext.CreationType == EMVVMBlueprintViewModelContextCreationType::Resolver)
+		{
+			UMVVMViewModelContextResolver* Resolver = DuplicateObject(ViewModelContext.Resolver.Get(), ViewExtension);
+			if (!Resolver)
 			{
-				if (ViewModelContext.GlobalViewModelIdentifier.IsNone())
-				{
-					AddErrorForViewModel(ViewModelContext, LOCTEXT("ViewModelInvalidGlobalIdentifier", "The viewmodel doesn't have a valid Global identifier. You can specify a new one in the Viewmodels panel."));
-					bAreSourcesCreatorValid = false;
-					continue;
-				}
+				AddMessageForViewModel(ViewModelContext
+					, LOCTEXT("ViewmodelFailedResolverDuplicate", "Internal error. The resolver could not be dupliated.")
+					, EMessageType::Error
+				);
+				bIsCompileStepValid = false;
+				continue;
+			}
 
-				FMVVMViewModelContext GlobalViewModelInstance;
-				GlobalViewModelInstance.ContextClass = ViewModelContext.GetViewModelClass();
-				GlobalViewModelInstance.ContextName = ViewModelContext.GlobalViewModelIdentifier;
-				if (!GlobalViewModelInstance.IsValid())
-				{
-					AddErrorForViewModel(ViewModelContext, LOCTEXT("ViewModelCouldNotBeCreated", "The context for viewmodel could not be created. You can change the viewmodel in the Viewmodels panel."));
-					bAreSourcesCreatorValid = false;
-					continue;
-				}
+			CompiledSourceCreator.Resolver = Resolver;
+			ensure(bIsOptional == ViewModelContext.bOptional);
+		}
+		else
+		{
+			AddMessageForViewModel(ViewModelContext
+				, LOCTEXT("ViewModelWithoutValidCreationType", "The viewmodel doesn't have a valid creation type.")
+				, EMessageType::Error
+			);
+			bIsCompileStepValid = false;
+			continue;
+		}
 
-				CompiledSourceCreator.GlobalViewModelInstance = MoveTemp(GlobalViewModelInstance);
-				bIsOptional = ViewModelContext.bOptional;
-			}
-			else if (ViewModelContext.CreationType == EMVVMBlueprintViewModelContextCreationType::Resolver)
+		if (SourceCreatorContext.DynamicContext)
+		{
+			if (SourceCreatorContext.DynamicContext->ParentSource)
 			{
-				UMVVMViewModelContextResolver* Resolver = DuplicateObject(ViewModelContext.Resolver.Get(), ViewExtension);
-				if (!Resolver)
-				{
-					AddErrorForViewModel(ViewModelContext, LOCTEXT("ViewmodelFailedResolverDuplicate", "Internal error. The resolver could not be dupliated."));
-					bAreSourcesCreatorValid = false;
-					continue;
-				}
-
-				CompiledSourceCreator.Resolver = Resolver;
-				bIsOptional = ViewModelContext.bOptional;
+				CompiledSourceCreator.ParentSourceName = SourceCreatorContext.DynamicContext->ParentSource->Name;
 			}
 			else
 			{
-				AddErrorForViewModel(ViewModelContext, LOCTEXT("ViewModelWithoutValidCreationType", "The viewmodel doesn't have a valid creation type."));
-				bAreSourcesCreatorValid = false;
+				AddMessageForViewModel(ViewModelContext
+					, LOCTEXT("ViewModelWithoutValidDyanmicParentSource", "The viewmodel doesn't have a valid parent source.")
+					, EMessageType::Error
+				);
+				bIsCompileStepValid = false;
 				continue;
 			}
 		}
-		else if (SourceCreatorContext.Type == ECompilerSourceCreatorType::ViewModelDynamic)
+
+		CompiledSourceCreator.Flags = 0;
+		CompiledSourceCreator.Flags |= bCreateInstance ? (uint8)FMVVMViewClass_SourceCreator::ESourceFlags::TypeCreateInstance : 0;
+		CompiledSourceCreator.Flags |= bIsUserWidgetProperty ? (uint8)FMVVMViewClass_SourceCreator::ESourceFlags::IsUserWidgetProperty : 0;
+		CompiledSourceCreator.Flags |= bIsOptional ? (uint8)FMVVMViewClass_SourceCreator::ESourceFlags::IsOptional : 0;
+		CompiledSourceCreator.Flags |= bCanBeSet ? (uint8)FMVVMViewClass_SourceCreator::ESourceFlags::CanBeSet : 0;
+		CompiledSourceCreator.Flags |= bCanBeEvaluated ? (uint8)FMVVMViewClass_SourceCreator::ESourceFlags::CanBeEvaluated : 0;
+
+		UnsortedSourceCreators.Add(MoveTemp(CompiledSourceCreator));
+	}
+
+	// The other sources needed by the view
+	for (FCompilerWidgetCreatorContext& WidgetCreator : WidgetCreatorContexts)
+	{
+		FMVVMViewClass_SourceCreator CompiledSourceCreator;
+		ensure(WidgetCreator.Source->AuthoritativeClass&& WidgetCreator.Source->AuthoritativeClass->ImplementsInterface(UNotifyFieldValueChanged::StaticClass()));
+		CompiledSourceCreator.ExpectedSourceType = const_cast<UClass*>(WidgetCreator.Source->AuthoritativeClass);
+		CompiledSourceCreator.PropertyName = WidgetCreator.Source->Name;
+		CompiledSourceCreator.Flags = 0;
+		CompiledSourceCreator.Flags |= WidgetCreator.bSelfReference ? (uint8)FMVVMViewClass_SourceCreator::ESourceFlags::SelfReference : 0;
+
+		if (!WidgetCreator.bSelfReference)
 		{
-			const FMVVMVCompiledFieldPath* CompiledFieldPath = CompileResult.FieldPaths.Find(SourceCreatorContext.ReadPropertyPath);
+			const FMVVMVCompiledFieldPath* CompiledFieldPath = CompileResult.FieldPaths.Find(WidgetCreator.ReadPropertyPathHandle);
 			if (CompiledFieldPath == nullptr)
 			{
-				AddErrorForViewModel(ViewModelContext, LOCTEXT("ViewModelInvalidInitializationBindingNotGenerated", "The viewmodel initialization binding was not generated."));
-				bAreSourcesCreatorValid = false;
+				AddMessage(FText::Format(LOCTEXT("WidgetInvalidInitializationBindingNotGenerated", "The widget {0} initialization binding was not generated."), FText::FromName(WidgetCreator.Source->Name))
+					, EMessageType::Error
+				);
+				bIsCompileStepValid = false;
 				continue;
 			}
-
-			CanBeEvaluated = true;
 			CompiledSourceCreator.FieldPath = *CompiledFieldPath;
-			CompiledSourceCreator.ParentSourceName = SourceCreatorContext.DynamicParentSourceName;
 		}
 
-		CompiledSourceCreator.Flags = 0;
-		CompiledSourceCreator.Flags |= bCreateInstance ? (uint8)FMVVMViewClass_SourceCreator::ESourceFlags::TypeCreateInstance : 0;
-		CompiledSourceCreator.Flags |= IsUserWidgetProperty ? (uint8)FMVVMViewClass_SourceCreator::ESourceFlags::IsUserWidgetProperty : 0;
-		CompiledSourceCreator.Flags |= bIsOptional ? (uint8)FMVVMViewClass_SourceCreator::ESourceFlags::IsOptional : 0;
-		CompiledSourceCreator.Flags |= bCanBeSet ? (uint8)FMVVMViewClass_SourceCreator::ESourceFlags::CanBeSet : 0;
-		CompiledSourceCreator.Flags |= CanBeEvaluated ? (uint8)FMVVMViewClass_SourceCreator::ESourceFlags::CanBeEvaluated : 0;
-
 		UnsortedSourceCreators.Add(MoveTemp(CompiledSourceCreator));
 	}
 
@@ -1552,629 +1938,522 @@ bool FMVVMViewBlueprintCompiler::CompileSourceCreators(const FCompiledBindingLib
 			});
 	}
 
-	// Add sorted array to the ViewExtension
-	{
-		ViewExtension->SourceCreators.Reset(UnsortedSourceCreators.Num());
-		for (FMVVMViewClass_SourceCreator& SourceCreator : UnsortedSourceCreators)
-		{
-			ViewExtension->SourceCreators.Add(MoveTemp(SourceCreator));
-		}
-		UnsortedSourceCreators.Reset();
-	}
-
-	return bAreSourcesCreatorValid;
+	// Add the sorted viewmode array to the ViewExtension
+	ViewExtension->SourceCreators = MoveTemp(UnsortedSourceCreators);
 }
 
 
-bool FMVVMViewBlueprintCompiler::PreCompileBindings(UWidgetBlueprintGeneratedClass* Class, UMVVMBlueprintView* BlueprintView)
+void FMVVMViewBlueprintCompiler::PreCompileBindings(UWidgetBlueprintGeneratedClass* Class)
 {
-	if (!bAreSourceContextsValid || !bAreBindingsValid)
-	{
-		return false;
-	}
-
-	auto GetConversionFunction = [Self = this, Class, BlueprintView](FMVVMBlueprintViewBinding& Binding, bool bIsForwardBinding) -> TValueOrError<const UFunction*, void>
-	{
-		const UFunction* ConversionFunction = nullptr;
-		const UMVVMBlueprintViewConversionFunction* ViewConversionFunction = Binding.Conversion.GetConversionFunction(bIsForwardBinding);
-		if (ViewConversionFunction)
-		{
-			ConversionFunction = ViewConversionFunction->GetCompiledFunction(Class);
-			if (ConversionFunction == nullptr)
-			{
-				Self->AddMessageForBinding(Binding, BlueprintView, FText::Format(LOCTEXT("ConversionFunctionNotFound", "The conversion function '{0}' could not be found."),
-					FText::FromName(ViewConversionFunction->GetCompiledFunctionName())),
-					EBindingMessageType::Error
-				);
-				Self->bAreBindingsValid = false;
-				return MakeError();
-			}
-
-			TVariant<const UFunction*, TSubclassOf<UK2Node>> FunctionOrWrapperFunction = ViewConversionFunction->GetConversionFunction(Class);
-			if (FunctionOrWrapperFunction.IsType<const UFunction*>())
-			{
-				if (!GetDefault<UMVVMDeveloperProjectSettings>()->IsConversionFunctionAllowed(Self->WidgetBlueprintCompilerContext.WidgetBlueprint(), FunctionOrWrapperFunction.Get<const UFunction*>()))
-				{
-					Self->AddMessageForBinding(Binding, BlueprintView, FText::Format(LOCTEXT("ConversionFunctionNotAllow", "The conversion function {0} is not allowed."),
-						FText::FromName(FunctionOrWrapperFunction.Get<const UFunction*>()->GetFName())),
-						EBindingMessageType::Error
-					);
-					Self->bAreBindingsValid = false;
-					return MakeError();
-				}
-			}
-			else
-			{
-				Self->AddMessageForBinding(Binding, BlueprintView, LOCTEXT("ConversionFunctionNodeNotAllow", "The conversion function node is not allowed."), EBindingMessageType::Error);
-				Self->bAreBindingsValid = false;
-				return MakeError();
-			}
-		}
-		return MakeValue(ConversionFunction);
-	};
-
-	auto TestExecutionMode = [Self = this, BlueprintView](FMVVMBlueprintViewBinding& Binding)
+	auto TestExecutionMode = [Self = this](const FMVVMBlueprintViewBinding& Binding) -> bool
 	{
 		if (Binding.bOverrideExecutionMode)
 		{
 			if (!GetDefault<UMVVMDeveloperProjectSettings>()->IsExecutionModeAllowed(Binding.OverrideExecutionMode))
 			{
-				Self->AddMessageForBinding(Binding, BlueprintView, LOCTEXT("NotAllowedExecutionMode", "The binding has a restricted execution mode."), EBindingMessageType::Error);
+				Self->AddMessageForBinding(Binding, LOCTEXT("NotAllowedExecutionMode", "The binding has a restricted execution mode."), EMessageType::Error, FName());
 				return false;
 			}
 		}
 		return true;
 	};
 
-	auto GetSetterPath = [Self = this, Class, BlueprintView](FMVVMBlueprintViewBinding& Binding, int32 BindingIndex, bool bIsForwardBinding) -> TValueOrError<TArray<UE::MVVM::FMVVMConstFieldVariant>, void>
+	auto AddFieldIds = [Self = this, Class](FCompilerBinding& ValidBinding, const FMVVMBlueprintViewBinding& Binding) -> bool
 	{
-		const FBindingDestinationContext* DestinationContext = Self->BindingDestinationContexts.FindByPredicate([BindingIndex, bIsForwardBinding](const FBindingDestinationContext& Other)
+		// Add FieldId
+		bool bHasFieldId = false;
+		if (!ValidBinding.bIsOneTimeBinding)
+		{
+			for (const TSharedPtr<FGeneratedReadFieldPathContext>& ReadPath : ValidBinding.ReadPaths)
 			{
-				return Other.BindingIndex == BindingIndex && Other.bIsForwardBinding == bIsForwardBinding;
-			});
+				if (ReadPath->NotificationField)
+				{
+					if (ReadPath->NotificationField->LibraryCompilerHandle.IsValid())
+					{
+						bHasFieldId = true;
+					}
+					else if (ReadPath->NotificationField->Source) // it is maybe OneTIme
+					{
+						const UClass* SourceContextClass = ReadPath->NotificationField->Source->AuthoritativeClass;
+						TValueOrError<FCompiledBindingLibraryCompiler::FFieldIdHandle, FText> FieldIdResult = Self->BindingLibraryCompiler.AddFieldId(SourceContextClass, ReadPath->NotificationField->NotificationId.GetFieldName());
+						if (FieldIdResult.HasError() || !FieldIdResult.GetValue().IsValid())
+						{
+							Self->AddMessageForBinding(Binding
+								, FText::Format(LOCTEXT("CouldNotCreateFieldId", "Could not create Field. {0}"), FieldIdResult.GetError())
+								, EMessageType::Error
+								, FName()
+							);
+							return false;
+						}
+						
+						ReadPath->NotificationField->LibraryCompilerHandle = FieldIdResult.StealValue();
+						bHasFieldId = ReadPath->NotificationField->LibraryCompilerHandle.IsValid();
+					}
+				}
+			}
+		}
 
-		if (DestinationContext == nullptr)
+		// Test correct numbers of FieldIds
+		bool bRequiresValidFieldId = !ValidBinding.bIsOneTimeBinding;
+		if (bRequiresValidFieldId && !bHasFieldId)
 		{
-			Self->AddMessageForBinding(Binding, BlueprintView, LOCTEXT("CouldNotFindDestination", "Could not find the pre compiled destination."), EBindingMessageType::Error);
-			return MakeError();
+			Self->AddMessageForBinding(Binding
+				, LOCTEXT("CouldNotCreateSourceFields", "There is no field to bind to. The binding must be a OneTime binding.")
+				, EMessageType::Error
+				, FName()
+			);
+			return false;
 		}
 
-		TArray<UE::MVVM::FMVVMConstFieldVariant> SetterPath;
-		if (!DestinationContext->bCanBeSetInNative)
+		return true;
+	};
+
+	auto AddReadPaths = [Self = this, Class](FCompilerBinding& ValidBinding, const FMVVMBlueprintViewBinding& Binding) -> bool
+	{
+		bool bHasValidReadHandle = false;
+		if (ValidBinding.Type != FCompilerBinding::EType::ComplexConversionFunction)
 		{
-			check(!DestinationContext->GeneratedFunctionName.IsNone());
-			UFunction* FoundFunction = Class->FindFunctionByName(DestinationContext->GeneratedFunctionName);
-			if (FoundFunction == nullptr)
+			for (const TSharedPtr<FGeneratedReadFieldPathContext>& ReadPath : ValidBinding.ReadPaths)
 			{
-				Self->AddMessageForBinding(Binding, BlueprintView, LOCTEXT("CouldNotFindDestinationBPFunction", "Could not find the generated destination function."), EBindingMessageType::Error);
-				return MakeError();
+				if (ReadPath->LibraryCompilerHandle.IsValid())
+				{
+					bHasValidReadHandle = true;
+				}
+				else
+				{
+					TValueOrError<FCompiledBindingLibraryCompiler::FFieldPathHandle, FText> FieldPathResult = Self->BindingLibraryCompiler.AddFieldPath(ReadPath->SkeletalGeneratedFields, true);
+					if (FieldPathResult.HasError() || !FieldPathResult.GetValue().IsValid())
+					{
+						Self->AddMessageForBinding(Binding
+							, FText::Format(Private::CouldNotCreateSourceFieldPathFormat, ::UE::MVVM::FieldPathHelper::ToText(ReadPath->GeneratedFields), FieldPathResult.GetError())
+							, EMessageType::Error
+							, FName()
+						);
+						return false;
+					}
+
+					ReadPath->LibraryCompilerHandle = FieldPathResult.StealValue();
+					bHasValidReadHandle = ReadPath->LibraryCompilerHandle.IsValid();
+				}
 			}
-			FMVVMBlueprintPropertyPath DestinationPPropertyPath;
-			DestinationPPropertyPath.SetSelfContext();
-			DestinationPPropertyPath.SetPropertyPath(BlueprintView->GetOuterUMVVMWidgetBlueprintExtension_View()->GetWidgetBlueprint(), UE::MVVM::FMVVMConstFieldVariant(FoundFunction));
-			SetterPath = Self->CreateBindingDestinationPath(BlueprintView, Class, DestinationPPropertyPath);
-		}
-		else
-		{
-			const FMVVMBlueprintPropertyPath& DestinationPath = bIsForwardBinding ? Binding.DestinationPath : Binding.SourcePath;
-			SetterPath = Self->CreateBindingDestinationPath(BlueprintView, Class, DestinationPath);
 		}
 
-		if (!IsPropertyPathValid(Self->WidgetBlueprintCompilerContext.WidgetBlueprint(), SetterPath))
+		// Test read
+		bool bRequiresReadHandle = ValidBinding.Type != FCompilerBinding::EType::ComplexConversionFunction;
+		if (bRequiresReadHandle && !bHasValidReadHandle)
 		{
-			const FMVVMBlueprintPropertyPath& DestinationPath = bIsForwardBinding ? Binding.DestinationPath : Binding.SourcePath;
-			Self->AddMessageForBinding(Binding, BlueprintView, FText::Format(Private::PropertyPathIsInvalidFormat, PropertyPathToText(Class, BlueprintView, DestinationPath)), EBindingMessageType::Error);
-			return MakeError();
+			Self->AddMessageForBinding(Binding
+				, LOCTEXT("CouldNotCreateSourceReadIsPresent", "Internal error. There should be no property to read from.")
+				, EMessageType::Error
+				, FName()
+			);
+			return false;
 		}
-		return MakeValue(SetterPath);
+
+		return true;
 	};
 
-	for (const FBindingSourceContext& BindingSourceContext : BindingSourceContexts)
+	auto AddWritePaths = [Self = this, Class](FCompilerBinding& ValidBinding, const FMVVMBlueprintViewBinding& Binding) -> bool
 	{
-		FMVVMBlueprintViewBinding* BindingPtr = BlueprintView->GetBindingAt(BindingSourceContext.BindingIndex);
-		check(BindingPtr);
-		FMVVMBlueprintViewBinding& Binding = *BindingPtr;
+		if (ValidBinding.WritePath->LibraryCompilerHandle.IsValid())
+		{
+			return true;
+		}
 
-		FMVVMViewBlueprintCompiler* Self = this;
-		auto AddFieldId = [Self](const UClass* SourceContextClass, bool bNotifyFieldValueChangedRequired, EMVVMBindingMode BindingMode, FName FieldToListenTo) -> TValueOrError<FCompiledBindingLibraryCompiler::FFieldIdHandle, FText>
+		if (!ValidBinding.WritePath)
 		{
-			if (!IsOneTimeBinding(BindingMode) && bNotifyFieldValueChangedRequired)
-			{
-				return Self->BindingLibraryCompiler.AddFieldId(SourceContextClass, FieldToListenTo);
-			}
-			return MakeValue(FCompiledBindingLibraryCompiler::FFieldIdHandle());
-		};
+			Self->AddMessageForBinding(Binding
+				, LOCTEXT("InvalidWritePath", "Internal Error. The write path is invalid.")
+				, EMessageType::Error
+				, FName()
+			);
+			return false;
+		}
 
-		if (!TestExecutionMode(Binding))
+		TValueOrError<FCompiledBindingLibraryCompiler::FFieldPathHandle, FText> FieldPathResult = Self->BindingLibraryCompiler.AddFieldPath(ValidBinding.WritePath->SkeletalGeneratedFields, false);
+		if (FieldPathResult.HasError())
 		{
-			bAreBindingsValid = false;
-			continue;
+			Self->AddMessageForBinding(Binding
+				, FText::Format(Private::CouldNotCreateDestinationFieldPathFormat, ::UE::MVVM::FieldPathHelper::ToText(ValidBinding.WritePath->GeneratedFields), FieldPathResult.GetError())
+				, EMessageType::Error
+				, FName()
+			);
+			return false;
 		}
+		ValidBinding.WritePath->LibraryCompilerHandle = FieldPathResult.StealValue();
 
-		const bool bIsComplexBinding = BindingSourceContext.ComplexConversionFunctionContextIndex != INDEX_NONE;
-		bool bFieldIdNeeded = !IsOneTimeBinding(Binding.BindingType);
+		// test write
+
+		return true;
+	};
 
-		TValueOrError<FCompiledBindingLibraryCompiler::FFieldIdHandle, FText> AddFieldResult = AddFieldId(BindingSourceContext.SourceClass, true, Binding.BindingType, BindingSourceContext.FieldId.GetFieldName());
-		if (AddFieldResult.HasError())
+	auto AddConversionFunction = [Self = this, Class](FCompilerBinding& ValidBinding, const FMVVMBlueprintViewBinding& Binding) -> bool
+	{
+		const UFunction* ConversionFunction = nullptr;
 		{
-			bFieldIdNeeded = false;
-			// For complex binding, at least one must be true
-			if (!bIsComplexBinding)
+			UMVVMBlueprintViewConversionFunction* ViewConversionFunction = ValidBinding.ConversionFunction.Get();
+			if (ViewConversionFunction)
 			{
-				AddMessageForBinding(Binding, BlueprintView, FText::Format(LOCTEXT("CouldNotCreateSource", "Could not create source. {0}"),
-					AddFieldResult.GetError()), EBindingMessageType::Error);
-				bAreBindingsValid = false;
-				continue;
+				ConversionFunction = ViewConversionFunction->GetCompiledFunction(Class);
+				if (ConversionFunction == nullptr)
+				{
+					Self->AddMessageForBinding(Binding
+						, FText::Format(LOCTEXT("ConversionFunctionNotFound", "The conversion function '{0}' could not be found."), FText::FromName(ViewConversionFunction->GetCompiledFunctionName(Class)))
+						, EMessageType::Error
+						, FName()
+					);
+					return false;
+				}
+
+				TVariant<const UFunction*, TSubclassOf<UK2Node>> FunctionOrWrapperFunction = ViewConversionFunction->GetConversionFunction(Self->WidgetBlueprintCompilerContext.WidgetBlueprint());
+				if (FunctionOrWrapperFunction.IsType<const UFunction*>())
+				{
+					if (!GetDefault<UMVVMDeveloperProjectSettings>()->IsConversionFunctionAllowed(Self->WidgetBlueprintCompilerContext.WidgetBlueprint(), FunctionOrWrapperFunction.Get<const UFunction*>()))
+					{
+						Self->AddMessageForBinding(Binding
+							, FText::Format(LOCTEXT("ConversionFunctionNotAllow", "The conversion function {0} is not allowed."), FText::FromName(FunctionOrWrapperFunction.Get<const UFunction*>()->GetFName()))
+							, EMessageType::Error
+							, FName()
+						);
+						return false;
+					}
+				}
+				else
+				{
+					Self->AddMessageForBinding(Binding, LOCTEXT("ConversionFunctionNodeNotAllow", "The conversion function node is not allowed."), EMessageType::Error, FName());
+					return false;
+				}
 			}
 		}
-		else if (bIsComplexBinding)
+
+		if (ConversionFunction != nullptr)
 		{
-			if (ensure(ComplexConversionFunctionContexts.IsValidIndex(BindingSourceContext.ComplexConversionFunctionContextIndex)))
+			TValueOrError<FCompiledBindingLibraryCompiler::FFieldPathHandle, FText> FieldPathResult = Self->BindingLibraryCompiler.AddConversionFunctionFieldPath(Class, ConversionFunction);
+			if (FieldPathResult.HasError())
 			{
-				ComplexConversionFunctionContexts[BindingSourceContext.ComplexConversionFunctionContextIndex].bHasValidFieldId = true;
+				Self->AddMessageForBinding(Binding
+					, FText::Format(LOCTEXT("CouldNotCreateConversionFunctionFieldPath", "Couldn't create the conversion function field path '{0}'. {1}")
+						, FText::FromString(ConversionFunction->GetPathName())
+						, FieldPathResult.GetError())
+					, EMessageType::Error
+					, FName()
+				);
+				return false;
 			}
+
+			ValidBinding.ConversionFunctionHandle = FieldPathResult.StealValue();
 		}
 
-		TArray<UE::MVVM::FMVVMConstFieldVariant> SetterPath;
+		// Sanity check
 		{
-			TValueOrError<TArray<UE::MVVM::FMVVMConstFieldVariant>, void> SetterPathResult = GetSetterPath(Binding, BindingSourceContext.BindingIndex, BindingSourceContext.bIsForwardBinding);
-			if (SetterPathResult.HasError())
+			const bool bShouldHaveConversionFunction = ValidBinding.Type == FCompilerBinding::EType::ComplexConversionFunction || ValidBinding.Type == FCompilerBinding::EType::SimpleConversionFunction;
+			if ((ConversionFunction != nullptr) != bShouldHaveConversionFunction)
 			{
-				bAreBindingsValid = false;
-				continue;
+				Self->AddMessageForBinding(Binding, LOCTEXT("ConversionFunctionShouldExist", "Internal error. The conversion function should exist."), EMessageType::Error, FName());
+				return false;
 			}
-			SetterPath = SetterPathResult.StealValue();
-		}
 
-		const UFunction* ConversionFunction = nullptr;
-		{
-			TValueOrError<const UFunction*, void> ConversionFunctionResult = GetConversionFunction(Binding, BindingSourceContext.bIsForwardBinding);
-			if (ConversionFunctionResult.HasError())
+			const bool bShouldHaveComplexConversionFunction = ValidBinding.Type == FCompilerBinding::EType::ComplexConversionFunction;
+			if (bShouldHaveComplexConversionFunction != BindingHelper::IsValidForComplexRuntimeConversion(ConversionFunction))
 			{
-				bAreBindingsValid = false;
-				continue;
+				Self->AddMessageForBinding(Binding, LOCTEXT("ConversionFunctionIsNotComplex", "Internal Error. The complex conversion function does not respect the prerequisite."), EMessageType::Error, FName());
+				return false;
 			}
-			ConversionFunction = ConversionFunctionResult.StealValue();
-		}
-
-		TValueOrError<FCompiledBinding, FText> AddBindingResult = CreateCompiledBinding(Class, BindingSourceContext.PropertyPath, SetterPath, ConversionFunction, bIsComplexBinding);
-		if (AddBindingResult.HasError())
-		{
-			AddMessageForBinding(Binding, BlueprintView,
-				FText::Format(LOCTEXT("CouldNotCreateBinding", "Could not create binding. {0}"), AddBindingResult.GetError()),
-				EBindingMessageType::Error
-			);
-			bAreBindingsValid = false;
-			continue;
 		}
 
-		FCompilerBinding NewBinding;
-		NewBinding.BindingIndex = BindingSourceContext.BindingIndex;
-		NewBinding.UserWidgetPropertyContextIndex = BindingSourceContext.UserWidgetPropertyContextIndex;
-		NewBinding.SourceCreatorContextIndex = BindingSourceContext.SourceCreatorContextIndex;
-		NewBinding.ComplexConversionFunctionContextIndex = BindingSourceContext.ComplexConversionFunctionContextIndex;
-		NewBinding.bSourceIsUserWidget = BindingSourceContext.bIsRootWidget;
-		NewBinding.bFieldIdNeeded = bFieldIdNeeded;
-		NewBinding.bIsForwardBinding = BindingSourceContext.bIsForwardBinding;
-
-		NewBinding.CompiledBinding = AddBindingResult.StealValue();
-		if (AddFieldResult.HasValue())
-		{
-			NewBinding.FieldIdHandle = AddFieldResult.StealValue();
-		}
-
-		CompilerBindings.Emplace(NewBinding);
-	}
-
+		return true;
+	};
 
-	for (const FSimpleBindingContext& BindingContext : SimpleBindingContexts)
+	for (TSharedRef<FCompilerBinding>& ValidBinding : ValidBindings)
 	{
-		FMVVMBlueprintViewBinding* BindingPtr = BlueprintView->GetBindingAt(BindingContext.BindingIndex);
-		check(BindingPtr);
-		FMVVMBlueprintViewBinding& Binding = *BindingPtr;
-
-		if (!TestExecutionMode(Binding))
+		const FMVVMBlueprintViewBinding& Binding = *(BlueprintView->GetBindingAt(ValidBinding->Key.ViewBindingIndex));
+		if (ValidBinding->Type != FCompilerBinding::EType::Assignment
+			&& ValidBinding->Type != FCompilerBinding::EType::ComplexConversionFunction
+			&& ValidBinding->Type != FCompilerBinding::EType::SimpleConversionFunction)
 		{
-			bAreBindingsValid = false;
+			AddMessageForBinding(Binding, LOCTEXT("CouldNotCreateFieldId", "The binding is invalid."), EMessageType::Error, FName());
+			bIsPreCompileStepValid = false;
 			continue;
 		}
 
-		if (!IsOneTimeBinding(Binding.BindingType))
+		if (!TestExecutionMode(Binding))
 		{
-			AddMessageForBinding(Binding, BlueprintView, LOCTEXT("FunctionRequiresOneTimeBinding", "There is no source. The binding must be a OneTime binding."), EBindingMessageType::Error);
-			bAreBindingsValid = false;
+			bIsPreCompileStepValid = false;
 			continue;
 		}
 
-		TArray<UE::MVVM::FMVVMConstFieldVariant> SetterPath;
+		if (!AddFieldIds(ValidBinding.Get(), Binding)
+			|| !AddReadPaths(ValidBinding.Get(), Binding)
+			|| !AddWritePaths(ValidBinding.Get(), Binding)
+			|| !AddConversionFunction(ValidBinding.Get(), Binding))
 		{
-			TValueOrError<TArray<UE::MVVM::FMVVMConstFieldVariant>, void> SetterPathResult = GetSetterPath(Binding, BindingContext.BindingIndex, BindingContext.bIsForwardBinding);
-			if (SetterPathResult.HasError())
-			{
-				bAreBindingsValid = false;
-				continue;
-			}
-			SetterPath = SetterPathResult.StealValue();
+			bIsPreCompileStepValid = false;
+			continue;
 		}
 
-		const UFunction* ConversionFunction = nullptr;
-		{
-			TValueOrError<const UFunction*, void> ConversionFunctionResult = GetConversionFunction(Binding, BindingContext.bIsForwardBinding);
-			if (ConversionFunctionResult.HasError())
-			{
-				bAreBindingsValid = false;
-				continue;
-			}
-			ConversionFunction = ConversionFunctionResult.StealValue();
-		}
+		// Generate the binding
+		TValueOrError<FCompiledBindingLibraryCompiler::FBindingHandle, FText> BindingResult = ValidBinding->Type == FCompilerBinding::EType::ComplexConversionFunction
+			? BindingLibraryCompiler.AddComplexBinding(ValidBinding->WritePath->LibraryCompilerHandle, ValidBinding->ConversionFunctionHandle)
+			: BindingLibraryCompiler.AddBinding(ValidBinding->ReadPaths[0]->LibraryCompilerHandle, ValidBinding->WritePath->LibraryCompilerHandle, ValidBinding->ConversionFunctionHandle);
 
-		TArrayView<const UE::MVVM::FMVVMConstFieldVariant> GetterFields;
-		const bool bIsComplexBinding = BindingContext.ComplexConversionFunctionContextIndex != INDEX_NONE;
-		TValueOrError<FCompiledBinding, FText> AddBindingResult = CreateCompiledBinding(Class, GetterFields, SetterPath, ConversionFunction, bIsComplexBinding);
-		if (AddBindingResult.HasError())
+		if (BindingResult.HasError())
 		{
-			AddMessageForBinding(Binding, BlueprintView,
-				FText::Format(LOCTEXT("CouldNotCreateBinding", "Could not create binding. {0}"), AddBindingResult.GetError()),
-				EBindingMessageType::Error
+			AddMessageForBinding(Binding
+				, FText::Format(LOCTEXT("CouldNotCreateBinding", "Could not create binding. {0}"), BindingResult.StealError())
+				, EMessageType::Error
+				, FName()
 			);
-			bAreBindingsValid = false;
+			bIsPreCompileStepValid = false;
 			continue;
 		}
-
-		FCompilerBinding NewBinding;
-		NewBinding.BindingIndex = BindingContext.BindingIndex;
-		NewBinding.bFieldIdNeeded = !IsOneTimeBinding(Binding.BindingType);
-		NewBinding.bIsForwardBinding = BindingContext.bIsForwardBinding;
-		NewBinding.bNeedsValidSource = false;
-
-		NewBinding.CompiledBinding = AddBindingResult.StealValue();
-
-		CompilerBindings.Emplace(NewBinding);
+		ValidBinding->BindingHandle = BindingResult.StealValue();
 	}
 
-	if (bAreBindingsValid)
+	// Add bindings for the dynamic viewmodels
+	for (TSharedRef<FCompilerSourceViewModelDynamicCreatorContext>& ViewModelDynamic : SourceViewModelDynamicCreatorContexts)
 	{
-		// Confirm that at least one source is valid for each pin of a complex conversion function.
-		for (const FComplexConversionFunctionContext& ComplexConversionFunctionContext : ComplexConversionFunctionContexts)
+		if (ensure(ViewModelDynamic->Source.IsValid() && ViewModelDynamic->ParentSource.IsValid()))
 		{
-			if (!ComplexConversionFunctionContext.bHasValidFieldId && ComplexConversionFunctionContext.bNeedsValidSource)
+			if (!ViewModelDynamic->NotificationId.IsValid())
+			{
+				AddMessage(FText::Format(LOCTEXT("InvalidNotificationFieldId", "{0} doesn't have a field Id."), FText::FromName(ViewModelDynamic->Source->Name)), EMessageType::Error);
+				bIsPreCompileStepValid = false;
+				continue;
+			}
+
+			const UClass* SourceContextClass = ViewModelDynamic->ParentSource->AuthoritativeClass;
+			TValueOrError<FCompiledBindingLibraryCompiler::FFieldIdHandle, FText> FieldIdResult = BindingLibraryCompiler.AddFieldId(SourceContextClass, ViewModelDynamic->NotificationId.GetFieldName());
+			if (FieldIdResult.HasError() || !FieldIdResult.GetValue().IsValid())
 			{
-				FMVVMBlueprintViewBinding* BindingPtr = BlueprintView->GetBindingAt(ComplexConversionFunctionContext.BindingIndex);
-				check(BindingPtr);
-				AddMessageForBinding(*BindingPtr, BlueprintView, LOCTEXT("CouldNotCreateSourceForConversionFunction", "There is no source. The binding must be a OneTime binding."), EBindingMessageType::Error);
-				bAreBindingsValid = false;
+				AddMessage(FText::Format(LOCTEXT("CouldNotCreateFieldIdForDynamic", "Could not create Field for {0}. {1}"), FText::FromName(ViewModelDynamic->Source->Name), FieldIdResult.GetError()), EMessageType::Error);
+				bIsPreCompileStepValid = false;
+				continue;
 			}
+
+			ViewModelDynamic->NotificationIdLibraryCompilerHandle = FieldIdResult.GetValue();
 		}
 	}
-
-	return bAreBindingsValid;
 }
 
 
-bool FMVVMViewBlueprintCompiler::CompileBindings(const FCompiledBindingLibraryCompiler::FCompileResult& CompileResult, UWidgetBlueprintGeneratedClass* Class, UMVVMBlueprintView* BlueprintView, UMVVMViewClass* ViewExtension)
+void FMVVMViewBlueprintCompiler::CompileBindings(const FCompiledBindingLibraryCompiler::FCompileResult& CompileResult, UWidgetBlueprintGeneratedClass* Class, UMVVMViewClass* ViewExtension)
 {
-	if (!bAreBindingsValid)
+	static IConsoleVariable* CVarDefaultExecutionMode = IConsoleManager::Get().FindConsoleVariable(TEXT("MVVM.DefaultExecutionMode"));
+	ensure(CVarDefaultExecutionMode);
+	if (!CVarDefaultExecutionMode)
 	{
-		return false;
+		WidgetBlueprintCompilerContext.MessageLog.Error(*LOCTEXT("CantFindDefaultExecutioMode", "The default execution mode cannot be found.").ToString());
+		return;
 	}
 
-	struct FLocal
+	struct FSortData
 	{
-		FLocal()
-		{
-			IConsoleVariable* CVarDefaultExecutionMode = IConsoleManager::Get().FindConsoleVariable(TEXT("MVVM.DefaultExecutionMode"));
-			if (ensure(CVarDefaultExecutionMode))
-			{
-				DefaultMode = (EMVVMExecutionMode)CVarDefaultExecutionMode->GetInt();
-			}
-		}
-		EMVVMExecutionMode DefaultMode = EMVVMExecutionMode::DelayedWhenSharedElseImmediate;
+		int32 ValidBindingIndex = INDEX_NONE;
+		int32 SourceIndex = INDEX_NONE;
+		FMVVMViewClass_CompiledBinding CompiledBinding;
 	};
-	static FLocal LocalDefaultMode;
+	TArray<FSortData> TempBindingToSort;
+	TempBindingToSort.Reserve(ValidBindings.Num()*2);
 
-	// Store bindings with corresponding ComplexConversionFunctionContextIndex to be sorted and mark the last binding 
-	// in the complex conversion as bExecuteAtInitialization == true
-	TArray<TPair<int32, FMVVMViewClass_CompiledBinding>> TempBindingPairs;
-	TempBindingPairs.Reserve(CompilerBindings.Num());
-
-	for (const FCompilerBinding& CompileBinding : CompilerBindings)
+	for (int32 ValidBindingIndex = 0; ValidBindingIndex < ValidBindings.Num(); ++ValidBindingIndex)
 	{
-		// PropertyBinding needs a valid CompileBinding.BindingIndex.
-		check(CompileBinding.Type != ECompilerBindingType::PropertyBinding || (CompileBinding.BindingIndex >= 0 && CompileBinding.BindingIndex < BlueprintView->GetNumBindings()));
-		FMVVMBlueprintViewBinding* ViewBinding = CompileBinding.Type == ECompilerBindingType::PropertyBinding
-			? BlueprintView->GetBindingAt(CompileBinding.BindingIndex)
-			: nullptr;
-
+		const TSharedRef<FCompilerBinding>& ValidBinding = ValidBindings[ValidBindingIndex];
+		const FMVVMBlueprintViewBinding& Binding = *(BlueprintView->GetBindingAt(ValidBinding->Key.ViewBindingIndex));
 
-		FMVVMViewClass_CompiledBinding NewBinding;
-		const bool bIsSourceSelf = CompileBinding.bSourceIsUserWidget;
-		if (!bIsSourceSelf)
+		for (TSharedPtr<FGeneratedReadFieldPathContext> ReadPath : ValidBinding->ReadPaths)
 		{
-			if (CompilerUserWidgetPropertyContexts.IsValidIndex(CompileBinding.UserWidgetPropertyContextIndex))
+			if (ReadPath->OptionalSource == nullptr)
 			{
-				NewBinding.SourcePropertyName = CompilerUserWidgetPropertyContexts[CompileBinding.UserWidgetPropertyContextIndex].PropertyName;
-			}
-			else if (CompilerSourceCreatorContexts.IsValidIndex(CompileBinding.SourceCreatorContextIndex))
-			{
-				NewBinding.SourcePropertyName = CompilerSourceCreatorContexts[CompileBinding.SourceCreatorContextIndex].ViewModelContext.GetViewModelName();
-			}
-			else if (CompileBinding.bNeedsValidSource)
-			{
-				WidgetBlueprintCompilerContext.MessageLog.Error(*LOCTEXT("InvalidSourceInternal", "Internal error. The binding doesn't have a valid source.").ToString());
-				bAreBindingsValid = false;
-				return false;
-			}
-		}
-
-		auto AddErrorMessage = [this, ViewBinding, BlueprintView](const FText& ErrorMessage)
-		{
-				if (ViewBinding)
-				{
-					AddMessageForBinding(*ViewBinding, BlueprintView, ErrorMessage, EBindingMessageType::Error);
-				}
-				else
-				{
-					WidgetBlueprintCompilerContext.MessageLog.Error(*ErrorMessage.ToString());
-				}
-		};
-
-		const FMVVMVCompiledFieldId* CompiledFieldId = CompileResult.FieldIds.Find(CompileBinding.FieldIdHandle);
-		if (CompiledFieldId == nullptr && CompileBinding.bFieldIdNeeded)
-		{
-			AddErrorMessage(FText::Format(LOCTEXT("FieldIdNotGenerated", "Could not generate field ID for property '{0}'."), FText::FromName(NewBinding.SourcePropertyName)));
-			bAreBindingsValid = false;
-			continue;
-		}
-
-		const FMVVMVCompiledBinding* CompiledBinding = CompileResult.Bindings.Find(CompileBinding.CompiledBinding.BindingHandle);
-		if (CompiledBinding == nullptr && CompileBinding.Type != ECompilerBindingType::ViewModelDynamic)
-		{
-			AddErrorMessage(LOCTEXT("CompiledBindingNotGenerated", "Could not generate compiled binding."));
-			bAreBindingsValid = false;
-			continue;
-		}
-
-		bool bIsOptional = false;
-		if (!bIsSourceSelf)
-		{
-			const FMVVMBlueprintViewModelContext* ViewModelContext = nullptr;
-			if (CompilerUserWidgetPropertyContexts.IsValidIndex(CompileBinding.UserWidgetPropertyContextIndex))
-			{
-				if (CompilerUserWidgetPropertyContexts[CompileBinding.UserWidgetPropertyContextIndex].ViewModelId.IsValid())
-				{
-					ViewModelContext = BlueprintView->FindViewModel(CompilerUserWidgetPropertyContexts[CompileBinding.UserWidgetPropertyContextIndex].ViewModelId);
-					if (ViewModelContext == nullptr)
-					{
-						AddErrorMessage(LOCTEXT("CompiledBindingWithInvalidIVewModelId", "Internal error: the viewmodel became invalid."));
-						bAreBindingsValid = false;
-						continue;
-					}
-				}
-			}
-			else if (CompilerSourceCreatorContexts.IsValidIndex(CompileBinding.SourceCreatorContextIndex))
-			{
-				ViewModelContext = &CompilerSourceCreatorContexts[CompileBinding.SourceCreatorContextIndex].ViewModelContext;
+				AddMessage(LOCTEXT("InvalidSourceInternal", "Internal error. The binding doesn't have a valid source."), EMessageType::Error);
+				bIsCompileStepValid = false;
+				continue;
 			}
 
-			if (ViewModelContext)
-			{
-				bIsOptional = ViewModelContext->bOptional || ViewModelContext->CreationType == EMVVMBlueprintViewModelContextCreationType::Manual;
-			}
-		}
-
-		if (CompileBinding.Type == ECompilerBindingType::ViewModelDynamic)
-		{
-			int32 FoundSourceCreatorIndex = ViewExtension->SourceCreators.IndexOfByPredicate([LookFor = CompileBinding.DynamicViewModelName](const FMVVMViewClass_SourceCreator& Other)
+			int32 ViewExtensionSourceCreatorsIndex = ViewExtension->SourceCreators.IndexOfByPredicate([LookFor = ReadPath->OptionalSource->Name](const FMVVMViewClass_SourceCreator& Other)
 				{
 					return Other.GetSourceName() == LookFor;
 				});
-
-			check(FoundSourceCreatorIndex < std::numeric_limits<int8>::max()); // the index is saved as a int8
-			NewBinding.EvaluateSourceCreatorIndex = FoundSourceCreatorIndex;
-			if (NewBinding.EvaluateSourceCreatorIndex == INDEX_NONE)
-			{
-				FText ErrorMessage = FText::Format(LOCTEXT("CompiledViewModelDynamicBindingNotGenerated", "Was not able to find the source for {0}."), FText::FromName(CompileBinding.DynamicViewModelName));
-				WidgetBlueprintCompilerContext.MessageLog.Error(*ErrorMessage.ToString());
-				bAreBindingsValid = false;
-				continue;
-			}
-		}
-
-		FText WrongBindingTypeInternalErrorMessage = LOCTEXT("WrongBindingTypeInternalError", "Internal Error. Wrong binding type.");
-		// FieldId can be null if it's a one time binding
-		NewBinding.FieldId = CompiledFieldId ? *CompiledFieldId : FMVVMVCompiledFieldId();
-		NewBinding.Binding = CompiledBinding ? *CompiledBinding : FMVVMVCompiledBinding();
-		NewBinding.Flags = 0;
-
-		int32 ComplexConversionFunctionContextIndex = INDEX_NONE;
-
-		if (ViewBinding)
-		{
-			if (CompileBinding.Type != ECompilerBindingType::PropertyBinding)
+			if (!ViewExtension->SourceCreators.IsValidIndex(ViewExtensionSourceCreatorsIndex))
 			{
-				WidgetBlueprintCompilerContext.MessageLog.Error(*WrongBindingTypeInternalErrorMessage.ToString());
-				bAreBindingsValid = false;
+				AddMessageForBinding(Binding, LOCTEXT("CompiledSourceCreatorNotGenerated", "Internal error. The source creator was not generated."), EMessageType::Error, FName());
+				bIsCompileStepValid = false;
 				continue;
 			}
 
-			// If it is forward binding and simple conversion, mark it as bExecuteAtInitialization == true
-			bool bExecuteAtInitialization = CompileBinding.bIsForwardBinding && CompileBinding.ComplexConversionFunctionContextIndex == INDEX_NONE;
-
-			// If it is forward binding and complex conversion, store ComplexConversionFunctionContextIndex to be handled later in this function.
-			if (CompileBinding.bIsForwardBinding && CompileBinding.ComplexConversionFunctionContextIndex != INDEX_NONE)
-			{
-				// Only forward binding with complex conversion will have a valid ComplexConversionFunctionContextIndex
-				ComplexConversionFunctionContextIndex = CompileBinding.ComplexConversionFunctionContextIndex;
-			}
+			const FMVVMVCompiledFieldId* CompiledFieldId = (ReadPath->NotificationField && !ValidBinding->bIsOneTimeBinding) ? CompileResult.FieldIds.Find(ReadPath->NotificationField->LibraryCompilerHandle) : nullptr;
 
-			NewBinding.ExecutionMode = ViewBinding->bOverrideExecutionMode ? ViewBinding->OverrideExecutionMode : LocalDefaultMode.DefaultMode;
-			NewBinding.EditorId = ViewBinding->BindingId;
-
-			NewBinding.Flags |= (ViewBinding->bEnabled) ? FMVVMViewClass_CompiledBinding::EBindingFlags::EnabledByDefault : 0;
-			NewBinding.Flags |= (bExecuteAtInitialization) ? FMVVMViewClass_CompiledBinding::EBindingFlags::ExecuteAtInitialization : 0;
-			NewBinding.Flags |= (IsOneTimeBinding(ViewBinding->BindingType)) ? FMVVMViewClass_CompiledBinding::EBindingFlags::OneTime : 0;
-			NewBinding.Flags |= (bIsOptional) ? FMVVMViewClass_CompiledBinding::EBindingFlags::ViewModelOptional : 0;
-			NewBinding.Flags |= (ViewBinding->bOverrideExecutionMode) ? FMVVMViewClass_CompiledBinding::EBindingFlags::OverrideExecuteMode : 0;
-			NewBinding.Flags |= (bIsSourceSelf) ? FMVVMViewClass_CompiledBinding::EBindingFlags::SourceObjectIsSelf : 0;
-		}
-		else
-		{
-			if (CompileBinding.Type != ECompilerBindingType::ViewModelDynamic)
+			const FMVVMVCompiledBinding* CompiledBinding = CompileResult.Bindings.Find(ValidBinding->BindingHandle);
+			if (CompiledBinding == nullptr)
 			{
-				WidgetBlueprintCompilerContext.MessageLog.Error(*WrongBindingTypeInternalErrorMessage.ToString());
-				bAreBindingsValid = false;
+				AddMessageForBinding(Binding, LOCTEXT("CompiledBindingNotGenerated", "Could not generate compiled binding."), EMessageType::Error, FName());
+				bIsCompileStepValid = false;
 				continue;
 			}
-			NewBinding.ExecutionMode = EMVVMExecutionMode::Immediate;
-			NewBinding.EditorId = FGuid();
 
-			NewBinding.Flags |= FMVVMViewClass_CompiledBinding::EBindingFlags::EnabledByDefault;
+			bool bExecuteAtInitialization = ValidBinding->Key.bIsForwardBinding;
+			bool bIsOptional = ViewExtension->SourceCreators[ViewExtensionSourceCreatorsIndex].IsOptional();
+
+			FMVVMViewClass_CompiledBinding NewBinding;
+			NewBinding.FieldId = CompiledFieldId ? *CompiledFieldId : FMVVMVCompiledFieldId();
+			NewBinding.SourcePropertyName = ReadPath->OptionalSource->Name;
+			NewBinding.Binding = CompiledBinding ? *CompiledBinding : FMVVMVCompiledBinding();
+			NewBinding.ExecutionMode = Binding.bOverrideExecutionMode ? Binding.OverrideExecutionMode : (EMVVMExecutionMode)CVarDefaultExecutionMode->GetInt();
+			NewBinding.EvaluateSourceCreatorIndex = INDEX_NONE;;
+			NewBinding.EditorId = Binding.BindingId;
+
+			NewBinding.Flags = 0;
+			NewBinding.Flags |= (Binding.bEnabled) ? FMVVMViewClass_CompiledBinding::EBindingFlags::EnabledByDefault : 0;
+			NewBinding.Flags |= (bExecuteAtInitialization) ? FMVVMViewClass_CompiledBinding::EBindingFlags::ExecuteAtInitialization : 0;
+			NewBinding.Flags |= (ValidBinding->bIsOneTimeBinding) ? FMVVMViewClass_CompiledBinding::EBindingFlags::OneTime : 0;
 			NewBinding.Flags |= (bIsOptional) ? FMVVMViewClass_CompiledBinding::EBindingFlags::ViewModelOptional : 0;
-			NewBinding.Flags |= FMVVMViewClass_CompiledBinding::EBindingFlags::OverrideExecuteMode; // The mode needs to be Immediate.
-			NewBinding.Flags |= (bIsSourceSelf) ? FMVVMViewClass_CompiledBinding::EBindingFlags::SourceObjectIsSelf : 0;
-		}
+			NewBinding.Flags |= (Binding.bOverrideExecutionMode) ? FMVVMViewClass_CompiledBinding::EBindingFlags::OverrideExecuteMode : 0;
+			NewBinding.Flags |= (ReadPath->OptionalSource->Type == FCompilerBindingSource::EType::Self) ? FMVVMViewClass_CompiledBinding::EBindingFlags::SourceObjectIsSelf : 0;
 
-		TempBindingPairs.Emplace(ComplexConversionFunctionContextIndex, MoveTemp(NewBinding));
+			FSortData& SortData = TempBindingToSort.AddDefaulted_GetRef();
+			SortData.ValidBindingIndex = ValidBindingIndex;
+			SortData.SourceIndex = ViewExtensionSourceCreatorsIndex;
+			SortData.CompiledBinding = MoveTemp(NewBinding);
+		}
 	}
 
-	if (ViewExtension->SourceCreators.Num() > 1)
+	for (const TSharedRef<FCompilerSourceViewModelDynamicCreatorContext>& ViewModelDynamic : SourceViewModelDynamicCreatorContexts)
 	{
-		TMap<FName, int32> SourceCreatorSortOrder;
-		SourceCreatorSortOrder.Reserve(ViewExtension->SourceCreators.Num());
-		for (int32 Index = 0; Index < ViewExtension->SourceCreators.Num(); ++Index)
+		int32 ViewExtensionSourceCreatorsIndex = ViewExtension->SourceCreators.IndexOfByPredicate([LookFor = ViewModelDynamic->Source->Name](const FMVVMViewClass_SourceCreator& Other)
+			{
+				return Other.GetSourceName() == LookFor;
+			});
+		if (!ViewExtension->SourceCreators.IsValidIndex(ViewExtensionSourceCreatorsIndex))
 		{
-			SourceCreatorSortOrder.Add(ViewExtension->SourceCreators[Index].GetSourceName(), Index);
+			AddMessage(LOCTEXT("CompiledSourceCreatorNotGenerated", "Internal error. The source creator was not generated."), EMessageType::Error);
+			bIsCompileStepValid = false;
+			continue;
 		}
 
-		TempBindingPairs.StableSort([&SourceCreatorSortOrder](const TPair<int32, FMVVMViewClass_CompiledBinding>& PairA, const TPair<int32, FMVVMViewClass_CompiledBinding>& PairB)
-			{
-				const FMVVMViewClass_CompiledBinding& A = PairA.Value;
-				const FMVVMViewClass_CompiledBinding& B = PairB.Value;
+		const FMVVMVCompiledFieldId* CompiledFieldId = CompileResult.FieldIds.Find(ViewModelDynamic->NotificationIdLibraryCompilerHandle);
+		if (CompiledFieldId == nullptr)
+		{
+			AddMessage(LOCTEXT("CompiledFieldNotifyNotGenerated", "Internal error. The field notify was not generated."), EMessageType::Error);
+			bIsCompileStepValid = false;
+			continue;
+		}
+		
+		bool bIsOptional = ViewModelDynamic->Source->bIsOptional;
+
+		FMVVMViewClass_CompiledBinding NewBinding;
+		NewBinding.FieldId = CompiledFieldId ? *CompiledFieldId : FMVVMVCompiledFieldId();
+		NewBinding.SourcePropertyName = ViewModelDynamic->ParentSource->Name;
+		NewBinding.Binding = FMVVMVCompiledBinding();
+		NewBinding.ExecutionMode = EMVVMExecutionMode::Immediate;
+		NewBinding.EvaluateSourceCreatorIndex = ViewExtensionSourceCreatorsIndex;
+		NewBinding.EditorId = FGuid();
+
+		NewBinding.Flags = 0;
+		NewBinding.Flags |= FMVVMViewClass_CompiledBinding::EBindingFlags::EnabledByDefault;
+		NewBinding.Flags |= (bIsOptional) ? FMVVMViewClass_CompiledBinding::EBindingFlags::ViewModelOptional : 0;
+		NewBinding.Flags |= FMVVMViewClass_CompiledBinding::EBindingFlags::OverrideExecuteMode; // The mode needs to be Immediate.
+		NewBinding.Flags |= (ViewModelDynamic->ParentSource->Type == FCompilerBindingSource::EType::Self) ? FMVVMViewClass_CompiledBinding::EBindingFlags::SourceObjectIsSelf : 0;
+
+		FSortData& SortData = TempBindingToSort.AddDefaulted_GetRef();
+		SortData.ValidBindingIndex = INDEX_NONE;
+		SortData.SourceIndex = ViewExtensionSourceCreatorsIndex;
+		SortData.CompiledBinding = MoveTemp(NewBinding);
+	}
 
-				int32* FoundA = SourceCreatorSortOrder.Find(A.GetSourceName());
-				int32* FoundB = SourceCreatorSortOrder.Find(B.GetSourceName());
-				int32 AIndex = FoundA ? *FoundA : -1;
-				int32 BIndex = FoundB ? *FoundB : -1;
 
-				if (AIndex == BIndex)
+	// sort the array
+	if (ViewExtension->SourceCreators.Num() > 1)
+	{
+		// sort by source, then regular binding, then "dynamic viewmodel" at the end
+		TempBindingToSort.StableSort([](const FSortData& A, const FSortData& B)
+			{
+				if (A.SourceIndex == B.SourceIndex)
 				{
-					if (A.GetEvaluateSourceCreatorBindingIndex() == B.GetEvaluateSourceCreatorBindingIndex())
+					if (A.CompiledBinding.GetEvaluateSourceCreatorBindingIndex() == B.CompiledBinding.GetEvaluateSourceCreatorBindingIndex())
 					{
-						return A.GetSourceName().LexicalLess(B.GetSourceName());
+						return A.CompiledBinding.GetSourceName().LexicalLess(B.CompiledBinding.GetSourceName());
 					}
 
-					return A.GetEvaluateSourceCreatorBindingIndex() < B.GetEvaluateSourceCreatorBindingIndex();
+					if (A.CompiledBinding.GetEvaluateSourceCreatorBindingIndex() == INDEX_NONE)
+					{
+						return false;
+					}
+					if (B.CompiledBinding.GetEvaluateSourceCreatorBindingIndex() == INDEX_NONE)
+					{
+						return true;
+					}
+
+					return A.CompiledBinding.GetEvaluateSourceCreatorBindingIndex() < B.CompiledBinding.GetEvaluateSourceCreatorBindingIndex();
 				}
 
-				return AIndex < BIndex;
+				return A.SourceIndex < B.SourceIndex;
 			});
 	}
 
 	// Go backward to be able to mark only the last binding in the complex conversion.
-	for (int32 Index = TempBindingPairs.Num() - 1; Index >= 0; --Index)
+	TSet<int32> ComplexFunctionInitialized;
+	for (int32 Index = TempBindingToSort.Num() - 1; Index >= 0; --Index)
 	{
-		FMVVMViewClass_CompiledBinding& Binding = TempBindingPairs[Index].Value;
-		const int32 ComplexConversionFunctionContextIndex = TempBindingPairs[Index].Key;
+		FSortData& Data = TempBindingToSort[Index];
 
 		// Only handle bindings with complex conversions.
-		if (ComplexConversionFunctionContextIndex != INDEX_NONE)
+		if (Data.CompiledBinding.IsConversionFunctionComplex())
 		{
-			// Mark the complex conversion as bExecAtInitGenerated.
-			if (!ComplexConversionFunctionContexts[ComplexConversionFunctionContextIndex].bExecAtInitGenerated)
+			// The last binding keep the init flag. Remove it for the others
+			if (ComplexFunctionInitialized.Contains(Data.ValidBindingIndex))
 			{
-				ComplexConversionFunctionContexts[ComplexConversionFunctionContextIndex].bExecAtInitGenerated = true;
-
-				// Mark the last binding with the complex conversion.
-				Binding.Flags |= FMVVMViewClass_CompiledBinding::EBindingFlags::ExecuteAtInitialization;
+				Data.CompiledBinding.Flags &= ~FMVVMViewClass_CompiledBinding::EBindingFlags::ExecuteAtInitialization;
+			}
+			else
+			{
+				ComplexFunctionInitialized.Add(Data.ValidBindingIndex);
 			}
 		}
 	}
 
-	ViewExtension->CompiledBindings.Reserve(TempBindingPairs.Num());
-	for (TPair<int32, FMVVMViewClass_CompiledBinding>& BindingPair : TempBindingPairs)
+	ViewExtension->CompiledBindings.Reset(TempBindingToSort.Num());
+	for (FSortData& Data : TempBindingToSort)
 	{
-		ViewExtension->CompiledBindings.Emplace(MoveTemp(BindingPair.Value));
+		ViewExtension->CompiledBindings.Emplace(MoveTemp(Data.CompiledBinding));
 	}
-
-	return bAreBindingsValid;
+	TempBindingToSort.Reset();
 }
 
-bool FMVVMViewBlueprintCompiler::PreCompileEvents(UWidgetBlueprintGeneratedClass* Class, UMVVMBlueprintView* BlueprintView)
-{
-	if (!bAreSourceContextsValid || !bAreEventsValid)
-	{
-		return false;
-	}
 
+void FMVVMViewBlueprintCompiler::PreCompileEvents(UWidgetBlueprintGeneratedClass* Class)
+{
 	if (!GetDefault<UMVVMDeveloperProjectSettings>()->bAllowBindingEvent && BlueprintView->GetEvents().Num() > 0)
 	{
 		WidgetBlueprintCompilerContext.MessageLog.Warning(*LOCTEXT("EventsAreNotAllowed", "Binding events are not allowed in your project settings.").ToString());
 	}
 
-	for (UMVVMBlueprintViewEvent* EventPtr : BlueprintView->GetEvents())
+	for (TSharedRef<FCompilerEvent>& ValidEvent : ValidEvents)
 	{
-		if (EventPtr == nullptr || !EventPtr->bCompile)
-		{
-			continue;
-		}
-
+		UMVVMBlueprintViewEvent* EventPtr = ValidEvent->Event.Get();
+		check(EventPtr);
 		UEdGraph* GeneratedGraph = EventPtr->GetOrCreateWrapperGraph();
-		if (GeneratedGraph == nullptr)
-		{
-			AddMessageForEvent(EventPtr
-				, FText::Format(LOCTEXT("InvalidEventGeneratedGraph", "The generated graph for event {0} is invalid."), EventPtr->GetDisplayName(true))
-				, UMVVMBlueprintViewEvent::EMessageType::Error
-				, FName());
-			bAreEventsValid = false;
-			continue;
-		}
-
-		auto TestPath = [Self = this, BlueprintView, Class](UMVVMBlueprintViewEvent* EventPtr, const FMVVMBlueprintPropertyPath& PropertyPath, FName ArgumentName)
-		{
-			TArray<UE::MVVM::FMVVMConstFieldVariant> GeneratedPath = Self->CreateBindingDestinationPath(BlueprintView, Class, PropertyPath);
-			if (!Self->IsPropertyPathValid(Self->WidgetBlueprintCompilerContext.WidgetBlueprint(), GeneratedPath))
-			{
-				Self->AddMessageForEvent(EventPtr
-					, FText::Format(Private::PropertyPathIsInvalidFormat, PropertyPathToText(Class, BlueprintView, PropertyPath))
-					, UMVVMBlueprintViewEvent::EMessageType::Error
-					, ArgumentName);
-				Self->bAreEventsValid = false;
-				GeneratedPath.Empty();
-			}
-			return GeneratedPath;
-		};
+		check(GeneratedGraph);
 
 		// Does it resolve and are the field allowed
-		TArray<UE::MVVM::FMVVMConstFieldVariant> EventPath = TestPath(EventPtr, EventPtr->GetEventPath(), FName());
-		if (EventPath.Num() == 0)
-		{
-			continue;
-		}
-		if (TestPath(EventPtr, EventPtr->GetDestinationPath(), FName()).Num() == 0)
+		TValueOrError<FCreateFieldsResult, FText> EventPathResult = CreateFieldContext(Class, EventPtr->GetEventPath(), true);
+		if (EventPathResult.HasError())
 		{
+			AddMessageForEvent(EventPtr
+				, FText::Format(Private::PropertyPathIsInvalidFormat, PropertyPathToText(Class, BlueprintView.Get(), EventPtr->GetEventPath()))
+				, EMessageType::Error
+				, FName());
+			bIsPreCompileStepValid = false;
 			continue;
 		}
 
-		for (const FMVVMBlueprintPin& Pin : EventPtr->GetPins())
-		{
-			if (Pin.UsedPathAsValue())
-			{
-				if (TestPath(EventPtr, Pin.GetPath(), Pin.GetName()).Num() == 0)
-				{
-					continue;
-				}
-			}
-		}
-
 		//Does EventPath resolves to a MulticastDelegateProperty
-		const FMulticastDelegateProperty* DelegateProperty = EventPath.Last().IsProperty() ? CastField<const FMulticastDelegateProperty>(EventPath.Last().GetProperty()) : nullptr;
+		const FMulticastDelegateProperty* DelegateProperty = EventPathResult.GetValue().SkeletalGeneratedFields.Last().IsProperty() ? CastField<const FMulticastDelegateProperty>(EventPathResult.GetValue().SkeletalGeneratedFields.Last().GetProperty()) : nullptr;
 		if (DelegateProperty == nullptr)
 		{
 			AddMessageForEvent(EventPtr
-				, FText::Format(LOCTEXT("EventPathIsNotMulticastDelegate", "The event {0} is not a multicast delegate."), PropertyPathToText(Class, BlueprintView, EventPtr->GetEventPath()))
-				, UMVVMBlueprintViewEvent::EMessageType::Error
+				, FText::Format(LOCTEXT("EventPathIsNotMulticastDelegate", "The event {0} is not a multicast delegate."), PropertyPathToText(Class, BlueprintView.Get(), EventPtr->GetEventPath()))
+				, EMessageType::Error
 				, FName());
-			bAreEventsValid = false;
+			bIsPreCompileStepValid = false;
 			continue;
 		}
 
@@ -2183,34 +2462,25 @@ bool FMVVMViewBlueprintCompiler::PreCompileEvents(UWidgetBlueprintGeneratedClass
 		{
 			AddMessageForEvent(EventPtr
 				, FText::Format(LOCTEXT("EventPathFunctionSignatureError", "The event {0} doesn't match the function signature."), DelegateProperty->GetDisplayNameText())
-				, UMVVMBlueprintViewEvent::EMessageType::Error
+				, EMessageType::Error
 				, FName());
-			bAreEventsValid = false;
+			bIsPreCompileStepValid = false;
 			continue;
 		}
 
 		// Generate the FieldPath to get the delegate property at runtime
-		TValueOrError<TArray<FMVVMConstFieldVariant>, FText> GeneratedEventFields = FieldPathHelper::GenerateFieldPathList(EventPath, true);
-		if (GeneratedEventFields.HasError())
+		TValueOrError<FCompiledBindingLibraryCompiler::FFieldPathHandle, FText> EventFieldPathResult = BindingLibraryCompiler.AddFieldPath(EventPathResult.GetValue().SkeletalGeneratedFields, true);
+		if (EventFieldPathResult.HasError())
 		{
 			AddMessageForEvent(EventPtr
-				, FText::Format(Private::CouldNotCreateSourceFieldPathFormat, ::UE::MVVM::FieldPathHelper::ToText(EventPath), GeneratedEventFields.GetError())
-				, UMVVMBlueprintViewEvent::EMessageType::Error
+				, FText::Format(Private::CouldNotCreateSourceFieldPathFormat, PropertyPathToText(Class, BlueprintView.Get(), EventPtr->GetEventPath()), EventFieldPathResult.GetError())
+				, EMessageType::Error
 				, FName());
-			bAreEventsValid = false;
+			bIsPreCompileStepValid = false;
 			continue;
 		}
 
-		TValueOrError<FCompiledBindingLibraryCompiler::FFieldPathHandle, FText> FieldPathResult = BindingLibraryCompiler.AddFieldPath(GeneratedEventFields.GetValue(), true);
-		if (FieldPathResult.HasError())
-		{
-			AddMessageForEvent(EventPtr
-				, FText::Format(Private::CouldNotCreateSourceFieldPathFormat, ::UE::MVVM::FieldPathHelper::ToText(EventPath), FieldPathResult.GetError())
-				, UMVVMBlueprintViewEvent::EMessageType::Error
-				, FName());
-			bAreEventsValid = false;
-			continue;
-		}
+		//todo ? this doesn't support long path?
 
 		FName SourceName;
 		switch (EventPtr->GetEventPath().GetSource(WidgetBlueprintCompilerContext.WidgetBlueprint()))
@@ -2225,11 +2495,7 @@ bool FMVVMViewBlueprintCompiler::PreCompileEvents(UWidgetBlueprintGeneratedClass
 			const bool bSourceIsUserWidget = WidgetName == Class->ClassGeneratedBy->GetFName();
 			ensure(!bSourceIsUserWidget);
 			
-			const int32 VariableContextIndex = CompilerUserWidgetPropertyContexts.IndexOfByPredicate([WidgetName](const FCompilerUserWidgetPropertyContext& Other) { return Other.PropertyName == WidgetName; });
-			if (ensureAlwaysMsgf(VariableContextIndex != INDEX_NONE, TEXT("Could not find source context for destination '%s'"), *WidgetName.ToString()))
-			{
-				SourceName = CompilerUserWidgetPropertyContexts[VariableContextIndex].PropertyName;
-			}
+			SourceName = WidgetName;
 			break;
 		}
 		case EMVVMBlueprintFieldPathSource::ViewModel:
@@ -2237,9 +2503,7 @@ bool FMVVMViewBlueprintCompiler::PreCompileEvents(UWidgetBlueprintGeneratedClass
 			const FMVVMBlueprintViewModelContext* SourceViewModelContext = BlueprintView->FindViewModel(EventPtr->GetEventPath().GetViewModelId());
 			check(SourceViewModelContext);
 			FName ViewModelName = SourceViewModelContext->GetViewModelName();
-			const int32 DestinationVariableContextIndex = CompilerUserWidgetPropertyContexts.IndexOfByPredicate([ViewModelName](const FCompilerUserWidgetPropertyContext& Other) { return Other.PropertyName == ViewModelName; });
-			check(DestinationVariableContextIndex != INDEX_NONE);
-			SourceName = CompilerUserWidgetPropertyContexts[DestinationVariableContextIndex].PropertyName;
+			SourceName = ViewModelName;
 			break;
 		}
 		default:
@@ -2249,207 +2513,235 @@ bool FMVVMViewBlueprintCompiler::PreCompileEvents(UWidgetBlueprintGeneratedClass
 		// No need to add the generated function to the field compiler.
 		//They are in the BP generated code.
 
-		FEventDelegateContext& NewContext = EventDelegateContexts.AddDefaulted_GetRef();
-		NewContext.BlueprintEvent = EventPtr;
-		NewContext.DelegateFieldPathHandle = FieldPathResult.StealValue();
-		NewContext.GeneratedGraphName = EventPtr->GetWrapperGraphName();
-		NewContext.SourceName = SourceName;
+		ValidEvent->DelegateFieldPathHandle = EventFieldPathResult.StealValue();
+		ValidEvent->GeneratedGraphName = EventPtr->GetWrapperGraphName();
+		ValidEvent->SourceName = SourceName;
 	}
-
-	return bAreEventsValid;
 }
 
 
-bool FMVVMViewBlueprintCompiler::CompileEvents(const FCompiledBindingLibraryCompiler::FCompileResult& CompileResult, UWidgetBlueprintGeneratedClass* Class, UMVVMBlueprintView* BlueprintView, UMVVMViewClass* ViewExtension)
+void FMVVMViewBlueprintCompiler::CompileEvents(const FCompiledBindingLibraryCompiler::FCompileResult& CompileResult, UWidgetBlueprintGeneratedClass* Class, UMVVMViewClass* ViewExtension)
 {
-	if (!bAreSourceContextsValid || !bAreEventsValid)
-	{
-		return false;
-	}
-
-	for (FEventDelegateContext& Context : EventDelegateContexts)
+	for (const TSharedRef<FCompilerEvent>& ValidEvent : ValidEvents)
 	{
-		check(Context.BlueprintEvent);
-
-		const FMVVMVCompiledFieldPath* CompiledFieldPath = CompileResult.FieldPaths.Find(Context.DelegateFieldPathHandle);
+		const FMVVMVCompiledFieldPath* CompiledFieldPath = CompileResult.FieldPaths.Find(ValidEvent->DelegateFieldPathHandle);
 		if (CompiledFieldPath == nullptr)
 		{
-			AddMessageForEvent(Context.BlueprintEvent, LOCTEXT("CompiledEventFieldPathNotGenerated", "Could not generate the event path."), UMVVMBlueprintViewEvent::EMessageType::Error, FName());
-			bAreEventsValid = false;
+			AddMessageForEvent(ValidEvent->Event.Get(), LOCTEXT("CompiledEventFieldPathNotGenerated", "Could not generate the event path."), EMessageType::Error, FName());
+			bIsCompileStepValid = false;
 			continue;
 		}
 
-		if (Context.GeneratedGraphName.IsNone() || Class->FindFunctionByName(Context.GeneratedGraphName) == nullptr)
+		if (ValidEvent->GeneratedGraphName.IsNone() || Class->FindFunctionByName(ValidEvent->GeneratedGraphName) == nullptr)
 		{
-			AddMessageForEvent(Context.BlueprintEvent, LOCTEXT("CompiledEventFieldPathNotGenerated", "Could not generate the event path."), UMVVMBlueprintViewEvent::EMessageType::Error, FName());
-			bAreEventsValid = false;
+			AddMessageForEvent(ValidEvent->Event.Get(), LOCTEXT("CompiledEventFieldPathNotGenerated", "Could not generate the event path."), EMessageType::Error, FName());
+			bIsCompileStepValid = false;
 			continue;
 		}
 
 		FMVVMViewClass_CompiledEvent& NewBinding = ViewExtension->CompiledEvents.AddDefaulted_GetRef();
 		NewBinding.FieldPath = *CompiledFieldPath;
-		NewBinding.FunctionName = Context.GeneratedGraphName;
-		NewBinding.SourceName = Context.SourceName;
+		NewBinding.FunctionName = ValidEvent->GeneratedGraphName;
+		NewBinding.SourceName = ValidEvent->SourceName;
 	}
-
-	return bAreEventsValid;
 }
 
 
-TValueOrError<FMVVMViewBlueprintCompiler::FCompiledBinding, FText> FMVVMViewBlueprintCompiler::CreateCompiledBinding(const UWidgetBlueprintGeneratedClass* Class, TArrayView<const UE::MVVM::FMVVMConstFieldVariant> GetterFields, TArrayView<const UE::MVVM::FMVVMConstFieldVariant> SetterFields, const UFunction* ConversionFunction, bool bIsComplexBinding)
+TValueOrError<FMVVMViewBlueprintCompiler::FGetFieldsResult, FText> FMVVMViewBlueprintCompiler::GetFields(const UWidgetBlueprintGeneratedClass* Class, const FMVVMBlueprintPropertyPath& PropertyPath) const
 {
-	FCompiledBinding Result;
+	FGetFieldsResult Result;
+	if (!PropertyPath.IsValid())
+	{
+		ensureAlwaysMsgf(false, TEXT("Empty property path found. It should have been catch before."));
+		return MakeError(FText::GetEmpty());
+	}
 
-	if (ConversionFunction != nullptr)
+	auto FindSource = [Self = this](const FName PropertyName, FCompilerBindingSource::EType ExpectedType) -> TValueOrError<TSharedPtr<FCompilerBindingSource>, FText>
 	{
-		TValueOrError<FCompiledBindingLibraryCompiler::FFieldPathHandle, FText> FieldPathResult = BindingLibraryCompiler.AddConversionFunctionFieldPath(Class, ConversionFunction);
-		if (FieldPathResult.HasError())
+		const TSharedRef<FCompilerBindingSource>* Found = Self->NeededBindingSources.FindByPredicate([PropertyName](const TSharedRef<FCompilerBindingSource>& Other) { return Other->Name == PropertyName; });
+		if (Found)
 		{
-			return MakeError(FText::Format(LOCTEXT("CouldNotCreateConversionFunctionFieldPath", "Couldn't create the conversion function field path '{0}'. {1}")
-				, FText::FromString(ConversionFunction->GetPathName())
-				, FieldPathResult.GetError()));
+			if ((*Found)->Type != ExpectedType)
+			{
+				return MakeError(LOCTEXT("NotExpectedSourceType", "Internal error. The source of the path is not of the expected type."));
+			}
+			return MakeValue(*Found);
 		}
-		Result.ConversionFunction = FieldPathResult.StealValue();
-		Result.bIsConversionFunctionComplex = bIsComplexBinding;
+		return MakeValue(TSharedPtr<FCompilerBindingSource>());
+	};
 
-		// Sanity check
-		if (bIsComplexBinding && !BindingHelper::IsValidForComplexRuntimeConversion(ConversionFunction))
+	switch (PropertyPath.GetSource(WidgetBlueprintCompilerContext.WidgetBlueprint()))
+	{
+	case EMVVMBlueprintFieldPathSource::ViewModel:
+	{
+		const FCompilerViewModelCreatorContext* FoundViewModelCreator = ViewModelCreatorContexts.FindByPredicate([ViewModelId = PropertyPath.GetViewModelId()](const FCompilerViewModelCreatorContext& Other){ return Other.ViewModelContext.GetViewModelId() == ViewModelId; });
+		check(FoundViewModelCreator);
+		const FName PropertyName = FoundViewModelCreator->ViewModelContext.GetViewModelName();
+		const FCompilerUserWidgetProperty* FoundUserWidgetProperty = NeededUserWidgetProperties.FindByPredicate([PropertyName](const FCompilerUserWidgetProperty& Other){ return Other.Name == PropertyName; });
+		check(FoundUserWidgetProperty);
+		check(FoundViewModelCreator->ViewModelContext.GetViewModelClass());
+		check(FoundUserWidgetProperty->AuthoritativeClass == FoundViewModelCreator->ViewModelContext.GetViewModelClass());
+
+		TValueOrError<TSharedPtr<FCompilerBindingSource>, FText> FindSourceResult = FindSource(PropertyName, FCompilerBindingSource::EType::ViewModel);
+		if (FindSourceResult.HasError())
 		{
-			return MakeError(LOCTEXT("ConversionFunctionIsNotComplex", "Internal Error. The complex conversion function does not respect the prerequisite."));
+			return MakeError(FindSourceResult.StealError());
+		}
+		if (!FindSourceResult.GetValue().IsValid())
+		{
+			return MakeError(LOCTEXT("ViewModelShouldHaveSource", "Internal error. Viewmodel should have a source."));
 		}
-	}
 
-	if (!Result.bIsConversionFunctionComplex)
+		Result.GeneratedFields = GetFields(Class, PropertyName, PropertyPath.GetFields(Class));
+		Result.Source = FindSourceResult.StealValue();
+		break;
+	}
+	case EMVVMBlueprintFieldPathSource::SelfContext:
 	{
-		// Generate a path to read the value at runtime
-		TValueOrError<TArray<FMVVMConstFieldVariant>, FText> GeneratedField = FieldPathHelper::GenerateFieldPathList(GetterFields, true);
-		if (GeneratedField.HasError())
+		TValueOrError<TSharedPtr<FCompilerBindingSource>, FText> FindSourceResult = FindSource(WidgetBlueprintCompilerContext.WidgetBlueprint()->GetFName(), FCompilerBindingSource::EType::Self);
+		if (FindSourceResult.HasError())
 		{
-			return MakeError(FText::Format(Private::CouldNotCreateSourceFieldPathFormat
-				, ::UE::MVVM::FieldPathHelper::ToText(GetterFields)
-				, GeneratedField.GetError()));
+			return MakeError(FindSourceResult.StealError());
 		}
-
-		TValueOrError<FCompiledBindingLibraryCompiler::FFieldPathHandle, FText> FieldPathResult = BindingLibraryCompiler.AddFieldPath(GeneratedField.GetValue(), true);
-		if (FieldPathResult.HasError())
+		if (!FindSourceResult.GetValue().IsValid())
 		{
-			return MakeError(FText::Format(Private::CouldNotCreateSourceFieldPathFormat
-				, ::UE::MVVM::FieldPathHelper::ToText(GetterFields)
-				, FieldPathResult.GetError()));
+			return MakeError(LOCTEXT("WidgetBlueprintShouldHaveSource", "Internal error. The blueprint should have a source."));
 		}
-		Result.SourceRead = FieldPathResult.StealValue();
-	}
 
+		Result.GeneratedFields = GetFields(Class, FName(), PropertyPath.GetFields(Class));
+		Result.Source = FindSourceResult.StealValue();
+		break;
+	}
+	case EMVVMBlueprintFieldPathSource::Widget:
 	{
-		static const FText CouldNotCreateDestinationFieldPathFormat = LOCTEXT("CouldNotCreateDestinationFieldPath", "Couldn't create the destination field path '{0}'. {1}");
+		FName DestinationWidgetName = PropertyPath.GetWidgetName();
+		check(WidgetNameToWidgetPointerMap.Contains(DestinationWidgetName));
+		checkf(!DestinationWidgetName.IsNone(), TEXT("The destination should have been checked and set bAreSourceContextsValid."));
 
-		TValueOrError<TArray<FMVVMConstFieldVariant>, FText> GeneratedField = FieldPathHelper::GenerateFieldPathList(SetterFields, false);
-		if (GeneratedField.HasError())
-		{
-			return MakeError(FText::Format(CouldNotCreateDestinationFieldPathFormat
-				, ::UE::MVVM::FieldPathHelper::ToText(SetterFields)
-				, GeneratedField.GetError()));
-		}
+		const FCompilerUserWidgetProperty* FoundUserWidgetProperty = NeededUserWidgetProperties.FindByPredicate([DestinationWidgetName](const FCompilerUserWidgetProperty& Other) { return Other.Name == DestinationWidgetName; });
+		check(FoundUserWidgetProperty);
 
-		TValueOrError<FCompiledBindingLibraryCompiler::FFieldPathHandle, FText> FieldPathResult = BindingLibraryCompiler.AddFieldPath(GeneratedField.GetValue(), false);
-		if (FieldPathResult.HasError())
+		TValueOrError<TSharedPtr<FCompilerBindingSource>, FText> FindSourceResult = FindSource(DestinationWidgetName, FCompilerBindingSource::EType::Widget);
+		if (FindSourceResult.HasError())
 		{
-			return MakeError(FText::Format(CouldNotCreateDestinationFieldPathFormat
-				, ::UE::MVVM::FieldPathHelper::ToText(SetterFields)
-				, FieldPathResult.GetError()));
+			return MakeError(FindSourceResult.StealError());
 		}
-		Result.DestinationWrite = FieldPathResult.StealValue();
-	}
 
-	// Generate the binding
-	TValueOrError<FCompiledBindingLibraryCompiler::FBindingHandle, FText> BindingResult = Result.bIsConversionFunctionComplex
-		? BindingLibraryCompiler.AddComplexBinding(Result.DestinationWrite, Result.ConversionFunction)
-		: BindingLibraryCompiler.AddBinding(Result.SourceRead, Result.DestinationWrite, Result.ConversionFunction);
-	if (BindingResult.HasError())
-	{
-		return MakeError(BindingResult.StealError());
+		Result.GeneratedFields = GetFields(Class, DestinationWidgetName, PropertyPath.GetFields(Class));
+		Result.Source = FindSourceResult.StealValue();
+		break;
 	}
-	Result.BindingHandle = BindingResult.StealValue();
-
-	return MakeValue(Result);
+	default:
+		ensureAlwaysMsgf(false, TEXT("Not supported yet."));
+		Result.GeneratedFields = GetFields(Class, FName(), PropertyPath.GetFields(Class));
+		break;
+	}
+	return MakeValue(MoveTemp(Result));
 }
 
 
-const FMVVMViewBlueprintCompiler::FCompilerSourceCreatorContext* FMVVMViewBlueprintCompiler::FindViewModelSource(FGuid Id) const
+TArray<FMVVMConstFieldVariant> FMVVMViewBlueprintCompiler::GetFields(const UClass* Class, FName PropertyName, TArray<FMVVMConstFieldVariant> Properties)
 {
-	return CompilerSourceCreatorContexts.FindByPredicate([Id](const FCompilerSourceCreatorContext& Other)
-		{
-			return Other.Type == ECompilerSourceCreatorType::ViewModel ? Other.ViewModelContext.GetViewModelId() == Id : false;
-		});
+	if (PropertyName.IsNone())
+	{
+		return Properties;
+	}
+
+	check(Class);
+	FMVVMConstFieldVariant NewProperty = BindingHelper::FindFieldByName(Class, FMVVMBindingName(PropertyName));
+	Properties.Insert(NewProperty, 0);
+	return Properties;
 }
 
 
-TValueOrError<FMVVMViewBlueprintCompiler::FBindingSourceContext, FText> FMVVMViewBlueprintCompiler::CreateBindingSourceContext(const UMVVMBlueprintView* BlueprintView, const UWidgetBlueprintGeneratedClass* Class, const FMVVMBlueprintPropertyPath& PropertyPath, bool bIsOneTimeBinding)
+TValueOrError<FMVVMViewBlueprintCompiler::FCreateFieldsResult, FText> FMVVMViewBlueprintCompiler::CreateFieldContext(const UWidgetBlueprintGeneratedClass* Class, const FMVVMBlueprintPropertyPath& PropertyPath, bool bForSourceReading) const
 {
-	if (!PropertyPath.IsValid())
-	{
-		ensureAlways(false);
-		return MakeError(LOCTEXT("EmptyPropertyPath", "Empty property path found. This is ilegal."));
-	}
+	FMVVMViewBlueprintCompiler::FCreateFieldsResult Result;
 
-	FBindingSourceContext Result;
-	switch (PropertyPath.GetSource(WidgetBlueprintCompilerContext.WidgetBlueprint()))
+	// Evaluate the getter/setter path.
+	TValueOrError<FGetFieldsResult, FText> GetFieldResult = GetFields(Class, PropertyPath);
+	if (GetFieldResult.HasError())
 	{
-	case EMVVMBlueprintFieldPathSource::ViewModel:
-	{
-		Result.bIsRootWidget = false;
+		return MakeError(GetFieldResult.StealError());
+	}
 
-		const FMVVMBlueprintViewModelContext* SourceViewModelContext = BlueprintView->FindViewModel(PropertyPath.GetViewModelId());
-		check(SourceViewModelContext);
-		const FName SourceName = SourceViewModelContext->GetViewModelName();
-		Result.UserWidgetPropertyContextIndex = CompilerUserWidgetPropertyContexts.IndexOfByPredicate([SourceName](const FCompilerUserWidgetPropertyContext& Other) { return Other.PropertyName == SourceName; });
-		check(Result.UserWidgetPropertyContextIndex != INDEX_NONE);
+	Result.Source = MoveTemp(GetFieldResult.GetValue().Source);
+	Result.GeneratedFields = MoveTemp(GetFieldResult.GetValue().GeneratedFields);
 
-		Result.SourceClass = CompilerUserWidgetPropertyContexts[Result.UserWidgetPropertyContextIndex].Class;
-		Result.PropertyPath = CreatePropertyPath(Class, CompilerUserWidgetPropertyContexts[Result.UserWidgetPropertyContextIndex].PropertyName, PropertyPath.GetFields(Class));
-		break;
+	if (!IsPropertyPathValid(WidgetBlueprintCompilerContext.WidgetBlueprint(), Result.GeneratedFields))
+	{
+		return MakeError(FText::Format(Private::PropertyPathIsInvalidFormat, PropertyPathToText(Class, BlueprintView.Get(), PropertyPath)));
 	}
 
-	case EMVVMBlueprintFieldPathSource::SelfContext:
+	// Generate the path with property converted to BP function
+	TValueOrError<TArray<FMVVMConstFieldVariant>, FText> SkeletalGeneratedFieldsResult = FieldPathHelper::GenerateFieldPathList(Result.GeneratedFields, bForSourceReading);
+	if (SkeletalGeneratedFieldsResult.HasError() || !IsPropertyPathValid(WidgetBlueprintCompilerContext.WidgetBlueprint(), SkeletalGeneratedFieldsResult.GetValue()))
 	{
-		Result.bIsRootWidget = true;
-		Result.UserWidgetPropertyContextIndex = INDEX_NONE;
-		Result.SourceClass = const_cast<UWidgetBlueprintGeneratedClass*>(Class);
-		Result.PropertyPath = CreatePropertyPath(Class, FName(), PropertyPath.GetFields(Class));
-		break;
+		return MakeError(FText::Format(Private::CouldNotCreateSourceFieldPathFormat, PropertyPathToText(Class, BlueprintView.Get(), PropertyPath), SkeletalGeneratedFieldsResult.GetError()));
 	}
-	case EMVVMBlueprintFieldPathSource::Widget:
-	{
-		Result.bIsRootWidget = false;
+	Result.SkeletalGeneratedFields = SkeletalGeneratedFieldsResult.StealValue();
 
-		const FName SourceName = PropertyPath.GetWidgetName();
-		Result.UserWidgetPropertyContextIndex = CompilerUserWidgetPropertyContexts.IndexOfByPredicate([SourceName](const FCompilerUserWidgetPropertyContext& Other) { return Other.PropertyName == SourceName; });
-		if (!CompilerUserWidgetPropertyContexts.IsValidIndex(Result.UserWidgetPropertyContextIndex))
-		{
-			return MakeError(LOCTEXT("InvalidUserWidgetPropertyContextIndexInternal", "Internal error. UserWidgetPropertyContextIndex is invalid."));
-		}
+	return MakeValue(Result);
+}
 
-		Result.SourceClass = CompilerUserWidgetPropertyContexts[Result.UserWidgetPropertyContextIndex].Class;
-		Result.PropertyPath = CreatePropertyPath(Class, CompilerUserWidgetPropertyContexts[Result.UserWidgetPropertyContextIndex].PropertyName, PropertyPath.GetFields(Class));
-		break;
-	}
 
-	default:
-		ensureAlwaysMsgf(false, TEXT("Not supported yet."));
-	}
+TValueOrError<FMVVMViewBlueprintCompiler::FCompilerNotifyFieldId, FText> FMVVMViewBlueprintCompiler::CreateNotifyFieldId(const UWidgetBlueprintGeneratedClass* Class, const TSharedPtr<FGeneratedReadFieldPathContext>& ReadFieldContext, const FMVVMBlueprintViewBinding& Binding)
+{
+	check(ReadFieldContext->SkeletalGeneratedFields.Num() > 0);
 
 	// The path may contains another INotifyFieldValueChanged
-	TValueOrError<FieldPathHelper::FParsedNotifyBindingInfo, FText> BindingInfoResult = FieldPathHelper::GetNotifyBindingInfoFromFieldPath(Class, Result.PropertyPath);
+	TValueOrError<FieldPathHelper::FParsedNotifyBindingInfo, FText> BindingInfoResult = FieldPathHelper::GetNotifyBindingInfoFromFieldPath(Class, ReadFieldContext->SkeletalGeneratedFields);
 	if (BindingInfoResult.HasError())
 	{
 		return MakeError(BindingInfoResult.StealError());
 	}
+
 	const FieldPathHelper::FParsedNotifyBindingInfo& BindingInfo = BindingInfoResult.GetValue();
-	Result.FieldId = BindingInfo.NotifyFieldId;
+	if (!BindingInfo.NotifyFieldId.IsValid() || ReadFieldContext->OptionalSource == nullptr)
+	{
+		return MakeValue(FCompilerNotifyFieldId());
+	}
+
+
+	FCompilerNotifyFieldId Result;
+	Result.NotificationId = BindingInfo.NotifyFieldId;
+	Result.Source = ReadFieldContext->OptionalSource;
+	Result.ViewModelDynamic.Reset();
 
-	if (BindingInfo.ViewModelIndex < 1 && BindingInfo.NotifyFieldClass && Result.SourceClass)
+	auto GetClassFromField = [](UE::MVVM::FMVVMConstFieldVariant Field) -> const UClass*
 	{
-		if (!Result.SourceClass->IsChildOf(BindingInfo.NotifyFieldClass))
+		if (Field.IsProperty())
+		{
+			const FObjectPropertyBase* ObjectProperty = CastField<const FObjectPropertyBase>(Field.GetProperty());
+			if (ObjectProperty)
+			{
+				return ObjectProperty->PropertyClass;
+			}
+		}
+		else if (Field.IsFunction())
+		{
+			const FObjectPropertyBase* ReturnValue = CastField<const FObjectPropertyBase>(BindingHelper::GetReturnProperty(Field.GetFunction()));
+			if (ReturnValue)
+			{
+				return ReturnValue->PropertyClass;
+			}
+		}
+		return nullptr;
+	};
+
+	// Sanity check
+	{
+		const UClass* ExpectedClass = nullptr;
+		if (BindingInfo.ViewModelIndex < 1 && BindingInfo.NotifyFieldClass)
+		{
+			ExpectedClass = ReadFieldContext->OptionalSource->AuthoritativeClass;
+		}
+		else if (BindingInfo.ViewModelIndex >= 0)
+		{
+			ExpectedClass = GetClassFromField(ReadFieldContext->SkeletalGeneratedFields[BindingInfo.ViewModelIndex]);
+		}
+
+		if (ExpectedClass == nullptr || BindingInfo.NotifyFieldClass == nullptr || !ExpectedClass->IsChildOf(BindingInfo.NotifyFieldClass))
 		{
 			return MakeError(LOCTEXT("InvalidNotifyFieldClassInternal", "Internal error. The viewmodel class doesn't matches."));
 		}
@@ -2458,25 +2750,24 @@ TValueOrError<FMVVMViewBlueprintCompiler::FBindingSourceContext, FText> FMVVMVie
 	// The INotifyFieldValueChanged/viewmodel is not the first and only INotifyFieldValueChanged/viewmodel property path.
 	//Create a new source in PropertyPath creator mode. Create a special binding to update the viewmodel when it changes.
 	//This binding (calling this function) will use the new source.
-	if (BindingInfo.ViewModelIndex >= 1 && !bIsOneTimeBinding)
+	if (BindingInfo.ViewModelIndex >= 1)
 	{
 		if (!GetDefault<UMVVMDeveloperProjectSettings>()->bAllowLongSourcePath)
 		{
 			return MakeError(LOCTEXT("DynamicSourceEntryNotSupport", "Long source entry is not supported. Add the viewmodel manually."));
 		}
 
-		//NB. The for loop order is important. The order is used in UMVVMView::SetViewModelInternal when enabling the bindings
-		int32 SourceCreatorContextIndex = INDEX_NONE;
 		for (int32 DynamicIndex = 1; DynamicIndex <= BindingInfo.ViewModelIndex; ++DynamicIndex)
 		{
-			if (!Result.PropertyPath.IsValidIndex(DynamicIndex))
+			if (!ReadFieldContext->SkeletalGeneratedFields.IsValidIndex(DynamicIndex))
 			{
 				return MakeError(LOCTEXT("DynamicSourceEntryInternalIndex", "Internal error. The source index is not valid."));
 			}
 
 			FName NewSourceName;
-			FName ParentSourceName;
+			FName NewParentSourceName;
 			FString NewSourcePropertyPath;
+			const UClass* NewSourceAuthoritativeClass = nullptr;
 			{
 				TStringBuilder<512> PropertyPathBuilder;
 				TStringBuilder<512> DynamicNameBuilder;
@@ -2487,157 +2778,163 @@ TValueOrError<FMVVMViewBlueprintCompiler::FBindingSourceContext, FText> FMVVMVie
 						PropertyPathBuilder << TEXT('.');
 						DynamicNameBuilder << TEXT('_');
 					}
-					PropertyPathBuilder << Result.PropertyPath[Index].GetName();
-					DynamicNameBuilder << Result.PropertyPath[Index].GetName();
+					PropertyPathBuilder << ReadFieldContext->SkeletalGeneratedFields[Index].GetName();
+					DynamicNameBuilder << ReadFieldContext->SkeletalGeneratedFields[Index].GetName();
 
 					if (Index == DynamicIndex - 1)
 					{
-						ParentSourceName = FName(DynamicNameBuilder.ToString());
+						NewParentSourceName = FName(DynamicNameBuilder.ToString());
 					}
 				}
 
 				NewSourceName = FName(DynamicNameBuilder.ToString());
 				NewSourcePropertyPath = PropertyPathBuilder.ToString();
+
+				const UClass* OwnerSkeletalClass = GetClassFromField(ReadFieldContext->SkeletalGeneratedFields[DynamicIndex]);
+				if (OwnerSkeletalClass == nullptr)
+				{
+					return MakeError(LOCTEXT("DVM_GeneratedFieldInvalid", "Internal error. The GeneratedFiled is invalid."));
+				}
+				NewSourceAuthoritativeClass = OwnerSkeletalClass->GetAuthoritativeClass();
+				if (NewSourceAuthoritativeClass == nullptr)
+				{
+					return MakeError(LOCTEXT("DVM_AuthoritativeClassInvalid", "Internal error. No authoritative class."));
+				}
 			}
 
-			// Did we already create the new source?
-			int32 PreviousSourceCreatorContextIndex = SourceCreatorContextIndex;
-			SourceCreatorContextIndex = CompilerSourceCreatorContexts.IndexOfByPredicate([NewSourceName](const FCompilerSourceCreatorContext& Other)
+			// Does the parent exist
+			TSharedPtr<FCompilerBindingSource> ParentBindingSource;
+			{
+				const TSharedRef<FCompilerBindingSource>* FoundParentBindingSource = NeededBindingSources.FindByPredicate([NewParentSourceName](const TSharedRef<FCompilerBindingSource>& Other)
+					{
+						return Other->Name == NewParentSourceName;
+					});
+				if (FoundParentBindingSource == nullptr)
 				{
-					return Other.Type == ECompilerSourceCreatorType::ViewModelDynamic
-						&& Other.ViewModelContext.GetViewModelName() == NewSourceName;
-				});
-			if (SourceCreatorContextIndex == INDEX_NONE)
+					return MakeError(LOCTEXT("DVM_InvalidParentBindingSource", "Internal error. Can't find the parent binding source."));
+				}
+				ParentBindingSource = *FoundParentBindingSource;
+			}
+
+			// Did we already create the new source. It could be a dynamic or one added by the user
 			{
-				// Create the new source
+				const TSharedRef<FCompilerBindingSource>* FoundBindingSource = NeededBindingSources.FindByPredicate([NewSourceName](const TSharedRef<FCompilerBindingSource>& Other)
+					{
+						return Other->Name == NewSourceName;
+					});
+				if (FoundBindingSource)
 				{
-					const UClass* ViewModelClass = Cast<const UClass>(Result.PropertyPath[DynamicIndex+1].GetOwner());
-					if (!ViewModelClass)
+					if ((*FoundBindingSource)->Type != FCompilerBindingSource::EType::DynamicViewmodel)
 					{
-						return MakeError(FText::GetEmpty());
+						return MakeError(LOCTEXT("DVM_ViewmodelOfSameName", "The dynamic viewmodel cannot be added. There is already a source with that name."));
 					}
-
-					FCompilerSourceCreatorContext SourceCreatorContext;
-					SourceCreatorContext.Type = ECompilerSourceCreatorType::ViewModelDynamic;
-					SourceCreatorContext.DynamicParentSourceName = ParentSourceName;
-					SourceCreatorContext.ViewModelContext = FMVVMBlueprintViewModelContext(ViewModelClass, NewSourceName);
-					SourceCreatorContext.ViewModelContext.bCreateSetterFunction = false;
-					SourceCreatorContext.ViewModelContext.bOptional = false;
-					SourceCreatorContext.ViewModelContext.CreationType = EMVVMBlueprintViewModelContextCreationType::PropertyPath;
-					SourceCreatorContext.ViewModelContext.ViewModelPropertyPath = NewSourcePropertyPath;
-
-					TValueOrError<FCompiledBindingLibraryCompiler::FFieldPathHandle, FText> ReadFieldPathResult = AddObjectFieldPath(BindingLibraryCompiler, Class, SourceCreatorContext.ViewModelContext.ViewModelPropertyPath, SourceCreatorContext.ViewModelContext.GetViewModelClass());
-					if (ReadFieldPathResult.HasError())
+					// is the class the same
+					if (!(*FoundBindingSource)->AuthoritativeClass->IsChildOf(NewSourceAuthoritativeClass))
 					{
-						return MakeError(FText::Format(LOCTEXT("DynamicSourceEntryInvalidPath", "Internal error. {0}."), ReadFieldPathResult.StealError()));
+						return MakeError(LOCTEXT("DVM_ExistingNotSameClass", "Internal error. The viewmodel was already added and is not the same type."));
 					}
-					SourceCreatorContext.ReadPropertyPath = ReadFieldPathResult.StealValue();
-
-					SourceCreatorContextIndex = CompilerSourceCreatorContexts.Add(MoveTemp(SourceCreatorContext));
 				}
 
-				// Create the binding to update the source when it changes.
+				const TSharedRef<FCompilerSourceViewModelDynamicCreatorContext>* FoundViewModelDynamicCreatorContext = SourceViewModelDynamicCreatorContexts.FindByPredicate([NewSourceName](const TSharedRef<FCompilerSourceViewModelDynamicCreatorContext>& Other)
+					{
+						return Other->Source->Name == NewSourceName;
+					});
+				if (FoundViewModelDynamicCreatorContext)
 				{
-					FCompilerBinding NewCompilerBinding;
-					NewCompilerBinding.Type = ECompilerBindingType::ViewModelDynamic;
-					NewCompilerBinding.bSourceIsUserWidget = false;
-					NewCompilerBinding.bIsForwardBinding = true;
-					NewCompilerBinding.bFieldIdNeeded = true;
-					NewCompilerBinding.DynamicViewModelName = NewSourceName;
-
-					if (PreviousSourceCreatorContextIndex == INDEX_NONE)
+					if (FoundBindingSource == nullptr)
 					{
-						NewCompilerBinding.UserWidgetPropertyContextIndex = Result.UserWidgetPropertyContextIndex;
+						return MakeError(LOCTEXT("DVM_BindingSourceShouldExist", "Internal error. The binding source should exist."));
 					}
-					else
+					if ((*FoundViewModelDynamicCreatorContext)->Source != *FoundBindingSource)
 					{
-						NewCompilerBinding.SourceCreatorContextIndex = PreviousSourceCreatorContextIndex;
+						return MakeError(LOCTEXT("DVM_BindingSourceShouldBeTheSame", "Internal error. The source should be the same."));
 					}
-
+					if ((*FoundViewModelDynamicCreatorContext)->ParentSource->Name != NewParentSourceName)
 					{
-						const UClass* OwnerClass = Cast<const UClass>(Result.PropertyPath[DynamicIndex].GetOwner());
-						if (!OwnerClass)
-						{
-							return MakeError(FText::GetEmpty());
-						}
+						return MakeError(LOCTEXT("DVM_BindingSourceParentSameName", "Internal error. The parent name should be the same."));
+					}
+				}
 
-						TValueOrError<FCompiledBindingLibraryCompiler::FFieldIdHandle, FText> AddFieldResult = BindingLibraryCompiler.AddFieldId(OwnerClass, Result.PropertyPath[DynamicIndex].GetName());
-						if (AddFieldResult.HasError())
-						{
-							return MakeError(AddFieldResult.StealError());
-						}
-						NewCompilerBinding.FieldIdHandle = AddFieldResult.StealValue();
+				const FCompilerViewModelCreatorContext* FoundViewModelCreatorContext = ViewModelCreatorContexts.FindByPredicate([NewSourceName](const FCompilerViewModelCreatorContext& Other)
+					{
+						return Other.ViewModelContext.GetViewModelName() == NewSourceName;
+					});
+				if (FoundViewModelCreatorContext)
+				{
+					if (FoundBindingSource == nullptr)
+					{
+						return MakeError(LOCTEXT("DVM_BindingSourceShouldExist", "Internal error. The binding source should exist."));
+					}
+					if (FoundViewModelCreatorContext->Source != *FoundBindingSource)
+					{
+						return MakeError(LOCTEXT("DVM_BindingSourceShouldBeTheSame", "Internal error. The source should be the same."));
+					}
+					if (FoundViewModelDynamicCreatorContext == nullptr)
+					{
+						return MakeError(LOCTEXT("DVM_DynamicCreatorContextShouldExit", "Internal error. The creator context should exist."));
 					}
+					if (FoundViewModelCreatorContext->DynamicContext != *FoundViewModelDynamicCreatorContext)
+					{
+						return MakeError(LOCTEXT("DVM_DynamicContextSHouldBeSame", "Internal error. The creator context should be the same."));
+					}
+					if (FoundViewModelCreatorContext->ViewModelContext.CreationType != EMVVMBlueprintViewModelContextCreationType::PropertyPath)
+					{
+						return MakeError(LOCTEXT("DVM_DynamicContextShouldBePropertyPath", "Internal error. The existing creator context should use a property path."));
+					}
+					if (FoundViewModelCreatorContext->ViewModelContext.ViewModelPropertyPath != NewSourcePropertyPath)
+					{
+						return MakeError(LOCTEXT("DVM_DynamicContextSamePropertyPath", "Internal error. The existing creator context use the same property path."));
+					}
+				}
+
+				if (FoundBindingSource && FoundViewModelCreatorContext == nullptr && FoundViewModelDynamicCreatorContext == nullptr)
+				{
+					return MakeError(LOCTEXT("DVM_MissingDefinition", "Internal error. There are missing definition for the dynamic viewmodel."));
+				}
 
-					CompilerBindings.Add(MoveTemp(NewCompilerBinding));
+				if (FoundBindingSource)
+				{
+					Result.Source = *FoundBindingSource;
+					Result.ViewModelDynamic = *FoundViewModelDynamicCreatorContext;
+					continue; // already exist and correct to use.
 				}
 			}
-		}
 
-		Result.SourceClass = BindingInfo.NotifyFieldClass;
-		Result.FieldId = BindingInfo.NotifyFieldId;
-		Result.UserWidgetPropertyContextIndex = INDEX_NONE;
-		Result.SourceCreatorContextIndex = SourceCreatorContextIndex;
-		Result.bIsRootWidget = false;
-	}
-
-	return MakeValue(Result);
-}
+			if (!NewSourceAuthoritativeClass->ImplementsInterface(UNotifyFieldValueChanged::StaticClass()))
+			{
+				return MakeError(LOCTEXT("DVM_NewDynamicNotViewmodel", "The dynamic viewmodel is not an actual viewmodel."));
+			}
 
+			// Create the new source
+			TSharedRef<FCompilerBindingSource> NewBindingSource = MakeShared<FCompilerBindingSource>();
+			NewBindingSource->AuthoritativeClass = NewSourceAuthoritativeClass;
+			NewBindingSource->Name = NewSourceName;
+			NewBindingSource->Type = FCompilerBindingSource::EType::DynamicViewmodel;
+			NewBindingSource->bIsOptional = ParentBindingSource->bIsOptional;
+			NeededBindingSources.Add(NewBindingSource);
 
-TArray<FMVVMConstFieldVariant> FMVVMViewBlueprintCompiler::CreateBindingDestinationPath(const UMVVMBlueprintView* BlueprintView, const UWidgetBlueprintGeneratedClass* Class, const FMVVMBlueprintPropertyPath& PropertyPath) const
-{
-	if (!PropertyPath.IsValid())
-	{
-		ensureAlwaysMsgf(false, TEXT("Empty property path found. This is legal."));
-		return TArray<FMVVMConstFieldVariant>();
-	}
+			Result.Source = NewBindingSource;
 
-	switch (PropertyPath.GetSource(WidgetBlueprintCompilerContext.WidgetBlueprint()))
-	{
-	case EMVVMBlueprintFieldPathSource::ViewModel:
-	{
-		const FMVVMBlueprintViewModelContext* SourceViewModelContext = BlueprintView->FindViewModel(PropertyPath.GetViewModelId());
-		check(SourceViewModelContext);
-		FName DestinationName = SourceViewModelContext->GetViewModelName();
-		const int32 DestinationVariableContextIndex = CompilerUserWidgetPropertyContexts.IndexOfByPredicate([DestinationName](const FCompilerUserWidgetPropertyContext& Other) { return Other.PropertyName == DestinationName; });
-		check(DestinationVariableContextIndex != INDEX_NONE);
+			TSharedRef<FCompilerSourceViewModelDynamicCreatorContext> NewViewModelDynamic = MakeShared<FCompilerSourceViewModelDynamicCreatorContext>();
+			NewViewModelDynamic->Source = NewBindingSource;
+			NewViewModelDynamic->ParentSource = ParentBindingSource;
+			NewViewModelDynamic->NotificationId = FFieldNotificationId(ReadFieldContext->SkeletalGeneratedFields[DynamicIndex].GetName());
+			SourceViewModelDynamicCreatorContexts.Add(NewViewModelDynamic);
 
-		return CreatePropertyPath(Class, CompilerUserWidgetPropertyContexts[DestinationVariableContextIndex].PropertyName, PropertyPath.GetFields(Class));
-	}
-	case EMVVMBlueprintFieldPathSource::SelfContext:
-	{
-		return CreatePropertyPath(Class, FName(), PropertyPath.GetFields(Class));
-	}
-	case EMVVMBlueprintFieldPathSource::Widget:
-	{
-		FName DestinationName = PropertyPath.GetWidgetName();
-		checkf(!DestinationName.IsNone(), TEXT("The destination should have been checked and set bAreSourceContextsValid."));
-		const int32 DestinationVariableContextIndex = CompilerUserWidgetPropertyContexts.IndexOfByPredicate([DestinationName](const FCompilerUserWidgetPropertyContext& Other) { return Other.PropertyName == DestinationName; });
+			Result.ViewModelDynamic = NewViewModelDynamic;
 
-		if (ensureAlwaysMsgf(DestinationVariableContextIndex != INDEX_NONE, TEXT("Could not find source context for destination '%s'"), *DestinationName.ToString()))
-		{
-			return CreatePropertyPath(Class, CompilerUserWidgetPropertyContexts[DestinationVariableContextIndex].PropertyName, PropertyPath.GetFields(Class));
+			FCompilerViewModelCreatorContext& NewViewModelCreatorContext = ViewModelCreatorContexts.AddDefaulted_GetRef();
+			NewViewModelCreatorContext.ViewModelContext = FMVVMBlueprintViewModelContext(NewSourceAuthoritativeClass, NewSourceName);
+			NewViewModelCreatorContext.ViewModelContext.bCreateSetterFunction = false;
+			NewViewModelCreatorContext.ViewModelContext.bOptional = NewBindingSource->bIsOptional;
+			NewViewModelCreatorContext.ViewModelContext.CreationType = EMVVMBlueprintViewModelContextCreationType::PropertyPath;
+			NewViewModelCreatorContext.ViewModelContext.ViewModelPropertyPath = NewSourcePropertyPath;
+			NewViewModelCreatorContext.Source = NewBindingSource;
+			NewViewModelCreatorContext.DynamicContext = NewViewModelDynamic;
 		}
-		return TArray<FMVVMConstFieldVariant>();
-	}
-	default:
-		ensureAlwaysMsgf(false, TEXT("Not supported yet."));
-		return CreatePropertyPath(Class, FName(), PropertyPath.GetFields(Class));
-	}
-}
-
-TArray<FMVVMConstFieldVariant> FMVVMViewBlueprintCompiler::CreatePropertyPath(const UClass* Class, FName PropertyName, TArray<FMVVMConstFieldVariant> Properties)
-{
-	if (PropertyName.IsNone())
-	{
-		return Properties;
 	}
 
-	check(Class);
-	FMVVMConstFieldVariant NewProperty = BindingHelper::FindFieldByName(Class, FMVVMBindingName(PropertyName));
-	Properties.Insert(NewProperty, 0);
-	return Properties;
+	return MakeValue(Result);
 }
 
 
