MVVM: Big rework of the MVVM view compiler. It add supports for conversion function to go from complex to simple. Supports pin split on simple conversion function. Add ability to to support setter fun

diff --git a/Engine/Plugins/Runtime/ModelViewViewModel/Source/ModelViewViewModelBlueprint/Private/MVVMViewBlueprintCompiler.h b/Engine/Plugins/Runtime/ModelViewViewModel/Source/ModelViewViewModelBlueprint/Private/MVVMViewBlueprintCompiler.h
--- a/Engine/Plugins/Runtime/ModelViewViewModel/Source/ModelViewViewModelBlueprint/Private/MVVMViewBlueprintCompiler.h
+++ b/Engine/Plugins/Runtime/ModelViewViewModel/Source/ModelViewViewModelBlueprint/Private/MVVMViewBlueprintCompiler.h
@@ -3,9 +3,12 @@
 #pragma once
 
 #include "Bindings/MVVMCompiledBindingLibraryCompiler.h"
+#include "MVVMBlueprintView.h"
 #include "MVVMBlueprintViewEvent.h"
 #include "MVVMBlueprintViewModelContext.h"
+#include "Templates/ValueOrError.h"
 #include "Types/MVVMFieldVariant.h"
+#include "UObject/StrongObjectPtr.h"
 #include "WidgetBlueprintCompiler.h"
 
 struct FMVVMBlueprintPropertyPath;
@@ -15,6 +18,8 @@ struct FMVVMBlueprintViewBinding;
 class FWidgetBlueprintCompilerContext;
 class UEdGraph;
 class UMVVMBlueprintView;
+class UMVVMBlueprintViewConversionFunction;
+class UMVVMBlueprintViewEvent;
 class UMVVMViewClass;
 class UWidgetBlueprintGeneratedClass;
 namespace  UE::MVVM
@@ -28,15 +33,19 @@ namespace UE::MVVM::Private
 struct FMVVMViewBlueprintCompiler
 {
 private:
-	struct FCompilerUserWidgetPropertyContext;
-	struct FCompilerSourceCreatorContext;
-	struct FCompiledBinding;
+	struct FCompilerBindingSource;
+	struct FCompilerUserWidgetProperty;
+	struct FCompilerViewModelSetter;
+	struct FCompilerViewModelCreatorContext;
+	struct FCompilerSourceViewModelDynamicCreatorContext;
+	struct FCompilerNotifyFieldId;
+	struct FGeneratedReadFieldPathContext;
+	struct FGeneratedWriteFieldPathContext;
 	struct FCompilerBinding;
-	struct FBindingSourceContext;
-	struct FBindingDestinationContext;
+	struct FCompilerEvent;
 
 public:
-	FMVVMViewBlueprintCompiler(FWidgetBlueprintCompilerContext& InCreationContext);
+	FMVVMViewBlueprintCompiler(FWidgetBlueprintCompilerContext& InCreationContext, UMVVMBlueprintView* BlueprintView);
 
 	FWidgetBlueprintCompilerContext& GetCompilerContext()
 	{
@@ -46,236 +55,321 @@ public:
 	void AddExtension(UWidgetBlueprintGeneratedClass* Class, UMVVMViewClass* ViewExtension);
 	void CleanOldData(UWidgetBlueprintGeneratedClass* ClassToClean, UObject* OldCDO);
 
+	/** Generate variable and public function in the Skeleton class and the generated class */
+	void CreateVariables(const FWidgetBlueprintCompilerContext::FCreateVariableContext& Context);
 	/** Generate function that are hidden from the user (not on the Skeleton class). */
-	void CreateFunctions(UMVVMBlueprintView* BlueprintView);
-	/** Generate variable in the Skeleton class */
-	void CreateVariables(const FWidgetBlueprintCompilerContext::FCreateVariableContext& Context, UMVVMBlueprintView* BlueprintView);
+	void CreateFunctions(const FWidgetBlueprintCompilerContext::FCreateFunctionContext& Context);
 
 	/** Add all the field path and the bindings to the library compiler. */
-	bool PreCompile(UWidgetBlueprintGeneratedClass* Class, UMVVMBlueprintView* BlueprintView);
+	bool PreCompile(UWidgetBlueprintGeneratedClass* Class);
 	/** Compile the library and fill the view and viewclass */
-	bool Compile(UWidgetBlueprintGeneratedClass* Class, UMVVMBlueprintView* BlueprintView, UMVVMViewClass* ViewExtension);
+	bool Compile(UWidgetBlueprintGeneratedClass* Class, UMVVMViewClass* ViewExtension);
 
 
 	static void TestGenerateSetter(const UBlueprint* Context, FStringView ObjectName, FStringView FieldPath, FStringView FunctionName);
 
 private:
-	void CreateWidgetMap(const FWidgetBlueprintCompilerContext::FCreateVariableContext& Context, UMVVMBlueprintView* BlueprintView);
-	void CreateSourceLists(const FWidgetBlueprintCompilerContext::FCreateVariableContext& Context, UMVVMBlueprintView* BlueprintView);
-	void CreateFunctionsDeclaration(const FWidgetBlueprintCompilerContext::FCreateVariableContext& Context, UMVVMBlueprintView* BlueprintView);
-	void CreateIntermediateGraphFunctions(const FWidgetBlueprintCompilerContext::FCreateVariableContext& Context, UMVVMBlueprintView* BlueprintView);
-	void CreateBindingDestinationContexts(UMVVMBlueprintView* BlueprintView);
-
-	bool PreCompileBindingSources(UWidgetBlueprintGeneratedClass* Class, UMVVMBlueprintView* BlueprintView);
-	bool CompileBindingSources(const FCompiledBindingLibraryCompiler::FCompileResult& CompileResult, UWidgetBlueprintGeneratedClass* Class, UMVVMBlueprintView* BlueprintView, UMVVMViewClass* ViewExtension);
-
-	bool PreCompileSourceCreators(UWidgetBlueprintGeneratedClass* Class, UMVVMBlueprintView* BlueprintView);
-	bool CompileSourceCreators(const FCompiledBindingLibraryCompiler::FCompileResult& CompileResult, UWidgetBlueprintGeneratedClass* Class, UMVVMBlueprintView* BlueprintView, UMVVMViewClass* ViewExtension);
-
-	bool PreCompileBindings(UWidgetBlueprintGeneratedClass* Class, UMVVMBlueprintView* BlueprintView);
-	bool CompileBindings(const FCompiledBindingLibraryCompiler::FCompileResult& CompileResult, UWidgetBlueprintGeneratedClass* Class, UMVVMBlueprintView* BlueprintView, UMVVMViewClass* ViewExtension);
-	
-	bool PreCompileEvents(UWidgetBlueprintGeneratedClass* Class, UMVVMBlueprintView* BlueprintView);
-	bool CompileEvents(const FCompiledBindingLibraryCompiler::FCompileResult& CompileResult, UWidgetBlueprintGeneratedClass* Class, UMVVMBlueprintView* BlueprintView, UMVVMViewClass* ViewExtension);
-
-	const FCompilerSourceCreatorContext* FindViewModelSource(FGuid Id) const;
-
-	void AddMessageForBinding(FMVVMBlueprintViewBinding& Binding, UMVVMBlueprintView* BlueprintView, const FText& MessageText, EBindingMessageType MessageType, FName ArgumentName = FName()) const;
-	void AddMessageForEvent(UMVVMBlueprintViewEvent* Event, const FText& MessageText, UMVVMBlueprintViewEvent::EMessageType MessageType, FName ArgumentName) const;
-	void AddErrorForViewModel(const FMVVMBlueprintViewModelContext& ViewModel, const FText& Message) const;
-
-	TValueOrError<FBindingSourceContext, FText> CreateBindingSourceContext(const UMVVMBlueprintView* BlueprintView, const UWidgetBlueprintGeneratedClass* Class, const FMVVMBlueprintPropertyPath& PropertyPath, bool bIsOneTimeBinding);
-	TArray<FMVVMConstFieldVariant> CreateBindingDestinationPath(const UMVVMBlueprintView* BlueprintView, const UWidgetBlueprintGeneratedClass* Class, const FMVVMBlueprintPropertyPath& PropertyPath) const;
-	TValueOrError<FCompiledBinding, FText> CreateCompiledBinding(const UWidgetBlueprintGeneratedClass* Class, TArrayView<const UE::MVVM::FMVVMConstFieldVariant> GetterFields, TArrayView<const UE::MVVM::FMVVMConstFieldVariant> SetterFields, const UFunction* ConversionFunction, bool bIsComplexBinding);
+	bool AreStepsValid() const
+	{
+		return bIsCreateVariableStepValid && bIsCreateFunctionsStepValid && bIsPreCompileStepValid && bIsCompileStepValid;
+	}
 
-	static TArray<FMVVMConstFieldVariant> CreatePropertyPath(const UClass* Class, FName PropertyName, TArray<FMVVMConstFieldVariant> Properties);
-	static bool IsPropertyPathValid(const UBlueprint* Context, TArrayView<const FMVVMConstFieldVariant> PropertyPath);
-	static bool CanBeSetInNative(TArrayView<const FMVVMConstFieldVariant> PropertyPath);
+	// CreateVariables
+	void CreateWidgetMap(const FWidgetBlueprintCompilerContext::FCreateVariableContext& Context);
+	void CreateBindingList(const FWidgetBlueprintCompilerContext::FCreateVariableContext& Context);
+	void CreateEventList(const FWidgetBlueprintCompilerContext::FCreateVariableContext& Context);
+	void CreateRequiredProperties(const FWidgetBlueprintCompilerContext::FCreateVariableContext& Context);
+	void CreatePublicFunctionsDeclaration(const FWidgetBlueprintCompilerContext::FCreateVariableContext& Context);
+
+	// CreateFuntions
+	void CategorizeBindings(const FWidgetBlueprintCompilerContext::FCreateFunctionContext& Context);
+	void CategorizeEvents(const FWidgetBlueprintCompilerContext::FCreateFunctionContext& Context);
+	void CreateWriteFieldContexts(const FWidgetBlueprintCompilerContext::FCreateFunctionContext& Context);
+	void CreateViewModelSetters(const FWidgetBlueprintCompilerContext::FCreateFunctionContext& Context);
+	void CreateIntermediateGraphFunctions(const FWidgetBlueprintCompilerContext::FCreateFunctionContext& Context);
+
+	// PreCompile
+	void FixCompilerBindingSelfSource(UWidgetBlueprintGeneratedClass* Class);
+	void AddWarningForPropertyWithMVVMAndLegacyBinding(UWidgetBlueprintGeneratedClass* Class);
+	void FixWriteFieldPathContext(UWidgetBlueprintGeneratedClass* Class);
+	void CreateReadFieldContexts(UWidgetBlueprintGeneratedClass* Class);
+	void CreateCreatorContentFromBindingSource(UWidgetBlueprintGeneratedClass* Class);
+	void PreCompileViewModelCreatorContexts(UWidgetBlueprintGeneratedClass* Class);
+	void PreCompileBindings(UWidgetBlueprintGeneratedClass* Class);
+	void PreCompileEvents(UWidgetBlueprintGeneratedClass* Class);
+
+	// Compile
+	void CompileViewModelCreatorContexts(const FCompiledBindingLibraryCompiler::FCompileResult& CompileResult, UWidgetBlueprintGeneratedClass* Class, UMVVMViewClass* ViewExtension);
+	void CompileBindings(const FCompiledBindingLibraryCompiler::FCompileResult& CompileResult, UWidgetBlueprintGeneratedClass* Class, UMVVMViewClass* ViewExtension);
+	void CompileEvents(const FCompiledBindingLibraryCompiler::FCompileResult& CompileResult, UWidgetBlueprintGeneratedClass* Class, UMVVMViewClass* ViewExtension);
 
 private:
 	/**
-	 * Describe the Property that need to be added (if it doesn't already exist)
+	 * List of all the sources needed by the view to register the bindings.
+	 * They could be viewmodel, widget or any properties on the UserWidget.
+	 * They could also be a viewmodel in a long path.
+	 * It may not have an associated property (dynamicviewmodel).
+	 * It may only have OneTime binding
+	 */
+	struct FCompilerBindingSource
+	{
+		enum class EType
+		{
+			ViewModel = 0,
+			DynamicViewmodel = 1,
+			Widget = 2,
+			Self = 3,
+		};
+		const UClass* AuthoritativeClass = nullptr;
+		FName Name;
+		EType Type;
+		bool bIsOptional = false;
+	};
+	TArray<TSharedRef<FCompilerBindingSource>> NeededBindingSources;
+
+	/**
+	 * Describe a Property that need to be added (if it doesn't already exist).
 	 * They can be a Widget, a viewmodel, or any object owned by the UserWidget.
+	 * They can be source or destination.
 	 */
-	struct FCompilerUserWidgetPropertyContext
+	struct FCompilerUserWidgetProperty
 	{
-		UClass* Class = nullptr;
-		FName PropertyName;
+		enum class ECreationType
+		{
+			None,
+			CreateIfDoesntExist,
+			CreateOnlyIfDoesntExist,
+		};
+		UClass* AuthoritativeClass = nullptr;
+		FName Name;
 		FText DisplayName;
 		FString CategoryName;
 		FString BlueprintSetter;
-		FMVVMConstFieldVariant Field;
-		// If the class is a viewmodel on the userwidget, what is the id of that viewmodel
-		FGuid ViewModelId;
-
+		ECreationType CreationType = ECreationType::None;
 		bool bExposeOnSpawn = false;
 		bool bPrivate = false;
+		bool bReadOnly = false;
+
+		const FProperty* Property = nullptr;
 	};
-	TArray<FCompilerUserWidgetPropertyContext> CompilerUserWidgetPropertyContexts;
+	TArray<FCompilerUserWidgetProperty> NeededUserWidgetProperties;
 
-	/** 
-	 * Describe the data initialize the view's properties/viewmodels.
+	/**
+	 * Describe a ViewModel generated setter function.
 	 */
-	enum class ECompilerSourceCreatorType
+	struct FCompilerViewModelSetter
 	{
-		ViewModel, // added manually
-		ViewModelDynamic, // added by a long binding chain
-	};
-	struct FCompilerSourceCreatorContext
-	{
-		FMVVMBlueprintViewModelContext ViewModelContext;
-		FCompiledBindingLibraryCompiler::FFieldPathHandle ReadPropertyPath;
-		ECompilerSourceCreatorType Type = ECompilerSourceCreatorType::ViewModel;
-		FString SetterFunctionName;
+		const UClass* Class = nullptr;
+		FName PropertyName;
+		FString BlueprintSetter;
+		FText DisplayName;
+
 		UEdGraph* SetterGraph = nullptr;
-		FName DynamicParentSourceName;
 	};
-	TArray<FCompilerSourceCreatorContext> CompilerSourceCreatorContexts;
+	TArray<FCompilerViewModelSetter> ViewModelSettersToGenerate;
 
-	/** The compiled binding from the binding compiler. */
-	struct FCompiledBinding
+	/** 
+	 * Describe the data to initialize the view's properties/viewmodels.
+	 */
+	struct FCompilerViewModelCreatorContext
 	{
-		FCompiledBindingLibraryCompiler::FBindingHandle BindingHandle;
+		FMVVMBlueprintViewModelContext ViewModelContext;
+		TSharedPtr<FCompilerBindingSource> Source;
 
-		// At runtime, OnFieldValueChanged
-		//if the ConversionFunction exist
-		// then we read the Source
-		// then we call ConversionFunction with the source as an input
-		// then we copy the result of the ConversionFunction to DestinationWrite
-		//else if the ConversionFunction is complex
-		// then we call ConversionFunction without input (SourceRead can be invalid)
-		// then we copy the result of the ConversionFunction to DestinationWrite
-		//else
-		// then we copy SourceRead to DestiantionWrite.
-		FCompiledBindingLibraryCompiler::FFieldPathHandle SourceRead;
-		FCompiledBindingLibraryCompiler::FFieldPathHandle DestinationWrite;
-		FCompiledBindingLibraryCompiler::FFieldPathHandle ConversionFunction;
-
-		bool bIsConversionFunctionComplex = false;
+		TSharedPtr<FCompilerSourceViewModelDynamicCreatorContext> DynamicContext;
+		FCompiledBindingLibraryCompiler::FFieldPathHandle ReadPropertyPathHandle;
 	};
+	TArray<FCompilerViewModelCreatorContext> ViewModelCreatorContexts;
 
-	/** 
-	 * Describe a binding for the compiler.
-	 * 2 FCompilerBinding will be generated if the binding is TwoWays.
+	/**
+	 * Describe the data to initialize the view's properties/widget.
 	 */
-	enum class ECompilerBindingType
+	struct FCompilerWidgetCreatorContext
 	{
-		PropertyBinding, // normal binding
-		ViewModelDynamic, // added by a long binding chain
+		TSharedPtr<FCompilerBindingSource> Source;
+		bool bSelfReference = false;
+
+		FCompiledBindingLibraryCompiler::FFieldPathHandle ReadPropertyPathHandle;
 	};
-	struct FCompilerBinding
+	TArray<FCompilerWidgetCreatorContext> WidgetCreatorContexts;
+
+	/**
+	 * Describe the data to initialize a viewmodel in a long path.
+	 * The viewmodel is not added in the BlueprintView but it needs to be dynamic added.
+	 */
+	struct FCompilerSourceViewModelDynamicCreatorContext
 	{
-		ECompilerBindingType Type = ECompilerBindingType::PropertyBinding;
-		int32 BindingIndex = INDEX_NONE;
-		int32 UserWidgetPropertyContextIndex = INDEX_NONE;
-		int32 SourceCreatorContextIndex = INDEX_NONE;
-		int32 ComplexConversionFunctionContextIndex = INDEX_NONE;
-		bool bSourceIsUserWidget = false;
-		bool bFieldIdNeeded = false;
-		bool bIsForwardBinding = false;
-		bool bNeedsValidSource = true;	// if false, the binding is a onetime or executed manually (onclick event).
-
-		FCompiledBindingLibraryCompiler::FFieldIdHandle FieldIdHandle;
-		FCompiledBinding CompiledBinding;
-		FName DynamicViewModelName;
+		TSharedPtr<FCompilerBindingSource> Source;
+		TSharedPtr<FCompilerBindingSource> ParentSource; // a dynamic always has a parent
+
+		FFieldNotificationId NotificationId;
+		FCompiledBindingLibraryCompiler::FFieldIdHandle NotificationIdLibraryCompilerHandle;
 	};
-	TArray<FCompilerBinding> CompilerBindings;
+	TArray<TSharedRef<FCompilerSourceViewModelDynamicCreatorContext>> SourceViewModelDynamicCreatorContexts;
 
 	/**
-	 * Binding that do not requires a source to execute.
-	 * They could be a OneTime ComplexConversionFunction or an event (onclick).
+	 * The field id for a specific ReadFieldPathContext
 	 */
-	struct FSimpleBindingContext
+	struct FCompilerNotifyFieldId
 	{
-		int32 BindingIndex = INDEX_NONE;
-		bool bIsForwardBinding = false;
+		TArray<FGuid> BindingEditorKeys;
+		TArray<FGuid> EventKeys;
+
+		FFieldNotificationId NotificationId;
+		TSharedPtr<FCompilerBindingSource> Source;
+		TSharedPtr<FCompilerSourceViewModelDynamicCreatorContext> ViewModelDynamic;
 
-		int32 ComplexConversionFunctionContextIndex = INDEX_NONE;
+		FCompiledBindingLibraryCompiler::FFieldIdHandle LibraryCompilerHandle;
 	};
-	TArray<FSimpleBindingContext> SimpleBindingContexts;
+	TArray<TSharedPtr<FCompilerNotifyFieldId>> NotificationFields;
 
 	/**
-	 * Source path of a binding that contains a FieldId that we can register/bind to
+	 * The source path we need to read from.
+	 * Can be any EMVVMBindingMode (OneTime, OneWay, ...)
 	 */
-	struct FBindingSourceContext
+	struct FGeneratedReadFieldPathContext
 	{
-		int32 BindingIndex = INDEX_NONE;
-		bool bIsForwardBinding = false;
-
-		const UClass* SourceClass = nullptr;
-		// The property that are registering to.
-		FFieldNotificationId FieldId;
-		// The path always start with the UserWidget as self.
-		// Viewmodel.Field.SubProperty.SubProperty
-		// or Widget.Field.SubProperty.SubProperty
-		// or Field.SubProperty.SubProperty
-		TArray<UE::MVVM::FMVVMConstFieldVariant> PropertyPath;
-
-		// Complex binding do not need to add the PropertyPath to the binding info.
-		int32 ComplexConversionFunctionContextIndex = INDEX_NONE;
-		// The source if it's a property
-		int32 UserWidgetPropertyContextIndex = INDEX_NONE;
-		// The source if it's a property
-		int32 SourceCreatorContextIndex = INDEX_NONE;
-		// The source is the UserWidget itself
-		bool bIsRootWidget = false;
+		TArray<TWeakPtr<FCompilerBinding>> UsedByBindings;
+		TArray<TWeakPtr<FCompilerEvent>> UsedByEvents;
+
+		TSharedPtr<FCompilerBindingSource> OptionalSource; // Can be invalid if it's an event
+		TArray<UE::MVVM::FMVVMConstFieldVariant> GeneratedFields; // the string path converted into field
+		TArray<UE::MVVM::FMVVMConstFieldVariant> SkeletalGeneratedFields; // the field path converted with getter and setter
+
+		TSharedPtr<FCompilerNotifyFieldId> NotificationField;
+		FCompiledBindingLibraryCompiler::FFieldPathHandle LibraryCompilerHandle;
 	};
-	TArray<FBindingSourceContext> BindingSourceContexts;
+	TArray<TSharedRef<FGeneratedReadFieldPathContext>> GeneratedReadFieldPaths;
 
 	/**
-	 * Destination path of a binding.
+	 * Destination path we need to write to.
+	 * Only if the bindings/events have a destination.
 	 * The info needs to be validate before we generate the functions list.
 	 */
-	struct FBindingDestinationContext
+	struct FGeneratedWriteFieldPathContext
 	{
-		int32 BindingIndex = INDEX_NONE;
-		bool bIsForwardBinding = false;
+		TArray<TWeakPtr<FCompilerBinding>> UsedByBindings;
+		TArray<TWeakPtr<FCompilerEvent>> UsedByEvents;
 
+		TArray<UE::MVVM::FMVVMConstFieldVariant> GeneratedFields; // the string path converted into field
+		TArray<UE::MVVM::FMVVMConstFieldVariant> SkeletalGeneratedFields; // the field path converted with getter and setter
+		EMVVMBlueprintFieldPathSource GeneratedFrom = EMVVMBlueprintFieldPathSource::None;
 		bool bCanBeSetInNative = true;
-		FName GeneratedFunctionName;
+		bool bUseByNativeBinding = false;
 
+		FName GeneratedFunctionName;
+		FCompiledBindingLibraryCompiler::FFieldPathHandle LibraryCompilerHandle;
 	};
-	TArray<FBindingDestinationContext> BindingDestinationContexts;
+	TArray<TSharedRef<FGeneratedWriteFieldPathContext>> GeneratedWriteFieldPaths;
 
 	/**
-	 * Complex Conversion function shared data
+	 * The list of all the valid binding to iterates on.
 	 */
-	struct FComplexConversionFunctionContext
+	struct FCompilerBinding
 	{
-		int32 BindingIndex = INDEX_NONE;
-		bool bIsForwardBinding = false;
-
-		// More than one binding can use the same conversion function.
-		// Only one should execute the binding at initialization.
-		bool bExecAtInitGenerated = false;
-		// More than one binding can use the same conversion function.
-		// At least one must have a valid source to trigger the binding.
-		//If not, then it should be a FSimpleBindingContext
-		bool bHasValidFieldId = false;
-		bool bNeedsValidSource = true;
+		struct FKey
+		{
+			int32 ViewBindingIndex = INDEX_NONE;
+			bool bIsForwardBinding = false;
+			FKey() = default;
+			FKey(int32 InIndex, bool bInForward)
+				: ViewBindingIndex(InIndex)
+				, bIsForwardBinding(bInForward)
+			{}
+			bool operator==(const FKey& Other) const
+			{
+				return ViewBindingIndex == Other.ViewBindingIndex && bIsForwardBinding == Other.bIsForwardBinding;
+			}
+		};
+		enum class EType : int32
+		{
+			Unknown = 0, // was not evaluated yet
+			Invalid = -1, // evaluated and not valid
+			Assignment = 1, //Destination=Source
+			SimpleConversionFunction = 2, //Destination=Function(Source)
+			ComplexConversionFunction = 3, //Destination=Function(SourceA, SourceB)
+			//Function = 4, //Function(SourceA, SourceB)
+		};
+		FKey Key;
+		EType Type = EType::Unknown;
+		bool bIsOneTimeBinding = false;
+
+		TArray<TSharedPtr<FGeneratedReadFieldPathContext>> ReadPaths;
+		TSharedPtr<FGeneratedWriteFieldPathContext> WritePath;
+		TWeakObjectPtr<UMVVMBlueprintViewConversionFunction> ConversionFunction = nullptr;
+
+		FCompiledBindingLibraryCompiler::FBindingHandle BindingHandle;
+		FCompiledBindingLibraryCompiler::FFieldPathHandle ConversionFunctionHandle;
 	};
-	TArray<FComplexConversionFunctionContext> ComplexConversionFunctionContexts;
+	TArray<TSharedRef<FCompilerBinding>> ValidBindings;
 
 	/**
-	 * Context for multicast delegate
+	 * The list of all the valid binding to iterates on.
 	 */
-	struct FEventDelegateContext
+	struct FCompilerEvent
 	{
-		UMVVMBlueprintViewEvent* BlueprintEvent = nullptr;
+		enum class EType : int32
+		{
+			Unknown = 0, // was not evaluated yet
+			Invalid = -1, // evaluated and not valid
+			Valid = 1,
+		};
+		TWeakObjectPtr<UMVVMBlueprintViewEvent> Event = nullptr;
+		EType Type = EType::Unknown;
+
+		TArray<TSharedPtr<FGeneratedReadFieldPathContext>> ReadPaths;
+		TSharedPtr<FGeneratedWriteFieldPathContext> WritePath;
+
+		TSharedPtr<FGeneratedWriteFieldPathContext> DelegateFieldPath;
 
-		FCompiledBindingLibraryCompiler::FFieldPathHandle DelegateFieldPathHandle;
 		FName GeneratedGraphName;
 		FName SourceName;
+		FCompiledBindingLibraryCompiler::FFieldPathHandle DelegateFieldPathHandle;
 	};
-	TArray<FEventDelegateContext> EventDelegateContexts;
+	TArray<TSharedRef<FCompilerEvent>> ValidEvents;
 
 	TMap<FName, UWidget*> WidgetNameToWidgetPointerMap;
 	FWidgetBlueprintCompilerContext& WidgetBlueprintCompilerContext;
+	TStrongObjectPtr<UMVVMBlueprintView> BlueprintView = nullptr;
 	FCompiledBindingLibraryCompiler BindingLibraryCompiler;
-	bool bAreSourcesCreatorValid = true;
-	bool bAreSourceContextsValid = true;
-	bool bAreBindingsValid = true;
-	bool bAreEventsValid = true;
+	bool bIsCreateVariableStepValid = true;
+	bool bIsCreateFunctionsStepValid = true;
+	bool bIsPreCompileStepValid = true;
+	bool bIsCompileStepValid = true;
+
+private:
+	enum class EMessageType
+	{
+		Info = 0, Warning = 1, Error = 2
+	};
+	void AddMessage(const FText& MessageText, EMessageType MessageType) const;
+	void AddMessages(TArrayView<TWeakPtr<FCompilerBinding>> Bindings, TArrayView<TWeakPtr<FCompilerEvent>> Events, const FText& MessageText, EMessageType MessageType) const;
+	void AddMessageForBinding(const TSharedPtr<FCompilerBinding>& Binding, const FText& MessageText, EMessageType MessageType, FName ArgumentName) const;
+	void AddMessageForBinding(const FMVVMBlueprintViewBinding& Binding, const FText& MessageText, EMessageType MessageType, FName ArgumentName) const;
+	void AddMessageForEvent(const TSharedPtr<FCompilerEvent>& Event, const FText& MessageText, EMessageType MessageType, FName ArgumentName) const;
+	void AddMessageForEvent(const UMVVMBlueprintViewEvent* Event, const FText& MessageText, EMessageType MessageType, FName ArgumentName) const;
+	void AddMessageForViewModel(const FMVVMBlueprintViewModelContext& ViewModel, const FText& Message, EMessageType MessageType) const;
+	void AddMessageForViewModel(const FText& ViewModelDisplayName, const FText& Message, EMessageType MessageType) const;
+
+	struct FGetFieldsResult
+	{
+		TSharedPtr<FCompilerBindingSource> Source;
+		TArray<FMVVMConstFieldVariant> GeneratedFields;
+	};
+	struct FCreateFieldsResult
+	{
+		TSharedPtr<FCompilerBindingSource> Source;
+		TArray<UE::MVVM::FMVVMConstFieldVariant> GeneratedFields;
+		TArray<UE::MVVM::FMVVMConstFieldVariant> SkeletalGeneratedFields;
+	};
+	TValueOrError<FGetFieldsResult, FText> GetFields(const UWidgetBlueprintGeneratedClass* Class, const FMVVMBlueprintPropertyPath& PropertyPath) const;
+	TValueOrError<FCreateFieldsResult, FText> CreateFieldContext(const UWidgetBlueprintGeneratedClass* Class, const FMVVMBlueprintPropertyPath& PropertyPath, bool bForSourceReading) const;
+	TValueOrError<FCompilerNotifyFieldId, FText> CreateNotifyFieldId(const UWidgetBlueprintGeneratedClass* Class, const TSharedPtr<FGeneratedReadFieldPathContext>& ReadFieldContext, const FMVVMBlueprintViewBinding& Binding);
+
+
+	static TArray<FMVVMConstFieldVariant> GetFields(const UClass* Class, FName PropertyName, TArray<FMVVMConstFieldVariant> Properties);
+	static bool IsPropertyPathValid(const UBlueprint* Context, TArrayView<const FMVVMConstFieldVariant> PropertyPath);
+	static bool CanBeSetInNative(TArrayView<const FMVVMConstFieldVariant> PropertyPath);
 };
 
 } //namespace
diff --git a/Engine/Plugins/Runtime/ModelViewViewModel/Source/ModelViewViewModelBlueprint/Private/MVVMWidgetBlueprintExtension_View.cpp b/Engine/Plugins/Runtime/ModelViewViewModel/Source/ModelViewViewModelBlueprint/Private/MVVMWidgetBlueprintExtension_View.cpp
--- a/Engine/Plugins/Runtime/ModelViewViewModel/Source/ModelViewViewModelBlueprint/Private/MVVMWidgetBlueprintExtension_View.cpp
+++ b/Engine/Plugins/Runtime/ModelViewViewModel/Source/ModelViewViewModelBlueprint/Private/MVVMWidgetBlueprintExtension_View.cpp
@@ -104,7 +104,7 @@ void UMVVMWidgetBlueprintExtension_View::HandleBeginCompilation(FWidgetBlueprint
 	if (BlueprintView)
 	{
 		BlueprintView->ResetBindingMessages();
-		CurrentCompilerContext = MakePimpl<UE::MVVM::Private::FMVVMViewBlueprintCompiler>(InCreationContext);
+		CurrentCompilerContext = MakePimpl<UE::MVVM::Private::FMVVMViewBlueprintCompiler>(InCreationContext, GetBlueprintView());
 	}
 }
 
@@ -130,17 +130,17 @@ void UMVVMWidgetBlueprintExtension_View::HandleCreateClassVariablesFromBlueprint
 {
 	Super::HandleCreateClassVariablesFromBlueprint(Context);
 
-	CurrentCompilerContext->CreateVariables(Context, GetBlueprintView());
+	CurrentCompilerContext->CreateVariables(Context);
 }
 
 
-void UMVVMWidgetBlueprintExtension_View::HandleCreateFunctionList()
+void UMVVMWidgetBlueprintExtension_View::HandleCreateFunctionList(const FWidgetBlueprintCompilerContext::FCreateFunctionContext& Context)
 {
-	Super::HandleCreateFunctionList();
+	Super::HandleCreateFunctionList(Context);
 
 	if (CurrentCompilerContext)
 	{
-		CurrentCompilerContext->CreateFunctions(BlueprintView);
+		CurrentCompilerContext->CreateFunctions(Context);
 	}
 }
 
@@ -155,7 +155,7 @@ void UMVVMWidgetBlueprintExtension_View::HandleFinishCompilingClass(UWidgetBluep
 	{
 		UMVVMViewClass* ViewExtension = nullptr;
 		bool bCompiled = false;
-		if (CurrentCompilerContext->PreCompile(Class, BlueprintView))
+		if (CurrentCompilerContext->PreCompile(Class))
 		{
 			FName ClassName = "ViewClass";
 			if (UObject* PreviousObj = StaticFindObjectFastInternal(nullptr, Class, ClassName, true))
@@ -166,7 +166,7 @@ void UMVVMWidgetBlueprintExtension_View::HandleFinishCompilingClass(UWidgetBluep
 				PreviousObj->Rename(*TrashName.ToString(), GetTransientPackage(), RenameFlags);
 			}
 			ViewExtension = NewObject<UMVVMViewClass>(Class);
-			bCompiled = CurrentCompilerContext->Compile(Class, BlueprintView, ViewExtension);
+			bCompiled = CurrentCompilerContext->Compile(Class, ViewExtension);
 		}
 
 		if (bCompiled)
diff --git a/Engine/Plugins/Runtime/ModelViewViewModel/Source/ModelViewViewModelBlueprint/Public/Bindings/MVVMCompiledBindingLibraryCompiler.h b/Engine/Plugins/Runtime/ModelViewViewModel/Source/ModelViewViewModelBlueprint/Public/Bindings/MVVMCompiledBindingLibraryCompiler.h
--- a/Engine/Plugins/Runtime/ModelViewViewModel/Source/ModelViewViewModelBlueprint/Public/Bindings/MVVMCompiledBindingLibraryCompiler.h
+++ b/Engine/Plugins/Runtime/ModelViewViewModel/Source/ModelViewViewModelBlueprint/Public/Bindings/MVVMCompiledBindingLibraryCompiler.h
@@ -158,7 +158,7 @@ public:
 	TValueOrError<FFieldPathHandle, FText> AddFieldPath(TArrayView<const UE::MVVM::FMVVMConstFieldVariant> FieldPath, bool bRead);
 
 	/** */
-	TValueOrError<FFieldPathHandle, FText> AddObjectFieldPath(TArrayView<const UE::MVVM::FMVVMConstFieldVariant> FieldPath, UClass* ExpectedType, bool bRead);
+	TValueOrError<FFieldPathHandle, FText> AddObjectFieldPath(TArrayView<const UE::MVVM::FMVVMConstFieldVariant> FieldPath, const UClass* ExpectedType, bool bRead);
 
 	/** */
 	TValueOrError<FFieldPathHandle, FText> AddConversionFunctionFieldPath(const UClass* SourceClass, const UFunction* Function);
diff --git a/Engine/Plugins/Runtime/ModelViewViewModel/Source/ModelViewViewModelBlueprint/Public/MVVMBlueprintPin.h b/Engine/Plugins/Runtime/ModelViewViewModel/Source/ModelViewViewModelBlueprint/Public/MVVMBlueprintPin.h
--- a/Engine/Plugins/Runtime/ModelViewViewModel/Source/ModelViewViewModelBlueprint/Public/MVVMBlueprintPin.h
+++ b/Engine/Plugins/Runtime/ModelViewViewModel/Source/ModelViewViewModelBlueprint/Public/MVVMBlueprintPin.h
@@ -7,8 +7,19 @@
 #include "MVVMBlueprintPin.generated.h"
 
 class UWidgetBlueprint;
+class UEdGraphNode;
 class UEdGraphPin;
 
+/**
+*
+*/
+UENUM()
+enum class EMVVMBlueprintPinStatus : uint8
+{
+	Valid,
+	Orphaned,
+};
+
 /**
 *
 */
@@ -18,6 +29,9 @@ struct MODELVIEWVIEWMODELBLUEPRINT_API FMVVMBlueprintPin
 	GENERATED_BODY()
 
 private:
+	UPROPERTY(VisibleAnywhere, Category = "Viewmodel")
+	FGuid PinId;
+
 	UPROPERTY(VisibleAnywhere, Category = "Viewmodel")
 	FName PinName;
 
@@ -36,6 +50,14 @@ private:
 	UPROPERTY(VisibleAnywhere, Category = "Viewmodel")
 	TObjectPtr<class UObject> DefaultObject;
 
+	/** The pin is split. */
+	UPROPERTY(VisibleAnywhere, Category = "Viewmodel")
+	bool bSplit = false;
+
+	/** The pin could not be set. */
+	UPROPERTY(VisibleAnywhere, Category = "Viewmodel")
+	mutable EMVVMBlueprintPinStatus Status = EMVVMBlueprintPinStatus::Valid;
+
 public:
 	FMVVMBlueprintPin() = default;
 	FMVVMBlueprintPin(FName PinName);
@@ -45,10 +67,22 @@ public:
 		return PinName;
 	}
 
+	/** The pin is split into its different components. */
+	bool IsSplit() const
+	{
+		return bSplit;
+	}
+
+	/** The pin could not be assigned to the graph pin. */
+	EMVVMBlueprintPinStatus GetStatus() const
+	{
+		return Status;
+	}
+
 	/** Are we using the path. */
 	bool UsedPathAsValue() const
 	{
-		return Path.IsValid();
+		return !bSplit && Path.IsValid();
 	}
 
 	const FMVVMBlueprintPropertyPath& GetPath() const
@@ -63,8 +97,13 @@ public:
 	void SetDefaultValue(const FString& Value);
 	void SetPath(const FMVVMBlueprintPropertyPath& Value);
 
-	static FMVVMBlueprintPin CreateFromPin(const UBlueprint* WidgetBlueprint, const UEdGraphPin* Pin);
-	void CopyTo(const UBlueprint* WidgetBlueprint, UEdGraphPin* Pin) const;
+	static bool IsInputPin(const UEdGraphPin* Pin);
+	static TArray<FMVVMBlueprintPin> CopyAndReturnMissingPins(UBlueprint* Blueprint, UEdGraphNode* GraphNode, const TArray<FMVVMBlueprintPin>& Pins);
+	static TArray<FMVVMBlueprintPin> CreateFromNode(UBlueprint* Blueprint, UEdGraphNode* GraphNode);
+	static FMVVMBlueprintPin CreateFromPin(const UBlueprint* Blueprint, const UEdGraphPin* Pin);
+
+	void CopyTo(const UBlueprint* WidgetBlueprint, UEdGraphNode* Node) const;
+	UEdGraphPin* FindGraphPin(UEdGraphNode* Node) const;
 
 private:
 	void Reset();
diff --git a/Engine/Plugins/Runtime/ModelViewViewModel/Source/ModelViewViewModelBlueprint/Public/MVVMBlueprintViewBinding.h b/Engine/Plugins/Runtime/ModelViewViewModel/Source/ModelViewViewModelBlueprint/Public/MVVMBlueprintViewBinding.h
--- a/Engine/Plugins/Runtime/ModelViewViewModel/Source/ModelViewViewModelBlueprint/Public/MVVMBlueprintViewBinding.h
+++ b/Engine/Plugins/Runtime/ModelViewViewModel/Source/ModelViewViewModelBlueprint/Public/MVVMBlueprintViewBinding.h
@@ -74,7 +74,7 @@ public:
 
 	void GenerateWrapper(UBlueprint* Blueprint);
 	void SavePinValues(UBlueprint* Blueprint);
-	void DeprecateViewConversionFunction(UBlueprint* Blueprint);
+	void DeprecateViewConversionFunction(UBlueprint* Blueprint, FMVVMBlueprintViewBinding& Owner);
 
 public:
 PRAGMA_DISABLE_DEPRECATION_WARNINGS
diff --git a/Engine/Plugins/Runtime/ModelViewViewModel/Source/ModelViewViewModelBlueprint/Public/MVVMBlueprintViewConversionFunction.h b/Engine/Plugins/Runtime/ModelViewViewModel/Source/ModelViewViewModelBlueprint/Public/MVVMBlueprintViewConversionFunction.h
--- a/Engine/Plugins/Runtime/ModelViewViewModel/Source/ModelViewViewModelBlueprint/Public/MVVMBlueprintViewConversionFunction.h
+++ b/Engine/Plugins/Runtime/ModelViewViewModel/Source/ModelViewViewModelBlueprint/Public/MVVMBlueprintViewConversionFunction.h
@@ -5,6 +5,7 @@
 #include "MVVMBlueprintPin.h"
 #include "MVVMBlueprintView.h"
 #include "Engine/MemberReference.h"
+#include "UObject/WeakObjectPtrTemplates.h"
 
 #include "MVVMBlueprintViewConversionFunction.generated.h"
 
@@ -12,6 +13,8 @@ class UK2Node;
 class UEdGraphPin;
 class UEdGraph;
 
+struct FEdGraphEditAction;
+
 /**
  *
  */
@@ -20,39 +23,43 @@ class MODELVIEWVIEWMODELBLUEPRINT_API UMVVMBlueprintViewConversionFunction : pub
 {
 	GENERATED_BODY()
 
+public:
+	static bool IsValidConversionFunction(const UBlueprint* WidgetBlueprint, const UFunction* Function);
+
 public:
 	/** @return the conversion function uses at runtime. The wrapper function if complex or GetFunction is simple. */
 	const UFunction* GetCompiledFunction(const UClass* SelfContext) const;
 
 	/** @return the conversion function uses at runtime. The wrapper function if complex or GetFunction is simple. */
-	FName GetCompiledFunctionName() const;
+	FName GetCompiledFunctionName(const UClass* SelfContext) const;
 
 	/** @return the conversion function. */
-	TVariant<const UFunction*, TSubclassOf<UK2Node>> GetConversionFunction(const UClass* SelfContext) const;
+	TVariant<const UFunction*, TSubclassOf<UK2Node>> GetConversionFunction(const UBlueprint* SelfContext) const;
 
-	/** Set the function. Generate a Graph is the conversion function is complex. */
-	void InitFromFunction(UBlueprint* SelfContext, const UFunction* Function);
+	/** Set the function. Generate a Graph. */
+	void InitializeFromFunction(UBlueprint* SelfContext, FName GraphName, const UFunction* Function);
 
-	/** Set the function. Generate a Graph is the conversion function is complex. */
-	void InitFromFunction(UBlueprint* SelfContext, const UFunction* Function, FName GraphName);
+	// For deprecation
+	void Deprecation_InitializeFromWrapperGraph(UBlueprint* SelfContext, UEdGraph* Graph);
 
 	// For deprecation
-	void InitializeFromWrapperGraph(UBlueprint* SelfContext, UEdGraph* Graph);
+	void Deprecation_InitializeFromMemberReference(UBlueprint* SelfContext, FName GraphName, FMemberReference MemberReference, const FMVVMBlueprintPropertyPath& Source);
 
 	// For deprecation
-	void InitializeFromMemberReference(UBlueprint* SelfContext, FMemberReference MemberReference);
+	void Deprecation_SetWrapperGraphName(UBlueprint* Context, FName GraphName, const FMVVMBlueprintPropertyPath& Source);
+
+	/**
+	 * The conversion is valid.
+	 * The function was valid when created but may not be anymore.
+	 * It doesn't check if the source and destination are valid.
+	 */
+	bool IsValid(const UBlueprint* SelfContext) const;
 
 	/** The function has more than one argument and requires a wrapper or it uses a FunctionNode. */
-	bool NeedsWrapperGraph() const
-	{
-		return !GraphName.IsNone();
-	}
+	bool NeedsWrapperGraph(const UBlueprint* SelfContext) const;
 
-	/** The wrapper Graph is generated on domains and is not saved. */
-	bool IsWrapperGraphTransient() const
-	{
-		return NeedsWrapperGraph()  && bWrapperGraphTransient;
-	}
+	/** The wrapper Graph is generated on load/compile and is not saved. */
+	bool IsWrapperGraphTransient() const;
 
 	/** Return the wrapper graph, if it exists. */
 	UEdGraph* GetWrapperGraph() const
@@ -74,10 +81,10 @@ public:
 	 * If needed, create the graph and all the nodes for that graph when compiling.
 	 * Returns the existing one, if one was created from GetOrCreateWrapperGraph.
 	 */
-	UEdGraph* GetOrCreateIntermediateWrapperGraph(FKismetCompilerContext& Context) const;
+	UEdGraph* GetOrCreateIntermediateWrapperGraph(FKismetCompilerContext& Context);
 
 	/** If needed, create the graph and all the nodes for that graph. */
-	UEdGraph* GetOrCreateWrapperGraph(UBlueprint* Blueprint) const;
+	UEdGraph* GetOrCreateWrapperGraph(UBlueprint* Blueprint);
 
 	/**
 	 * The conversion function is going to be removed from the Blueprint.
@@ -89,7 +96,7 @@ public:
 	 * Returns the pin from the graph.
 	 * Create the graph and all the nodes for that graph if the graph doesn't exist and it's needed.
 	 */
-	UEdGraphPin* GetOrCreateGraphPin(UBlueprint* Blueprint, FName PinName) const;
+	UEdGraphPin* GetOrCreateGraphPin(UBlueprint* Blueprint, FName PinName);
 
 	TArrayView<const FMVVMBlueprintPin> GetPins() const
 	{
@@ -102,10 +109,17 @@ public:
 	/** Generates SavedPins from the wrapper graph, if it exists. */
 	void SavePinValues(UBlueprint* Blueprint);
 
+	FSimpleMulticastDelegate OnWrapperGraphModified;
+
 private:
-	UEdGraph* GetOrCreateWrapperGraphInternal(FKismetCompilerContext& Context, const UFunction* Function) const;
-	UEdGraph* GetOrCreateWrapperGraphInternal(UBlueprint* Blueprint, const UFunction* Function) const;
-	void LoadPinValuesInternal(UBlueprint* Blueprint) const;
+	void HandleGraphChanged(const FEdGraphEditAction& Action, TWeakObjectPtr<UBlueprint> Context);
+	void HandleUserDefinedPinRenamed(UK2Node* InNode, FName OldPinName, FName NewPinName, TWeakObjectPtr<UBlueprint> WeakBlueprint);
+	void SetCachedWrapperGraph(UBlueprint* Blueprint, UEdGraph* CachedGraph, UK2Node* CachedNode);
+	UEdGraph* GetOrCreateWrapperGraphInternal(FKismetCompilerContext& Context, const UFunction* Function);
+	UEdGraph* GetOrCreateWrapperGraphInternal(UBlueprint* Blueprint, const UFunction* Function);
+	bool NeedsWrapperGraphInternal(const UClass* SkeletalSelfContext) const;
+	void LoadPinValuesInternal(UBlueprint* Blueprint);
+	void CreateWrapperGraphName();
 	void Reset();
 
 private:
@@ -143,4 +157,7 @@ private:
 	
 	UPROPERTY(Transient, DuplicateTransient)
 	mutable TObjectPtr<UK2Node> CachedWrapperNode;
+
+	FDelegateHandle OnGraphChangedHandle;
+	FDelegateHandle OnUserDefinedPinRenamedHandle;
 };
diff --git a/Engine/Plugins/Runtime/ModelViewViewModel/Source/ModelViewViewModelBlueprint/Public/MVVMBlueprintViewEvent.h b/Engine/Plugins/Runtime/ModelViewViewModel/Source/ModelViewViewModelBlueprint/Public/MVVMBlueprintViewEvent.h
--- a/Engine/Plugins/Runtime/ModelViewViewModel/Source/ModelViewViewModelBlueprint/Public/MVVMBlueprintViewEvent.h
+++ b/Engine/Plugins/Runtime/ModelViewViewModel/Source/ModelViewViewModelBlueprint/Public/MVVMBlueprintViewEvent.h
@@ -7,6 +7,7 @@
 
 #include "MVVMBlueprintViewEvent.generated.h"
 
+struct FEdGraphEditAction;
 class UEdGraph;
 class UK2Node;
 class UWidgetBlueprint;
@@ -94,10 +95,12 @@ public:
 	// To set a pin when loading the asset (no graph generation)
 	void SetPinPathNoGraphGeneration(FName PinName, const FMVVMBlueprintPropertyPath& Path);
 
+	FSimpleMulticastDelegate OnWrapperGraphModified;
+
 public:
 	TArray<FText> GetCompilationMessages(EMessageType InMessageType) const;
 	bool HasCompilationMessage(EMessageType InMessageType) const;
-	void AddCompilationToBinding(FMessage MessageToAdd);
+	void AddCompilationToBinding(FMessage MessageToAdd) const;
 	void ResetCompilationMessages();
 
 public:
@@ -119,7 +122,10 @@ public:
 private:
 	static const UFunction* GetEventSignature(const UWidgetBlueprint* WidgetBlueprint, const FMVVMBlueprintPropertyPath& PropertyPath);
 	const UFunction* GetEventSignature() const;
+	void HandleGraphChanged(const FEdGraphEditAction& Action);
+	void HandleUserDefinedPinRenamed(UK2Node* InNode, FName OldPinName, FName NewPinName);
 	UWidgetBlueprint* GetWidgetBlueprintInternal() const;
+	void SetCachedWrapperGraphInternal(UEdGraph* Graph, UK2Node* Node);
 	UEdGraph* CreateWrapperGraphInternal();
 	UEdGraph* CreateWrapperGraphInternal(const UFunction* Signature, const UFunction* Function);
 	UEdGraph* CreateWrapperGraphInternal(const UFunction* Signature, const FProperty* Property);
@@ -142,11 +148,14 @@ private:
 	UPROPERTY(VisibleAnywhere, Category = "Viewmodel")
 	FName GraphName;
 
-	TArray<FMessage> Messages;
+	mutable TArray<FMessage> Messages;
 
 	UPROPERTY(Transient, DuplicateTransient)
 	mutable TObjectPtr<UEdGraph> CachedWrapperGraph;
 
 	UPROPERTY(Transient, DuplicateTransient)
 	mutable TObjectPtr<UK2Node> CachedWrapperNode;
+
+	FDelegateHandle OnGraphChangedHandle;
+	FDelegateHandle OnUserDefinedPinRenamedHandle;
 };
diff --git a/Engine/Plugins/Runtime/ModelViewViewModel/Source/ModelViewViewModelBlueprint/Public/MVVMDeveloperProjectSettings.h b/Engine/Plugins/Runtime/ModelViewViewModel/Source/ModelViewViewModelBlueprint/Public/MVVMDeveloperProjectSettings.h
--- a/Engine/Plugins/Runtime/ModelViewViewModel/Source/ModelViewViewModelBlueprint/Public/MVVMDeveloperProjectSettings.h
+++ b/Engine/Plugins/Runtime/ModelViewViewModel/Source/ModelViewViewModelBlueprint/Public/MVVMDeveloperProjectSettings.h
@@ -128,7 +128,7 @@ public:
 
 	/** When a conversion function requires a wrapper graph, add and save the generated graph to the blueprint. */
 	UPROPERTY(EditAnywhere, config, Category = "Viewmodel")
-	bool bAllowConversionFunctionGeneratedGraphInEditor = true;
+	bool bAllowConversionFunctionGeneratedGraphInEditor = false;
 
 	/** When binding to a multicast delegate property, allow to create an event. */
 	UPROPERTY(EditAnywhere, config, Category = "Viewmodel")
diff --git a/Engine/Plugins/Runtime/ModelViewViewModel/Source/ModelViewViewModelBlueprint/Public/MVVMWidgetBlueprintExtension_View.h b/Engine/Plugins/Runtime/ModelViewViewModel/Source/ModelViewViewModelBlueprint/Public/MVVMWidgetBlueprintExtension_View.h
--- a/Engine/Plugins/Runtime/ModelViewViewModel/Source/ModelViewViewModelBlueprint/Public/MVVMWidgetBlueprintExtension_View.h
+++ b/Engine/Plugins/Runtime/ModelViewViewModel/Source/ModelViewViewModelBlueprint/Public/MVVMWidgetBlueprintExtension_View.h
@@ -56,7 +56,7 @@ public:
 	virtual void HandleBeginCompilation(FWidgetBlueprintCompilerContext& InCreationContext) override;
 	virtual void HandleCleanAndSanitizeClass(UWidgetBlueprintGeneratedClass* ClassToClean, UObject* OldCDO) override;
 	virtual void HandleCreateClassVariablesFromBlueprint(const FWidgetBlueprintCompilerContext::FCreateVariableContext& Context) override;
-	virtual void HandleCreateFunctionList() override;
+	virtual void HandleCreateFunctionList(const FWidgetBlueprintCompilerContext::FCreateFunctionContext& Context) override;
 	virtual void HandleFinishCompilingClass(UWidgetBlueprintGeneratedClass* Class) override;
 	virtual void HandleEndCompilation() override;
 	virtual FSearchData HandleGatherSearchData(const UBlueprint* OwningBlueprint) const override;
diff --git a/Engine/Plugins/Runtime/ModelViewViewModel/Source/ModelViewViewModelEditor/Private/MVVMEditorSubsystem.cpp b/Engine/Plugins/Runtime/ModelViewViewModel/Source/ModelViewViewModelEditor/Private/MVVMEditorSubsystem.cpp
--- a/Engine/Plugins/Runtime/ModelViewViewModel/Source/ModelViewViewModelEditor/Private/MVVMEditorSubsystem.cpp
+++ b/Engine/Plugins/Runtime/ModelViewViewModel/Source/ModelViewViewModelEditor/Private/MVVMEditorSubsystem.cpp
@@ -338,7 +338,7 @@ UFunction* UMVVMEditorSubsystem::GetConversionFunction(const UWidgetBlueprint* W
 {
 	if (UMVVMBlueprintViewConversionFunction* ConversionFunction = Binding.Conversion.GetConversionFunction(bSourceToDestination))
 	{
-		TVariant<const UFunction*, TSubclassOf<UK2Node>> Result = ConversionFunction->GetConversionFunction(WidgetBlueprint->SkeletonGeneratedClass);
+		TVariant<const UFunction*, TSubclassOf<UK2Node>> Result = ConversionFunction->GetConversionFunction(WidgetBlueprint);
 		if (Result.IsType<const UFunction*>())
 		{
 			return const_cast<UFunction*>(Result.Get<const UFunction*>());
@@ -378,7 +378,8 @@ void UMVVMEditorSubsystem::SetSourceToDestinationConversionFunction(UWidgetBluep
 			{
 				Binding.Conversion.SourceToDestinationConversion = NewObject<UMVVMBlueprintViewConversionFunction>(WidgetBlueprint);
 				FName GraphName = UE::MVVM::ConversionFunctionHelper::CreateWrapperName(Binding, true);
-				Binding.Conversion.SourceToDestinationConversion->InitFromFunction(WidgetBlueprint, ConversionFunction, GraphName);
+				Binding.Conversion.SourceToDestinationConversion->InitializeFromFunction(WidgetBlueprint, GraphName, ConversionFunction);
+				Binding.SourcePath = FMVVMBlueprintPropertyPath();
 			}
 
 			UE::MVVM::Private::OnBindingPostEditChange(View, GET_MEMBER_NAME_CHECKED(FMVVMBlueprintViewBinding, Conversion));
@@ -409,7 +410,8 @@ void UMVVMEditorSubsystem::SetDestinationToSourceConversionFunction(UWidgetBluep
 			{
 				Binding.Conversion.DestinationToSourceConversion = NewObject<UMVVMBlueprintViewConversionFunction>(WidgetBlueprint);
 				FName GraphName = UE::MVVM::ConversionFunctionHelper::CreateWrapperName(Binding, false);
-				Binding.Conversion.DestinationToSourceConversion->InitFromFunction(WidgetBlueprint, ConversionFunction, GraphName);
+				Binding.Conversion.DestinationToSourceConversion->InitializeFromFunction(WidgetBlueprint, GraphName, ConversionFunction);
+				Binding.DestinationPath = FMVVMBlueprintPropertyPath();
 			}
 
 			UE::MVVM::Private::OnBindingPostEditChange(View, GET_MEMBER_NAME_CHECKED(FMVVMBlueprintViewBinding, Conversion));
@@ -426,9 +428,9 @@ void UMVVMEditorSubsystem::SetDestinationPathForBinding(UWidgetBlueprint* Widget
 
 		UE::MVVM::Private::OnBindingPreEditChange(View, GET_MEMBER_NAME_CHECKED(FMVVMBlueprintViewBinding, DestinationPath));
 
-		bool bSupports = UMVVMBlueprintViewEvent::Supports(WidgetBlueprint, PropertyPath);
+		bool bEventSupported = UMVVMBlueprintViewEvent::Supports(WidgetBlueprint, PropertyPath);
 		UMVVMBlueprintViewEvent* Event = nullptr;
-		if (bSupports)
+		if (bEventSupported)
 		{
 			Event = AddEvent(WidgetBlueprint);
 		}
@@ -440,6 +442,11 @@ void UMVVMEditorSubsystem::SetDestinationPathForBinding(UWidgetBlueprint* Widget
 		}
 		else
 		{
+			if (Binding.Conversion.DestinationToSourceConversion && PropertyPath.IsValid())
+			{
+				Binding.Conversion.DestinationToSourceConversion->RemoveWrapperGraph(WidgetBlueprint);
+				Binding.Conversion.DestinationToSourceConversion = nullptr;
+			}
 			Binding.DestinationPath = PropertyPath;
 		}
 
@@ -448,7 +455,7 @@ void UMVVMEditorSubsystem::SetDestinationPathForBinding(UWidgetBlueprint* Widget
 	}
 }
 
-void UMVVMEditorSubsystem::SetSourcePathForBinding(UWidgetBlueprint* WidgetBlueprint, FMVVMBlueprintViewBinding& Binding, FMVVMBlueprintPropertyPath Field)
+void UMVVMEditorSubsystem::SetSourcePathForBinding(UWidgetBlueprint* WidgetBlueprint, FMVVMBlueprintViewBinding& Binding, FMVVMBlueprintPropertyPath PropertyPath)
 {
 	if (UMVVMBlueprintView* View = GetView(WidgetBlueprint))
 	{
@@ -456,7 +463,12 @@ void UMVVMEditorSubsystem::SetSourcePathForBinding(UWidgetBlueprint* WidgetBluep
 
 		UE::MVVM::Private::OnBindingPreEditChange(View, GET_MEMBER_NAME_CHECKED(FMVVMBlueprintViewBinding, SourcePath));
 
-		Binding.SourcePath = Field;
+		if (Binding.Conversion.SourceToDestinationConversion && PropertyPath.IsValid())
+		{
+			Binding.Conversion.SourceToDestinationConversion->RemoveWrapperGraph(WidgetBlueprint);
+			Binding.Conversion.SourceToDestinationConversion = nullptr;
+		}
+		Binding.SourcePath = PropertyPath;
 
 		UE::MVVM::Private::OnBindingPostEditChange(View, GET_MEMBER_NAME_CHECKED(FMVVMBlueprintViewBinding, SourcePath));
 		FBlueprintEditorUtils::MarkBlueprintAsModified(WidgetBlueprint);
@@ -652,69 +664,71 @@ void UMVVMEditorSubsystem::SetCompileForEvent(UMVVMBlueprintViewEvent* Event, bo
 
 bool UMVVMEditorSubsystem::IsValidConversionFunction(const UWidgetBlueprint* WidgetBlueprint, const UFunction* Function, const FMVVMBlueprintPropertyPath& Source, const FMVVMBlueprintPropertyPath& Destination) const
 {
-	TValueOrError<const FProperty*, FText> ReturnResult = UE::MVVM::BindingHelper::TryGetReturnTypeForConversionFunction(Function);
-	if (ReturnResult.HasError())
+	if (!UMVVMBlueprintViewConversionFunction::IsValidConversionFunction(WidgetBlueprint, Function))
 	{
 		return false;
 	}
 
-	const FProperty* ReturnProperty = ReturnResult.GetValue();
-	if (ReturnProperty == nullptr)
-	{
-		return false;
-	}
-
-	TValueOrError<TArray<const FProperty*>, FText> ArgumentsResult = UE::MVVM::BindingHelper::TryGetArgumentsForConversionFunction(Function);
-	if (ArgumentsResult.HasError())
 	{
-		return false;
-	}
-
-	const FProperty* SourceProperty = nullptr;
+		const FProperty* SourceProperty = nullptr;
+		TArray<UE::MVVM::FMVVMConstFieldVariant> SourceFields = Source.GetFields(WidgetBlueprint->SkeletonGeneratedClass);
+		if (SourceFields.Num() > 0)
+		{
+			SourceProperty = SourceFields.Last().IsProperty() ? SourceFields.Last().GetProperty() : UE::MVVM::BindingHelper::GetReturnProperty(SourceFields.Last().GetFunction());
+		}
 
-	TArray<UE::MVVM::FMVVMConstFieldVariant> SourceFields = Source.GetFields(WidgetBlueprint->SkeletonGeneratedClass);
-	if (SourceFields.Num() > 0)
-	{
-		SourceProperty = SourceFields.Last().IsProperty() ? SourceFields.Last().GetProperty() : UE::MVVM::BindingHelper::GetReturnProperty(SourceFields.Last().GetFunction());
-	}
+		// check that at least one source -> argument binding is compatible
+		if (SourceProperty)
+		{
+			TValueOrError<TArray<const FProperty*>, FText> ArgumentsResult = UE::MVVM::BindingHelper::TryGetArgumentsForConversionFunction(Function);
+			check(ArgumentsResult.HasValue());
 
-	const FProperty* DestinationProperty = nullptr;
+			bool bAnyCompatible = false;
 
-	TArray<UE::MVVM::FMVVMConstFieldVariant> DestFields = Destination.GetFields(WidgetBlueprint->SkeletonGeneratedClass);
-	if (DestFields.Num() > 0)
-	{
-		DestinationProperty = DestFields.Last().IsProperty() ? DestFields.Last().GetProperty() : UE::MVVM::BindingHelper::GetFirstArgumentProperty(DestFields.Last().GetFunction());
+			const TArray<const FProperty*>& ConversionArgProperties = ArgumentsResult.GetValue();
+			for (const FProperty* ArgumentProperty : ConversionArgProperties)
+			{
+				if (UE::MVVM::BindingHelper::ArePropertiesCompatible(SourceProperty, ArgumentProperty))
+				{
+					bAnyCompatible = true;
+					break;
+				}
+			}
+			if (!bAnyCompatible)
+			{
+				return false;
+			}
+		}
 	}
 
-	// check that at least one source -> argument binding is compatible
-	bool bAnyCompatible = false;
-
-	const TArray<const FProperty*>& ConversionArgProperties = ArgumentsResult.GetValue();
-	for (const FProperty* ArgumentProperty : ConversionArgProperties)
 	{
-
-		if (SourceProperty == nullptr || UE::MVVM::BindingHelper::ArePropertiesCompatible(SourceProperty, ArgumentProperty))
+		const FProperty* DestinationProperty = nullptr;
+		TArray<UE::MVVM::FMVVMConstFieldVariant> DestFields = Destination.GetFields(WidgetBlueprint->SkeletonGeneratedClass);
+		if (DestFields.Num() > 0)
 		{
-			bAnyCompatible = true;
+			if (!DestFields.Last().IsEmpty())
+			{
+				DestinationProperty = DestFields.Last().IsProperty() ? DestFields.Last().GetProperty() : UE::MVVM::BindingHelper::GetFirstArgumentProperty(DestFields.Last().GetFunction());
+			}
 		}
-	}
 
-	if (!bAnyCompatible)
-	{
-		return false;
-	}
+		if (DestinationProperty)
+		{
+			TValueOrError<const FProperty*, FText> ReturnResult = UE::MVVM::BindingHelper::TryGetReturnTypeForConversionFunction(Function);
+			check(ReturnResult.HasValue());
 
-	// check that the return -> dest is valid
-	if (DestinationProperty != nullptr &&
-		!UE::MVVM::BindingHelper::ArePropertiesCompatible(ReturnProperty, DestinationProperty))
-	{
-		return false;
+			// check that the return -> dest is valid
+			if (!UE::MVVM::BindingHelper::ArePropertiesCompatible(ReturnResult.GetValue(), DestinationProperty))
+			{
+				return false;
+			}
+		}
 	}
 
-	return GetDefault<UMVVMDeveloperProjectSettings>()->IsConversionFunctionAllowed(WidgetBlueprint, Function);
+	return true;
 }
 
-bool UMVVMEditorSubsystem::IsSimpleConversionFunction(const UFunction* Function) const
+bool UMVVMEditorSubsystem::IsSimpleConversionFunctionA(const UFunction* Function) const
 {
 	TValueOrError<const FProperty*, FText> ReturnResult = UE::MVVM::BindingHelper::TryGetReturnTypeForConversionFunction(Function);
 	if (ReturnResult.HasError())
@@ -761,18 +775,9 @@ TArray<UFunction*> UMVVMEditorSubsystem::GetAvailableConversionFunctions(const U
 
 	auto AddFunction = [this, &Source, &Destination, WidgetBlueprint, &ConversionFunctions](const UFunction* Function)
 	{
-		// functions in the widget blueprint can do anything they want, other functions have to be static functions in a BlueprintFunctionLibrary
-		const UClass* FunctionClass = Function->GetOuterUClass();
-
-		bool bIsFromWidgetBlueprint = WidgetBlueprint->GeneratedClass && WidgetBlueprint->GeneratedClass->IsChildOf(FunctionClass) && Function->HasAllFunctionFlags(FUNC_BlueprintPure);
-		bool bIsFromSkeletonWidgetBlueprint = WidgetBlueprint->SkeletonGeneratedClass && WidgetBlueprint->SkeletonGeneratedClass->IsChildOf(FunctionClass) && Function->HasAllFunctionFlags(FUNC_BlueprintPure);
-		bool bFromBlueprintFunctionLibrary = FunctionClass->IsChildOf<UBlueprintFunctionLibrary>() && Function->HasAllFunctionFlags(FUNC_Static | FUNC_BlueprintPure);
-		if (bIsFromWidgetBlueprint || bIsFromSkeletonWidgetBlueprint || bFromBlueprintFunctionLibrary)
+		if (IsValidConversionFunction(WidgetBlueprint, Function, Source, Destination))
 		{
-			if (IsValidConversionFunction(WidgetBlueprint, Function, Source, Destination))
-			{
-				ConversionFunctions.Add(const_cast<UFunction*>(Function));
-			}
+			ConversionFunctions.Add(const_cast<UFunction*>(Function));
 		}
 	};
 
@@ -846,17 +851,9 @@ TArray<UFunction*> UMVVMEditorSubsystem::GetAvailableConversionFunctions(const U
 FMVVMBlueprintPropertyPath UMVVMEditorSubsystem::GetPathForConversionFunctionArgument(const UWidgetBlueprint* WidgetBlueprint, const FMVVMBlueprintViewBinding& Binding, FName ArgumentName, bool bSourceToDestination) const
 {
 	UMVVMBlueprintViewConversionFunction* ConversionFunction = Binding.Conversion.GetConversionFunction(bSourceToDestination);
-	if (ConversionFunction == nullptr || !ConversionFunction->NeedsWrapperGraph())
+	if (ConversionFunction == nullptr)
 	{
-		// simple conversion function
-		if (bSourceToDestination)
-		{
-			return Binding.SourcePath;
-		}
-		else
-		{
-			return Binding.DestinationPath;
-		}
+		return FMVVMBlueprintPropertyPath();
 	}
 
 	UEdGraphPin* GraphPin = ConversionFunction->GetOrCreateGraphPin(const_cast<UWidgetBlueprint*>(WidgetBlueprint), ArgumentName);
@@ -871,21 +868,8 @@ FMVVMBlueprintPropertyPath UMVVMEditorSubsystem::GetPathForConversionFunctionArg
 void UMVVMEditorSubsystem::SetPathForConversionFunctionArgument(UWidgetBlueprint* WidgetBlueprint, FMVVMBlueprintViewBinding& Binding, FName ArgumentName, const FMVVMBlueprintPropertyPath& Path, bool bSourceToDestination) const
 {
 	UMVVMBlueprintViewConversionFunction* ConversionFunction = Binding.Conversion.GetConversionFunction(bSourceToDestination);
-	if (ConversionFunction == nullptr || !ConversionFunction->NeedsWrapperGraph())
+	if (ConversionFunction)
 	{
-		// simple conversion function
-		if (bSourceToDestination)
-		{
-			Binding.SourcePath = Path;
-		}
-		else
-		{
-			Binding.DestinationPath = Path;
-		}
-	}
-	else
-	{
-		check(ConversionFunction);
 		ConversionFunction->SetGraphPin(WidgetBlueprint, ArgumentName, Path);
 	}
 	FBlueprintEditorUtils::MarkBlueprintAsModified(WidgetBlueprint);
@@ -911,106 +895,209 @@ namespace Private
 
 		return ViewEvent->GetOrCreateGraphPin(PinName);
 	}
-}
 
-void UMVVMEditorSubsystem::SplitPin(UWidgetBlueprint* WidgetBlueprint, FMVVMBlueprintViewBinding& Binding, FName PinName, bool bSourceToDestination) const
-{
-	UEdGraphPin* GraphPin = Private::GetGraphPin(this, WidgetBlueprint, Binding, PinName, bSourceToDestination);
-	const UEdGraphSchema_K2* K2Schema = GetDefault<UEdGraphSchema_K2>();
-	if (GraphPin == nullptr || !K2Schema->CanSplitStructPin(*GraphPin))
+	void DoAction(const UMVVMEditorSubsystem* Subsystem, UWidgetBlueprint* WidgetBlueprint, UMVVMBlueprintViewEvent* ViewEvent, FMVVMBlueprintViewBinding* Binding, FName PinName, bool bSourceToDestination
+		, FText TransactionName
+		, TFunctionRef<bool(const UEdGraphSchema_K2*, UEdGraphPin*)> Test, TFunctionRef<void(const UEdGraphSchema_K2*, UEdGraphPin*)> Action)
 	{
-		return;
+		UEdGraphPin* GraphPin = ViewEvent
+			? Private::GetGraphPin(WidgetBlueprint, ViewEvent, PinName)
+			: Private::GetGraphPin(Subsystem, WidgetBlueprint, *Binding, PinName, bSourceToDestination);
+		const UEdGraphSchema_K2* K2Schema = GetDefault<UEdGraphSchema_K2>();
+		if (!Test(K2Schema, GraphPin))
+		{
+			return;
+		}
+
+		FScopedTransaction Transaction(TransactionName);
+		if (ViewEvent)
+		{
+			FName NAME_SavedPin = "SavedPins";
+			UE::MVVM::Private::OnEventPreEditChange(ViewEvent, NAME_SavedPin);
+			Action(K2Schema, GraphPin);
+			ViewEvent->SavePinValues();
+			UE::MVVM::Private::OnEventPostEditChange(ViewEvent, NAME_SavedPin);
+		}
+		else
+		{
+			FName NAME_Conversion = GET_MEMBER_NAME_CHECKED(FMVVMBlueprintViewBinding, Conversion);
+			UMVVMBlueprintView* View = Subsystem->GetView(WidgetBlueprint);
+			UMVVMBlueprintViewConversionFunction* ConversionFunction = Binding->Conversion.GetConversionFunction(bSourceToDestination);
+			check(View);
+			check(ConversionFunction);
+
+			UE::MVVM::Private::OnBindingPreEditChange(View, NAME_Conversion);
+			Action(K2Schema, GraphPin);
+			ConversionFunction->SavePinValues(WidgetBlueprint);
+			UE::MVVM::Private::OnBindingPostEditChange(View, NAME_Conversion);
+		}
 	}
 
-	UMVVMBlueprintView* View = GetView(WidgetBlueprint);
-	UMVVMBlueprintViewConversionFunction* ConversionFunction = Binding.Conversion.GetConversionFunction(bSourceToDestination);
+	void SplitPin(const UMVVMEditorSubsystem* Subsystem, UWidgetBlueprint* WidgetBlueprint, UMVVMBlueprintViewEvent* ViewEvent, FMVVMBlueprintViewBinding* Binding, FName PinName, bool bSourceToDestination)
+	{
+		DoAction(Subsystem, WidgetBlueprint, ViewEvent, Binding, PinName, bSourceToDestination, LOCTEXT("BreakPin", "Split Struct Pin")
+			, [](const UEdGraphSchema_K2* K2Schema, UEdGraphPin* GraphPin)->bool
+			{
+				return GraphPin != nullptr && K2Schema->CanSplitStructPin(*GraphPin);
+			}
+			, [](const UEdGraphSchema_K2* K2Schema, UEdGraphPin* GraphPin)
+			{
+				K2Schema->SplitPin(GraphPin);
+			});
+	}
+
+	bool CanSplitPin(UEdGraphPin* GraphPin)
+	{
+		const UEdGraphSchema_K2* K2Schema = GetDefault<UEdGraphSchema_K2>();
+		return GraphPin ? K2Schema->CanSplitStructPin(*GraphPin) && !GraphPin->bOrphanedPin : false;
+	}
+	
+	void RecombinePin(const UMVVMEditorSubsystem* Subsystem, UWidgetBlueprint* WidgetBlueprint, UMVVMBlueprintViewEvent* ViewEvent, FMVVMBlueprintViewBinding* Binding, FName PinName, bool bSourceToDestination)
+	{
+		DoAction(Subsystem, WidgetBlueprint, ViewEvent, Binding, PinName, bSourceToDestination, LOCTEXT("BreakPin", "Split Struct Pin")
+			, [](const UEdGraphSchema_K2* K2Schema, UEdGraphPin* GraphPin)->bool
+			{
+				return GraphPin != nullptr && K2Schema->CanRecombineStructPin(*GraphPin);
+			}
+			, [](const UEdGraphSchema_K2* K2Schema, UEdGraphPin* GraphPin)
+			{
+				K2Schema->RecombinePin(GraphPin);
+			});
+	}
+
+	bool CanRecombinePin(UEdGraphPin* GraphPin)
+	{
+		const UEdGraphSchema_K2* K2Schema = GetDefault<UEdGraphSchema_K2>();
+		return GraphPin ? K2Schema->CanRecombineStructPin(*GraphPin) && !GraphPin->bOrphanedPin : false;
+	}
+	
+	void ResetPinToDefaultValue(const UMVVMEditorSubsystem* Subsystem, UWidgetBlueprint* WidgetBlueprint, UMVVMBlueprintViewEvent* ViewEvent, FMVVMBlueprintViewBinding* Binding, FName PinName, bool bSourceToDestination)
+	{
+		DoAction(Subsystem, WidgetBlueprint, ViewEvent, Binding, PinName, bSourceToDestination, LOCTEXT("BreakPin", "Split Struct Pin")
+			, [](const UEdGraphSchema_K2* K2Schema, UEdGraphPin* GraphPin)->bool
+			{
+				return GraphPin != nullptr && !K2Schema->DoesDefaultValueMatchAutogenerated(*GraphPin);
+			}
+			, [](const UEdGraphSchema_K2* K2Schema, UEdGraphPin* GraphPin)
+			{
+				K2Schema->ResetPinToAutogeneratedDefaultValue(GraphPin);
+			});
+	}
+
+	bool CanResetPinToDefaultValue(UEdGraphPin* GraphPin)
+	{
+		const UEdGraphSchema_K2* K2Schema = GetDefault<UEdGraphSchema_K2>();
+		return GraphPin ? !K2Schema->DoesDefaultValueMatchAutogenerated(*GraphPin) && !GraphPin->bOrphanedPin : false;
+	}
+
+	void ResetOrphanedPin(const UMVVMEditorSubsystem* Subsystem, UWidgetBlueprint* WidgetBlueprint, UMVVMBlueprintViewEvent* ViewEvent, FMVVMBlueprintViewBinding* Binding, FName PinName, bool bSourceToDestination)
+	{
+		DoAction(Subsystem, WidgetBlueprint, ViewEvent, Binding, PinName, bSourceToDestination, LOCTEXT("BreakPin", "Split Struct Pin")
+			, [](const UEdGraphSchema_K2* K2Schema, UEdGraphPin* GraphPin)->bool
+			{
+				return GraphPin == nullptr || GraphPin->bOrphanedPin;
+			}
+			, [](const UEdGraphSchema_K2* K2Schema, UEdGraphPin* GraphPin)
+				{
+					if (GraphPin)
+					{
+						K2Schema->ResetPinToAutogeneratedDefaultValue(GraphPin);
+					}
+				});
+	}
+	
+	bool CanResetOrphanedPin(UEdGraphPin* GraphPin)
+	{
+		return GraphPin ? GraphPin->bOrphanedPin : true;
+	}
+} //namespace
 
-	FScopedTransaction Transaction(LOCTEXT("BreakPin", "Split Struct Pin"));
-	UE::MVVM::Private::OnBindingPreEditChange(View, GET_MEMBER_NAME_CHECKED(FMVVMBlueprintViewBinding, Conversion));
-	K2Schema->SplitPin(GraphPin);
-	ConversionFunction->SavePinValues(WidgetBlueprint);
-	UE::MVVM::Private::OnBindingPostEditChange(View, GET_MEMBER_NAME_CHECKED(FMVVMBlueprintViewBinding, Conversion));
+void UMVVMEditorSubsystem::SplitPin(UWidgetBlueprint* WidgetBlueprint, FMVVMBlueprintViewBinding& Binding, FName PinName, bool bSourceToDestination) const
+{
+	Private::SplitPin(this, WidgetBlueprint, nullptr, &Binding, PinName, bSourceToDestination);
 }
 
 bool UMVVMEditorSubsystem::CanSplitPin(const UWidgetBlueprint* WidgetBlueprint, const FMVVMBlueprintViewBinding& Binding, FName PinName, bool bSourceToDestination) const
 {
 	UEdGraphPin* GraphPin = Private::GetGraphPin(this, WidgetBlueprint, Binding, PinName, bSourceToDestination);
-	const UEdGraphSchema_K2* K2Schema = GetDefault<UEdGraphSchema_K2>();
-	return GraphPin ? K2Schema->CanSplitStructPin(*GraphPin) : false;
+	return Private::CanSplitPin(GraphPin);
 }
 
 void UMVVMEditorSubsystem::SplitPin(UWidgetBlueprint* WidgetBlueprint, UMVVMBlueprintViewEvent* ViewEvent, FName PinName) const
 {
-	UEdGraphPin* GraphPin = Private::GetGraphPin(WidgetBlueprint, ViewEvent, PinName);
-	const UEdGraphSchema_K2* K2Schema = GetDefault<UEdGraphSchema_K2>();
-	if (GraphPin == nullptr || !K2Schema->CanSplitStructPin(*GraphPin))
-	{
-		return;
-	}
-
-	FName NAME_SavedPin = "SavedPins";
-	FScopedTransaction Transaction(LOCTEXT("BreakPin", "Split Struct Pin"));
-	UE::MVVM::Private::OnEventPreEditChange(ViewEvent, NAME_SavedPin);
-	K2Schema->SplitPin(GraphPin);
-	ViewEvent->SavePinValues();
-	UE::MVVM::Private::OnEventPostEditChange(ViewEvent, NAME_SavedPin);
+	Private::SplitPin(this, WidgetBlueprint, ViewEvent, nullptr, PinName, true);
 }
 
 bool UMVVMEditorSubsystem::CanSplitPin(const UWidgetBlueprint* WidgetBlueprint, UMVVMBlueprintViewEvent* Event, FName PinName) const
 {
 	UEdGraphPin* GraphPin = Private::GetGraphPin(WidgetBlueprint, Event, PinName);
-	const UEdGraphSchema_K2* K2Schema = GetDefault<UEdGraphSchema_K2>();
-	return GraphPin ? K2Schema->CanSplitStructPin(*GraphPin) : false;
+	return Private::CanSplitPin(GraphPin);
 }
 
 void UMVVMEditorSubsystem::RecombinePin(UWidgetBlueprint* WidgetBlueprint, FMVVMBlueprintViewBinding& Binding, FName PinName, bool bSourceToDestination) const
+{
+	Private::RecombinePin(this, WidgetBlueprint, nullptr, &Binding, PinName, bSourceToDestination);
+}
+
+bool UMVVMEditorSubsystem::CanRecombinePin(const UWidgetBlueprint* WidgetBlueprint, const FMVVMBlueprintViewBinding& Binding, FName PinName, bool bSourceToDestination) const
 {
 	UEdGraphPin* GraphPin = Private::GetGraphPin(this, WidgetBlueprint, Binding, PinName, bSourceToDestination);
-	const UEdGraphSchema_K2* K2Schema = GetDefault<UEdGraphSchema_K2>();
-	if (GraphPin == nullptr || !K2Schema->CanRecombineStructPin(*GraphPin))
-	{
-		return;
-	}
+	return Private::CanRecombinePin(GraphPin);
+}
 
-	UMVVMBlueprintView* View = GetView(WidgetBlueprint);
-	UMVVMBlueprintViewConversionFunction* ConversionFunction = Binding.Conversion.GetConversionFunction(bSourceToDestination);
+void UMVVMEditorSubsystem::RecombinePin(UWidgetBlueprint* WidgetBlueprint, UMVVMBlueprintViewEvent* ViewEvent, FName PinName) const
+{
+	Private::RecombinePin(this, WidgetBlueprint, ViewEvent, nullptr, PinName, true);
+}
 
-	FScopedTransaction Transaction(LOCTEXT("RecombinePin", "Recombine Struct Pin"));
-	UE::MVVM::Private::OnBindingPreEditChange(View, GET_MEMBER_NAME_CHECKED(FMVVMBlueprintViewBinding, Conversion));
-	K2Schema->RecombinePin(GraphPin);
-	ConversionFunction->SavePinValues(WidgetBlueprint);
-	UE::MVVM::Private::OnBindingPostEditChange(View, GET_MEMBER_NAME_CHECKED(FMVVMBlueprintViewBinding, Conversion));
+bool UMVVMEditorSubsystem::CanRecombinePin(const UWidgetBlueprint* WidgetBlueprint, UMVVMBlueprintViewEvent* Event, FName PinName) const
+{
+	UEdGraphPin* GraphPin = Private::GetGraphPin(WidgetBlueprint, Event, PinName);
+	return Private::CanRecombinePin(GraphPin);
 }
 
-bool UMVVMEditorSubsystem::CanRecombinePin(const UWidgetBlueprint* WidgetBlueprint, const FMVVMBlueprintViewBinding& Binding, FName PinName, bool bSourceToDestination) const
+void UMVVMEditorSubsystem::ResetPinToDefaultValue(UWidgetBlueprint* WidgetBlueprint, FMVVMBlueprintViewBinding& Binding, FName PinName, bool bSourceToDestination) const
+{
+	Private::ResetPinToDefaultValue(this, WidgetBlueprint, nullptr, &Binding, PinName, bSourceToDestination);
+}
+
+bool UMVVMEditorSubsystem::CanResetPinToDefaultValue(const UWidgetBlueprint* WidgetBlueprint, const FMVVMBlueprintViewBinding& Binding, FName PinName, bool bSourceToDestination) const
 {
 	UEdGraphPin* GraphPin = Private::GetGraphPin(this, WidgetBlueprint, Binding, PinName, bSourceToDestination);
-	const UEdGraphSchema_K2* K2Schema = GetDefault<UEdGraphSchema_K2>();
-	return GraphPin ? K2Schema->CanRecombineStructPin(*GraphPin) : false;
+	return Private::CanResetPinToDefaultValue(GraphPin);
 }
 
-void UMVVMEditorSubsystem::RecombinePin(UWidgetBlueprint* WidgetBlueprint, UMVVMBlueprintViewEvent* ViewEvent, FName PinName) const
+void UMVVMEditorSubsystem::ResetPinToDefaultValue(UWidgetBlueprint* WidgetBlueprint, UMVVMBlueprintViewEvent* ViewEvent, FName PinName) const
 {
-	UEdGraphPin* GraphPin = Private::GetGraphPin(WidgetBlueprint, ViewEvent, PinName);
-	const UEdGraphSchema_K2* K2Schema = GetDefault<UEdGraphSchema_K2>();
-	if (GraphPin == nullptr || !K2Schema->CanRecombineStructPin(*GraphPin))
-	{
-		return;
-	}
+	Private::ResetPinToDefaultValue(this, WidgetBlueprint, ViewEvent, nullptr, PinName, true);
+}
 
-	FName NAME_SavedPin = "SavedPins";
-	FScopedTransaction Transaction(LOCTEXT("RecombinePin", "Recombine Struct Pin"));
-	UE::MVVM::Private::OnEventPreEditChange(ViewEvent, NAME_SavedPin);
-	K2Schema->RecombinePin(GraphPin);
-	ViewEvent->SavePinValues();
-	UE::MVVM::Private::OnEventPostEditChange(ViewEvent, NAME_SavedPin);
+bool UMVVMEditorSubsystem::CanResetPinToDefaultValue(const UWidgetBlueprint* WidgetBlueprint, UMVVMBlueprintViewEvent* Event, FName PinName) const
+{
+	UEdGraphPin* GraphPin = Private::GetGraphPin(WidgetBlueprint, Event, PinName);
+	return Private::CanResetPinToDefaultValue(GraphPin);
 }
 
-bool UMVVMEditorSubsystem::CanRecombinePin(const UWidgetBlueprint* WidgetBlueprint, UMVVMBlueprintViewEvent* Event, FName PinName) const
+void UMVVMEditorSubsystem::ResetOrphanedPin(UWidgetBlueprint* WidgetBlueprint, FMVVMBlueprintViewBinding& Binding, FName PinName, bool bSourceToDestination) const
+{
+	Private::ResetOrphanedPin(this, WidgetBlueprint, nullptr, &Binding, PinName, bSourceToDestination);
+}
+
+bool UMVVMEditorSubsystem::CanResetOrphanedPin(const UWidgetBlueprint* WidgetBlueprint, const FMVVMBlueprintViewBinding& Binding, FName PinName, bool bSourceToDestination) const
+{
+	UEdGraphPin* GraphPin = Private::GetGraphPin(this, WidgetBlueprint, Binding, PinName, bSourceToDestination);
+	return Private::CanResetOrphanedPin(GraphPin);
+}
+
+void UMVVMEditorSubsystem::ResetOrphanedPin(UWidgetBlueprint* WidgetBlueprint, UMVVMBlueprintViewEvent* ViewEvent, FName PinName) const
+{
+	Private::ResetOrphanedPin(this, WidgetBlueprint, ViewEvent, nullptr, PinName, true);
+}
+
+bool UMVVMEditorSubsystem::CanResetOrphanedPin(const UWidgetBlueprint* WidgetBlueprint, UMVVMBlueprintViewEvent* Event, FName PinName) const
 {
 	UEdGraphPin* GraphPin = Private::GetGraphPin(WidgetBlueprint, Event, PinName);
-	const UEdGraphSchema_K2* K2Schema = GetDefault<UEdGraphSchema_K2>();
-	return GraphPin ? K2Schema->CanRecombineStructPin(*GraphPin) : false;
+	return Private::CanResetOrphanedPin(GraphPin);
 }
 
 TArray<UE::MVVM::FBindingSource> UMVVMEditorSubsystem::GetBindableWidgets(const UWidgetBlueprint* WidgetBlueprint) const
diff --git a/Engine/Plugins/Runtime/ModelViewViewModel/Source/ModelViewViewModelEditor/Private/Widgets/SMVVMViewBindingListView.cpp b/Engine/Plugins/Runtime/ModelViewViewModel/Source/ModelViewViewModelEditor/Private/Widgets/SMVVMViewBindingListView.cpp
--- a/Engine/Plugins/Runtime/ModelViewViewModel/Source/ModelViewViewModelEditor/Private/Widgets/SMVVMViewBindingListView.cpp
+++ b/Engine/Plugins/Runtime/ModelViewViewModel/Source/ModelViewViewModelEditor/Private/Widgets/SMVVMViewBindingListView.cpp
@@ -133,7 +133,7 @@ struct FBindingEntry
 		return Name;
 	}
 
-	void SetBindingParameterName(FGuid Id, FName ParameterName)
+	void SetBindingParameter(FGuid Id, FName ParameterName)
 	{
 		check(RowType == ERowType::None);
 		RowType = ERowType::BindingParameter;
@@ -158,7 +158,7 @@ struct FBindingEntry
 		return Name;
 	}
 
-	void SetEventParameterName(UMVVMBlueprintViewEvent* InEvent, FName ParameterName)
+	void SetEventParameter(UMVVMBlueprintViewEvent* InEvent, FName ParameterName)
 	{
 		check(RowType == ERowType::None);
 		RowType = ERowType::EventParameter;
@@ -201,9 +201,10 @@ struct FBindingEntry
 
 	bool operator==(const FBindingEntry& Other) const
 	{
-		return RowType == Other.RowType &&
-			Name == Other.Name &&
-			BindingId == Other.BindingId;
+		return RowType == Other.RowType
+			&& Name == Other.Name
+			&& BindingId == Other.BindingId
+			&& Event == Other.Event;
 	}
 
 	FString GetSearchNameString(UMVVMBlueprintView* View, UWidgetBlueprint* WidgetBP)
@@ -1802,11 +1803,8 @@ class SFunctionParameterRow : public STableRow<TSharedPtr<FBindingEntry>>
 
 		UMVVMEditorSubsystem* EditorSubsystem = GEditor->GetEditorSubsystem<UMVVMEditorSubsystem>();
 
-		FSlateColor PrimaryColor, SecondaryColor;
-		const FSlateBrush* PrimaryBrush = nullptr;
-		const FSlateBrush* SecondaryBrush = nullptr;
-		FText DisplayName, ToolTip;
-
+		const FMVVMBlueprintPin* Pin = nullptr;
+		const UEdGraphPin* GraphPin = nullptr;
 		TSharedPtr<SWidget> ContentWidget;
 
 		if (Entry->GetRowType() == FBindingEntry::ERowType::BindingParameter)
@@ -1814,26 +1812,14 @@ class SFunctionParameterRow : public STableRow<TSharedPtr<FBindingEntry>>
 			bool bSimpleConversionFunction = false;
 			UMVVMBlueprintView* View = EditorSubsystem->GetView(InWidgetBlueprint);
 			FMVVMBlueprintViewBinding* Binding = Entry->GetBinding(View);
-			const bool bSourceToDestination = UE::MVVM::IsForwardBinding(Binding->BindingType);
+			check(Binding);
 
-			if (UEdGraphPin* Pin = EditorSubsystem->GetConversionFunctionArgumentPin(InWidgetBlueprint, *Binding, Entry->GetBindingParameterName(), bSourceToDestination))
-			{
-				PrimaryBrush = FBlueprintEditor::GetVarIconAndColorFromPinType(Pin->PinType, PrimaryColor, SecondaryBrush, SecondaryColor);
-				DisplayName = Pin->GetDisplayName();
-				ToolTip = FText::FromString(Pin->PinToolTip);
-			}
-			else if (const UFunction* Function = EditorSubsystem->GetConversionFunction(InWidgetBlueprint, *Binding, bSourceToDestination))
-			{
-				// no wrapper graph, this is a simple conversion function of the form: int32 Convert(float x)
-				if (const FProperty* Argument = BindingHelper::GetFirstArgumentProperty(Function))
-				{
-					bSimpleConversionFunction = true;
+			const bool bSourceToDestination = UE::MVVM::IsForwardBinding(Binding->BindingType);
+			UMVVMBlueprintViewConversionFunction* ConversionFunction = Binding->Conversion.GetConversionFunction(bSourceToDestination);
+			check(ConversionFunction);
 
-					PrimaryBrush = FBlueprintEditor::GetVarIconAndColorFromProperty(Argument, PrimaryColor, SecondaryBrush, SecondaryColor);
-					DisplayName = Argument->GetDisplayNameText();
-					ToolTip = Argument->GetToolTipText();
-				}
-			}
+			Pin = ConversionFunction->GetPins().FindByPredicate([Name = Entry->GetBindingParameterName()](const FMVVMBlueprintPin& Other){ return Other.GetName() == Name; });
+			GraphPin = ConversionFunction->GetOrCreateGraphPin(InWidgetBlueprint, Entry->GetBindingParameterName());
 
 			ContentWidget = SNew(SFunctionParameter, InWidgetBlueprint)
 				.BindingId(Binding->BindingId)
@@ -1844,12 +1830,10 @@ class SFunctionParameterRow : public STableRow<TSharedPtr<FBindingEntry>>
 		else if (Entry->GetRowType() == FBindingEntry::ERowType::EventParameter)
 		{
 			UMVVMBlueprintViewEvent* ViewEvent = Entry->GetEvent();
-			if (UEdGraphPin* Pin = ViewEvent ? ViewEvent->GetOrCreateGraphPin(Entry->GetEventParameterName()) : nullptr)
-			{
-				PrimaryBrush = FBlueprintEditor::GetVarIconAndColorFromPinType(Pin->PinType, PrimaryColor, SecondaryBrush, SecondaryColor);
-				DisplayName = Pin->GetDisplayName();
-				ToolTip = FText::FromString(Pin->PinToolTip);
-			}
+			check(ViewEvent);
+
+			Pin = ViewEvent->GetPins().FindByPredicate([Name = Entry->GetEventParameterName()](const FMVVMBlueprintPin& Other) { return Other.GetName() == Name; });
+			GraphPin = ViewEvent->GetOrCreateGraphPin(Entry->GetEventParameterName());
 
 			ContentWidget = SNew(SEventParameter, InWidgetBlueprint)
 				.Event(Entry->GetEvent())
@@ -1857,6 +1841,26 @@ class SFunctionParameterRow : public STableRow<TSharedPtr<FBindingEntry>>
 				.AllowDefault(true);
 		}
 
+		FSlateColor PrimaryColor, SecondaryColor;
+		const FSlateBrush* PrimaryBrush = nullptr;
+		const FSlateBrush* SecondaryBrush = nullptr;
+		FText DisplayName, ToolTip;
+		bool bTextColorIsRed = false;
+		if (Pin == nullptr || GraphPin == nullptr || Pin->GetStatus() == EMVVMBlueprintPinStatus::Orphaned)
+		{
+			PrimaryBrush = FBlueprintEditorUtils::GetIconFromPin(FEdGraphPinType());
+			PrimaryColor = FLinearColor::Red;
+			DisplayName = GraphPin ? GraphPin->GetDisplayName() : (Pin ? FText::FromName(Pin->GetName()) : FText::GetEmpty());
+			bTextColorIsRed = true;
+		}
+		else
+		{
+			PrimaryBrush = FBlueprintEditor::GetVarIconAndColorFromPinType(GraphPin->PinType, PrimaryColor, SecondaryBrush, SecondaryColor);
+			DisplayName = GraphPin->GetDisplayName();
+			ToolTip = FText::FromString(GraphPin->PinToolTip);
+		}
+
+		TSharedPtr<STextBlock> DisplayNameTextBlock;
 		STableRow<TSharedPtr<FBindingEntry>>::Construct(
 			STableRow<TSharedPtr<FBindingEntry>>::FArguments()
 			.ShowWires(true)
@@ -1887,7 +1891,7 @@ class SFunctionParameterRow : public STableRow<TSharedPtr<FBindingEntry>>
 						]
 						+ SHorizontalBox::Slot()
 						[
-							SNew(STextBlock)
+							SAssignNew(DisplayNameTextBlock, STextBlock)
 							.Text(DisplayName)
 						]
 					]
@@ -1902,6 +1906,11 @@ class SFunctionParameterRow : public STableRow<TSharedPtr<FBindingEntry>>
 			], 
 			OwnerTableView
 		);
+
+		if (bTextColorIsRed)
+		{
+			DisplayNameTextBlock->SetColorAndOpacity(PrimaryColor);
+		}
 	}
 
 private:
@@ -2088,67 +2097,33 @@ void SBindingsList::Refresh()
 			if (UMVVMBlueprintViewConversionFunction* ConversionFunction = Binding.Conversion.GetConversionFunction(UE::MVVM::IsForwardBinding(Binding.BindingType)))
 			{
 				ConversionFunction->GetOrCreateWrapperGraph(MVVMExtensionPtr->GetWidgetBlueprint());
-				if (ConversionFunction->NeedsWrapperGraph())
+				for (const FMVVMBlueprintPin& Pin : ConversionFunction->GetPins())
 				{
-					for (const FMVVMBlueprintPin& Pin : ConversionFunction->GetPins())
+					UEdGraphPin* GraphPin = ConversionFunction->GetOrCreateGraphPin(MVVMExtensionPtr->GetWidgetBlueprint(), Pin.GetName());
+					if (GraphPin && GraphPin->bHidden)
 					{
-						UEdGraphPin* GraphPin = ConversionFunction->GetOrCreateGraphPin(MVVMExtensionPtr->GetWidgetBlueprint(), Pin.GetName());
-						if (GraphPin && GraphPin->bHidden)
-						{
-							continue;
-						}
-
-						TSharedPtr<FBindingEntry> ArgumentEntry;
-						if (PreviousGroupEntry)
-						{
-							if (TSharedPtr<FBindingEntry>* FoundParameter = PreviousGroupEntry->Children.FindByPredicate([BindingId, ArgumentName = Pin.GetName()](const TSharedPtr<FBindingEntry>& Other)
-								{ return Other->GetBindingId() == BindingId && Other->GetRowType() == FBindingEntry::ERowType::BindingParameter && Other->GetBindingParameterName() == ArgumentName; }))
-							{
-								ArgumentEntry = *FoundParameter;
-							}
-						}
+						continue;
+					}
 
-						if (!ArgumentEntry.IsValid())
+					TSharedPtr<FBindingEntry> ArgumentEntry;
+					if (PreviousGroupEntry)
+					{
+						if (TSharedPtr<FBindingEntry>* FoundParameter = PreviousGroupEntry->Children.FindByPredicate([BindingId, ArgumentName = Pin.GetName()](const TSharedPtr<FBindingEntry>& Other)
+							{ return Other->GetBindingId() == BindingId && Other->GetRowType() == FBindingEntry::ERowType::BindingParameter && Other->GetBindingParameterName() == ArgumentName; }))
 						{
-							ArgumentEntry = MakeShared<FBindingEntry>();
-							ArgumentEntry->SetBindingParameterName(Binding.BindingId, Pin.GetName());
-
-							NewEntries.Add(ArgumentEntry);
+							ArgumentEntry = *FoundParameter;
 						}
-						BindingEntry->AddChild(ArgumentEntry);
 					}
-				}
-				else // simple conversion function with one parameters (no subpins)
-				{
-					const UFunction* Function = GEditor->GetEditorSubsystem<UMVVMEditorSubsystem>()->GetConversionFunction(MVVMExtensionPtr->GetWidgetBlueprint(), Binding, UE::MVVM::IsForwardBinding(Binding.BindingType));
-					if (Function != nullptr)
-					{
-						TValueOrError<TArray<const FProperty*>, FText> ArgumentsResult = BindingHelper::TryGetArgumentsForConversionFunction(Function);
-						if (ArgumentsResult.HasValue())
-						{
-							for (const FProperty* Argument : ArgumentsResult.GetValue())
-							{
-								TSharedPtr<FBindingEntry> ArgumentEntry;
-								if (PreviousGroupEntry)
-								{
-									if (TSharedPtr<FBindingEntry>* FoundParameter = PreviousGroupEntry->Children.FindByPredicate([BindingId, ArgumentName = Argument->GetFName()](const TSharedPtr<FBindingEntry>& Other)
-										{ return Other->GetBindingId() == BindingId && Other->GetRowType() == FBindingEntry::ERowType::BindingParameter && Other->GetBindingParameterName() == ArgumentName; }))
-									{
-										ArgumentEntry = *FoundParameter;
-									}
-								}
 
-								if (!ArgumentEntry.IsValid())
-								{
-									ArgumentEntry = MakeShared<FBindingEntry>();
-									ArgumentEntry->SetBindingParameterName(Binding.BindingId, Argument->GetFName());
+					if (!ArgumentEntry.IsValid())
+					{
+						ArgumentEntry = MakeShared<FBindingEntry>();
+						ArgumentEntry->SetBindingParameter(Binding.BindingId, Pin.GetName());
+						ConversionFunction->OnWrapperGraphModified.AddSP(this, &SBindingsList::ForceRefresh);
 
-									NewEntries.Add(ArgumentEntry);
-								}
-								BindingEntry->AddChild(ArgumentEntry);
-							}
-						}
+						NewEntries.Add(ArgumentEntry);
 					}
+					BindingEntry->AddChild(ArgumentEntry);
 				}
 			}
 		}
@@ -2226,7 +2201,8 @@ void SBindingsList::Refresh()
 				if (!ArgumentEntry.IsValid())
 				{
 					ArgumentEntry = MakeShared<FBindingEntry>();
-					ArgumentEntry->SetEventParameterName(Event, Pin.GetName());
+					ArgumentEntry->SetEventParameter(Event, Pin.GetName());
+					Event->OnWrapperGraphModified.AddSP(this, &SBindingsList::ForceRefresh);
 
 					NewEntries.Add(ArgumentEntry);
 				}
@@ -2396,6 +2372,33 @@ void SBindingsList::HandleDeleteSelected()
 	}
 }
 
+void SBindingsList::HandleResetSelectedPin()
+{
+	if (UMVVMWidgetBlueprintExtension_View* MVVMExtensionPtr = MVVMExtension.Get())
+	{
+		UWidgetBlueprint* WidgetBlueprint = MVVMExtensionPtr->GetWidgetBlueprint();
+		check(WidgetBlueprint);
+		UMVVMEditorSubsystem* EditorSubsystem = GEditor->GetEditorSubsystem<UMVVMEditorSubsystem>();
+
+		TArray<TSharedPtr<FBindingEntry>> Selection = TreeView->GetSelectedItems();
+		for (const TSharedPtr<FBindingEntry>& Entry : Selection)
+		{
+			if (Entry->GetRowType() == FBindingEntry::ERowType::EventParameter)
+			{
+				EditorSubsystem->ResetPinToDefaultValue(WidgetBlueprint, Entry->GetEvent(), Entry->GetEventParameterName());
+			}
+			else if (Entry->GetRowType() == FBindingEntry::ERowType::BindingParameter)
+			{
+				if (FMVVMBlueprintViewBinding* Binding = Entry->GetBinding(MVVMExtensionPtr->GetBlueprintView()))
+				{
+					const bool bSourceToDestination = UE::MVVM::IsForwardBinding(Binding->BindingType);
+					EditorSubsystem->ResetPinToDefaultValue(WidgetBlueprint, *Binding, Entry->GetBindingParameterName(), bSourceToDestination);
+				}
+			}
+		}
+	}
+}
+
 void SBindingsList::HandleBreakSelectedPin()
 {
 	if (UMVVMWidgetBlueprintExtension_View* MVVMExtensionPtr = MVVMExtension.Get())
@@ -2416,7 +2419,7 @@ void SBindingsList::HandleBreakSelectedPin()
 				if (FMVVMBlueprintViewBinding* Binding = Entry->GetBinding(MVVMExtensionPtr->GetBlueprintView()))
 				{
 					const bool bSourceToDestination = UE::MVVM::IsForwardBinding(Binding->BindingType);
-					EditorSubsystem->SplitPin(WidgetBlueprint, *Binding, Entry->GetEventParameterName(), bSourceToDestination);
+					EditorSubsystem->SplitPin(WidgetBlueprint, *Binding, Entry->GetBindingParameterName(), bSourceToDestination);
 				}
 			}
 		}
@@ -2427,8 +2430,6 @@ void SBindingsList::HandleRecombineSelectedPin()
 {
 	if (UMVVMWidgetBlueprintExtension_View* MVVMExtensionPtr = MVVMExtension.Get())
 	{
-		const FScopedTransaction Transaction(LOCTEXT("RecombineStructPin", "Recombine Struct Pin"));
-
 		UWidgetBlueprint* WidgetBlueprint = MVVMExtensionPtr->GetWidgetBlueprint();
 		check(WidgetBlueprint);
 		UMVVMEditorSubsystem* EditorSubsystem = GEditor->GetEditorSubsystem<UMVVMEditorSubsystem>();
@@ -2445,7 +2446,34 @@ void SBindingsList::HandleRecombineSelectedPin()
 				if (FMVVMBlueprintViewBinding* Binding = Entry->GetBinding(MVVMExtensionPtr->GetBlueprintView()))
 				{
 					const bool bSourceToDestination = UE::MVVM::IsForwardBinding(Binding->BindingType);
-					EditorSubsystem->RecombinePin(WidgetBlueprint, *Binding, Entry->GetEventParameterName(), bSourceToDestination);
+					EditorSubsystem->RecombinePin(WidgetBlueprint, *Binding, Entry->GetBindingParameterName(), bSourceToDestination);
+				}
+			}
+		}
+	}
+}
+
+void SBindingsList::HandleResetOrphanedSelectedPin()
+{
+	if (UMVVMWidgetBlueprintExtension_View* MVVMExtensionPtr = MVVMExtension.Get())
+	{
+		UWidgetBlueprint* WidgetBlueprint = MVVMExtensionPtr->GetWidgetBlueprint();
+		check(WidgetBlueprint);
+		UMVVMEditorSubsystem* EditorSubsystem = GEditor->GetEditorSubsystem<UMVVMEditorSubsystem>();
+
+		TArray<TSharedPtr<FBindingEntry>> Selection = TreeView->GetSelectedItems();
+		for (const TSharedPtr<FBindingEntry>& Entry : Selection)
+		{
+			if (Entry->GetRowType() == FBindingEntry::ERowType::EventParameter)
+			{
+				EditorSubsystem->ResetOrphanedPin(WidgetBlueprint, Entry->GetEvent(), Entry->GetEventParameterName());
+			}
+			else if (Entry->GetRowType() == FBindingEntry::ERowType::BindingParameter)
+			{
+				if (FMVVMBlueprintViewBinding* Binding = Entry->GetBinding(MVVMExtensionPtr->GetBlueprintView()))
+				{
+					const bool bSourceToDestination = UE::MVVM::IsForwardBinding(Binding->BindingType);
+					EditorSubsystem->ResetOrphanedPin(WidgetBlueprint, *Binding, Entry->GetBindingParameterName(), bSourceToDestination);
 				}
 			}
 		}
@@ -2460,28 +2488,51 @@ TSharedPtr<SWidget> SBindingsList::OnSourceConstructContextMenu()
 	TArray<TSharedPtr<FBindingEntry>> Selection = TreeView->GetSelectedItems();
 	if (Selection.Num() > 0)
 	{
-		{
-			FUIAction RemoveAction;
-			RemoveAction.ExecuteAction = FExecuteAction::CreateSP(this, &SBindingsList::HandleDeleteSelected);
-			MenuBuilder.AddMenuEntry(LOCTEXT("RemoveBinding", "Remove Binding"),
-				LOCTEXT("RemoveBindingTooltip", "Remove this binding."),
-				FSlateIcon(FAppStyle::GetAppStyleSetName(), "Icons.Delete"),
-				RemoveAction);
-		}
-
-		{
-			const UMVVMWidgetBlueprintExtension_View* MVVMExtensionPtr = MVVMExtension.Get();
-			const UWidgetBlueprint* WidgetBlueprint = MVVMExtensionPtr ? MVVMExtensionPtr->GetWidgetBlueprint() : nullptr;
-			const UMVVMBlueprintView* View = MVVMExtensionPtr ? MVVMExtensionPtr->GetBlueprintView() : nullptr;
-			const UMVVMEditorSubsystem* EditorSubsystem = GEditor->GetEditorSubsystem<UMVVMEditorSubsystem>();
+		const UMVVMWidgetBlueprintExtension_View* MVVMExtensionPtr = MVVMExtension.Get();
+		const UWidgetBlueprint* WidgetBlueprint = MVVMExtensionPtr ? MVVMExtensionPtr->GetWidgetBlueprint() : nullptr;
+		const UMVVMBlueprintView* View = MVVMExtensionPtr ? MVVMExtensionPtr->GetBlueprintView() : nullptr;
+		const UMVVMEditorSubsystem* EditorSubsystem = GEditor->GetEditorSubsystem<UMVVMEditorSubsystem>();
+
+		if (View && WidgetBlueprint)
+		{
+			bool bCanRemoveEntry = false;
+			bool bCanSplitPin = false;
+			bool bCanRecombinePin = false;
+			bool bCanRecombinePinVisible = false;
+			bool bCanResetPin = false;
+			bool bCanResetPinVisible = false;
+			bool bCanResetOrphanedPin = false;
+			auto SetAll = [&](bool Value)
+			{
+				bCanRemoveEntry = Value;
+				bCanSplitPin = Value;
+				bCanRecombinePin = Value;
+				bCanRecombinePinVisible = Value;
+				bCanResetPin = Value;
+				bCanResetPinVisible = Value;
+				bCanResetOrphanedPin = Value;
+			};
+			auto AllFalse = [&]()
+			{
+				return !bCanRemoveEntry && !bCanRecombinePin && !bCanSplitPin && !bCanRecombinePinVisible && bCanResetPin && !bCanResetPinVisible && !bCanResetOrphanedPin;
+			};
 
-			if (View && WidgetBlueprint)
+			if (Selection.Num())
 			{
-				bool bCanSplitPin = true;
-				bool bCanRecombinePin = true;
-				bool bCanRecombinePinVisible = true;
+				SetAll(true);
 				for (const TSharedPtr<FBindingEntry>& Entry : Selection)
 				{
+					switch(Entry->GetRowType())
+					{
+						case FBindingEntry::ERowType::Group:
+						case FBindingEntry::ERowType::Binding:
+						case FBindingEntry::ERowType::Event:
+							break;
+						default:
+							bCanRemoveEntry = false;
+							break;
+					}
+
 					if (Entry->GetRowType() == FBindingEntry::ERowType::EventParameter)
 					{
 						if (!EditorSubsystem->CanSplitPin(WidgetBlueprint, Entry->GetEvent(), Entry->GetEventParameterName()))
@@ -2492,11 +2543,32 @@ TSharedPtr<SWidget> SBindingsList::OnSourceConstructContextMenu()
 						{
 							bCanRecombinePin = false;
 						}
+						if (!EditorSubsystem->CanResetPinToDefaultValue(WidgetBlueprint, Entry->GetEvent(), Entry->GetEventParameterName()))
+						{
+							bCanResetPin = false;
+						}
+						if (!EditorSubsystem->CanResetOrphanedPin(WidgetBlueprint, Entry->GetEvent(), Entry->GetEventParameterName()))
+						{
+							bCanResetOrphanedPin = false;
+						}
+
 						UMVVMBlueprintViewEvent* ViewEvent = Entry->GetEvent();
 						UEdGraphPin* GraphPin = ViewEvent ? ViewEvent->GetOrCreateGraphPin(Entry->GetEventParameterName()) : nullptr;
-						if (GraphPin == nullptr || GraphPin->ParentPin == nullptr)
+						if (GraphPin == nullptr)
 						{
 							bCanRecombinePinVisible = false;
+							bCanResetPinVisible = false;
+						}
+						else
+						{
+							if (GraphPin->ParentPin == nullptr)
+							{
+								bCanRecombinePinVisible = false;
+							}
+							if (!FMVVMBlueprintPin::IsInputPin(GraphPin) || GetDefault<UEdGraphSchema_K2>()->ShouldHidePinDefaultValue(GraphPin))
+							{
+								bCanResetPinVisible = false;
+							}
 						}
 					}
 					else if (Entry->GetRowType() == FBindingEntry::ERowType::BindingParameter)
@@ -2513,46 +2585,97 @@ TSharedPtr<SWidget> SBindingsList::OnSourceConstructContextMenu()
 							{
 								bCanRecombinePin = false;
 							}
+							if (!EditorSubsystem->CanResetPinToDefaultValue(WidgetBlueprint, *Binding, Entry->GetBindingParameterName(), bSourceToDestination))
+							{
+								bCanResetPin = false;
+							}
+							if (!EditorSubsystem->CanResetOrphanedPin(WidgetBlueprint, *Binding, Entry->GetBindingParameterName(), bSourceToDestination))
+							{
+								bCanResetOrphanedPin = false;
+							}
 
 							UEdGraphPin* GraphPin = EditorSubsystem->GetConversionFunctionArgumentPin(WidgetBlueprint, *Binding, Entry->GetBindingParameterName(), bSourceToDestination);
-							if (GraphPin == nullptr || GraphPin->ParentPin == nullptr)
+							if (GraphPin == nullptr)
 							{
 								bCanRecombinePinVisible = false;
+								bCanResetPinVisible = false;
+							}
+							else
+							{
+								if (GraphPin->ParentPin == nullptr)
+								{
+									bCanRecombinePinVisible = false;
+								}
+								if (!FMVVMBlueprintPin::IsInputPin(GraphPin) || GetDefault<UEdGraphSchema_K2>()->ShouldHidePinDefaultValue(GraphPin))
+								{
+									bCanResetPinVisible = false;
+								}
 							}
 						}
 						else
 						{
-							bCanSplitPin = false;
-							bCanRecombinePin = false;
-							bCanRecombinePinVisible = false;
+							SetAll(false);
 						}
 					}
+					else
+					{
+						SetAll(false);
+					}
 
-					if (!bCanRecombinePin && !bCanSplitPin && !bCanRecombinePinVisible)
+					if (AllFalse())
 					{
 						break;
 					}
 				}
+			}
 
-				if (bCanSplitPin)
-				{
-					FUIAction SplitPinAction;
-					SplitPinAction.ExecuteAction = FExecuteAction::CreateSP(this, &SBindingsList::HandleBreakSelectedPin);
-					MenuBuilder.AddMenuEntry(LOCTEXT("BreakPin", "Split Struct Pin"),
-						LOCTEXT("BreakPinTooltip", "Breaks a struct pin in to a separate pin per element."),
-						FSlateIcon(),
-						SplitPinAction);
-				}
-				if (bCanRecombinePinVisible)
-				{
-					FUIAction RecombinePinAction;
-					RecombinePinAction.ExecuteAction = FExecuteAction::CreateSP(this, &SBindingsList::HandleRecombineSelectedPin);
-					RecombinePinAction.CanExecuteAction = FCanExecuteAction::CreateLambda([bCanRecombinePin](){ return bCanRecombinePin; });
-					MenuBuilder.AddMenuEntry(LOCTEXT("RecombinePin", "Recombine Struct Pin"),
-						LOCTEXT("RecombinePinTooltip", "Takes struct pins that have been broken in to composite elements and combines them back to a single struct pin."),
-						FSlateIcon(),
-						RecombinePinAction);
-				}
+			{
+				FUIAction RemoveAction;
+				RemoveAction.ExecuteAction = FExecuteAction::CreateSP(this, &SBindingsList::HandleDeleteSelected);
+				RemoveAction.CanExecuteAction = FCanExecuteAction::CreateLambda([bCanRemoveEntry]() { return bCanRemoveEntry; });
+				MenuBuilder.AddMenuEntry(LOCTEXT("RemoveBinding", "Remove Binding"),
+					LOCTEXT("RemoveBindingTooltip", "Remove this binding."),
+					FSlateIcon(FAppStyle::GetAppStyleSetName(), "Icons.Delete"),
+					RemoveAction);
+			}
+
+			if (bCanResetPinVisible)
+			{
+				FUIAction ResetPinAction;
+				ResetPinAction.ExecuteAction = FExecuteAction::CreateSP(this, &SBindingsList::HandleResetSelectedPin);
+				ResetPinAction.CanExecuteAction = FCanExecuteAction::CreateLambda([bCanResetPin]() { return bCanResetPin; });
+				MenuBuilder.AddMenuEntry(LOCTEXT("ResetPin", "Reset to Default Value"),
+					LOCTEXT("ResetPinTooltip", "Reset value of this pin to the default"),
+					FSlateIcon(),
+					ResetPinAction);
+			}
+			if (bCanSplitPin)
+			{
+				FUIAction SplitPinAction;
+				SplitPinAction.ExecuteAction = FExecuteAction::CreateSP(this, &SBindingsList::HandleBreakSelectedPin);
+				MenuBuilder.AddMenuEntry(LOCTEXT("BreakPin", "Split Struct Pin"),
+					LOCTEXT("BreakPinTooltip", "Breaks a struct pin in to a separate pin per element."),
+					FSlateIcon(),
+					SplitPinAction);
+			}
+			if (bCanRecombinePinVisible)
+			{
+				FUIAction RecombinePinAction;
+				RecombinePinAction.ExecuteAction = FExecuteAction::CreateSP(this, &SBindingsList::HandleRecombineSelectedPin);
+				RecombinePinAction.CanExecuteAction = FCanExecuteAction::CreateLambda([bCanRecombinePin](){ return bCanRecombinePin; });
+				MenuBuilder.AddMenuEntry(LOCTEXT("RecombinePin", "Recombine Struct Pin"),
+					LOCTEXT("RecombinePinTooltip", "Takes struct pins that have been broken in to composite elements and combines them back to a single struct pin."),
+					FSlateIcon(),
+					RecombinePinAction);
+			}
+			if (bCanResetOrphanedPin)
+			{
+				FUIAction ResetOrphanedPinAction;
+				ResetOrphanedPinAction.ExecuteAction = FExecuteAction::CreateSP(this, &SBindingsList::HandleResetOrphanedSelectedPin);
+				MenuBuilder.AddMenuEntry(LOCTEXT("ResetOrphanedPin", "Remove the Orphaned Struct Pin"),
+					LOCTEXT("ResetOrphanedPinTooltip", "Removes pins that used to exist but do not exist anymore."),
+					FSlateIcon(),
+					ResetOrphanedPinAction);
 			}
 		}
 	}
diff --git a/Engine/Plugins/Runtime/ModelViewViewModel/Source/ModelViewViewModelEditor/Private/Widgets/SMVVMViewBindingListView.h b/Engine/Plugins/Runtime/ModelViewViewModel/Source/ModelViewViewModelEditor/Private/Widgets/SMVVMViewBindingListView.h
--- a/Engine/Plugins/Runtime/ModelViewViewModel/Source/ModelViewViewModelEditor/Private/Widgets/SMVVMViewBindingListView.h
+++ b/Engine/Plugins/Runtime/ModelViewViewModel/Source/ModelViewViewModelEditor/Private/Widgets/SMVVMViewBindingListView.h
@@ -46,8 +46,10 @@ private:
 
 	virtual FReply OnKeyDown(const FGeometry& MyGeometry, const FKeyEvent& InKeyEvent) override;
 	void HandleDeleteSelected();
+	void HandleResetSelectedPin();
 	void HandleBreakSelectedPin();
 	void HandleRecombineSelectedPin();
+	void HandleResetOrphanedSelectedPin();
 
 	void ForceRefresh();
 
diff --git a/Engine/Plugins/Runtime/ModelViewViewModel/Source/ModelViewViewModelEditor/Private/Widgets/SMVVMViewModelPanel.cpp b/Engine/Plugins/Runtime/ModelViewViewModel/Source/ModelViewViewModelEditor/Private/Widgets/SMVVMViewModelPanel.cpp
--- a/Engine/Plugins/Runtime/ModelViewViewModel/Source/ModelViewViewModelEditor/Private/Widgets/SMVVMViewModelPanel.cpp
+++ b/Engine/Plugins/Runtime/ModelViewViewModel/Source/ModelViewViewModelEditor/Private/Widgets/SMVVMViewModelPanel.cpp
@@ -287,7 +287,7 @@ void SMVVMViewModelPanel::Construct(const FArguments& InArgs, TSharedPtr<FWidget
 				SNew(SWarningOrErrorBox)
 				.Visibility(this, &SMVVMViewModelPanel::GetWarningPanelVisibility)
 				.MessageStyle(EMessageStyle::Warning)
-				.Message(LOCTEXT("InitializationPanelWarningDescription", "The view will not initialize automatically. It was manually set in the View Settings."))
+				.Message(this, &SMVVMViewModelPanel::GetWarningMessage)
 				[
 					SNew(SButton)
 					.OnClicked(this, &SMVVMViewModelPanel::HandleDisableWarningPanel)
@@ -886,6 +886,10 @@ EVisibility SMVVMViewModelPanel::GetWarningPanelVisibility() const
 	{
 		if (UMVVMBlueprintView* WidgetBlueprint = WeakBlueprintView.Get())
 		{
+			if (WidgetBlueprint->GetNumBindings() == 0 && WidgetBlueprint->GetViewModels().Num() != 0)
+			{
+				return EVisibility::Visible;
+			}
 			if (!WidgetBlueprint->GetSettings()->bInitializeSourcesOnConstruct || !WidgetBlueprint->GetSettings()->bInitializeBindingsOnConstruct)
 			{
 				return EVisibility::Visible;
@@ -897,6 +901,26 @@ EVisibility SMVVMViewModelPanel::GetWarningPanelVisibility() const
 }
 
 
+FText SMVVMViewModelPanel::GetWarningMessage() const
+{
+	if (!bDisableWarningPanel)
+	{
+		if (UMVVMBlueprintView* WidgetBlueprint = WeakBlueprintView.Get())
+		{
+			if (WidgetBlueprint->GetNumBindings() == 0 && WidgetBlueprint->GetViewModels().Num() != 0)
+			{
+				return LOCTEXT("NoBindingWarningDescription", "No view will be created for this widget because there are no bindings. Your viewmodels won't be initialized.");
+			}
+			if (!WidgetBlueprint->GetSettings()->bInitializeSourcesOnConstruct || !WidgetBlueprint->GetSettings()->bInitializeBindingsOnConstruct)
+			{
+				return LOCTEXT("InitializationPanelWarningDescription", "The view will not initialize automatically. It was manually set in the View Settings.");
+			}
+		}
+	}
+	return FText::GetEmpty();
+}
+
+
 FReply SMVVMViewModelPanel::HandleDisableWarningPanel()
 {
 	bDisableWarningPanel = true;
diff --git a/Engine/Plugins/Runtime/ModelViewViewModel/Source/ModelViewViewModelEditor/Private/Widgets/SMVVMViewModelPanel.h b/Engine/Plugins/Runtime/ModelViewViewModel/Source/ModelViewViewModelEditor/Private/Widgets/SMVVMViewModelPanel.h
--- a/Engine/Plugins/Runtime/ModelViewViewModel/Source/ModelViewViewModelEditor/Private/Widgets/SMVVMViewModelPanel.h
+++ b/Engine/Plugins/Runtime/ModelViewViewModel/Source/ModelViewViewModelEditor/Private/Widgets/SMVVMViewModelPanel.h
@@ -98,6 +98,7 @@ private:
 	TSharedRef<SWidget> HandleAddViewModelContextMenu();
 
 	EVisibility GetWarningPanelVisibility() const;
+	FText GetWarningMessage() const;
 	FReply HandleDisableWarningPanel();
 
 	void CreateCommandList();
