MVVM: Big rework of the MVVM view compiler. It add supports for conversion function to go from complex to simple. Supports pin split on simple conversion function. Add ability to to support setter fun

diff --git a/Engine/Plugins/Runtime/ModelViewViewModel/Source/ModelViewViewModelEditor/Public/MVVMEditorSubsystem.h b/Engine/Plugins/Runtime/ModelViewViewModel/Source/ModelViewViewModelEditor/Public/MVVMEditorSubsystem.h
--- a/Engine/Plugins/Runtime/ModelViewViewModel/Source/ModelViewViewModelEditor/Public/MVVMEditorSubsystem.h
+++ b/Engine/Plugins/Runtime/ModelViewViewModel/Source/ModelViewViewModelEditor/Public/MVVMEditorSubsystem.h
@@ -89,7 +89,7 @@ public:
 	bool IsValidConversionFunction(const UWidgetBlueprint* WidgeteBlueprint, const UFunction* Function, const FMVVMBlueprintPropertyPath& Source, const FMVVMBlueprintPropertyPath& Destination) const;
 
 	UFUNCTION(BlueprintCallable, Category = "Viewmodel")
-	bool IsSimpleConversionFunction(const UFunction* Function) const;
+	bool IsSimpleConversionFunctionA(const UFunction* Function) const;
 
 	UFUNCTION(BlueprintCallable, Category = "Viewmodel")
 	UEdGraph* GetConversionFunctionGraph(const UWidgetBlueprint* WidgetBlueprint, const FMVVMBlueprintViewBinding& Binding, bool bSourceToDestination) const;
@@ -116,6 +116,14 @@ public:
 	bool CanRecombinePin(const UWidgetBlueprint* WidgetBlueprint, const FMVVMBlueprintViewBinding& Binding, FName PinName, bool bSourceToDestination) const;
 	void RecombinePin(UWidgetBlueprint* WidgetBlueprint, UMVVMBlueprintViewEvent* Event, FName PinName) const;
 	bool CanRecombinePin(const UWidgetBlueprint* WidgetBlueprint, UMVVMBlueprintViewEvent* Event, FName PinName) const;
+	void ResetPinToDefaultValue(UWidgetBlueprint* WidgetBlueprint, FMVVMBlueprintViewBinding& Binding, FName PinName, bool bSourceToDestination) const;
+	bool CanResetPinToDefaultValue(const UWidgetBlueprint* WidgetBlueprint, const FMVVMBlueprintViewBinding& Binding, FName PinName, bool bSourceToDestination) const;
+	void ResetPinToDefaultValue(UWidgetBlueprint* WidgetBlueprint, UMVVMBlueprintViewEvent* Event, FName PinName) const;
+	bool CanResetPinToDefaultValue(const UWidgetBlueprint* WidgetBlueprint, UMVVMBlueprintViewEvent* Event, FName PinName) const;
+	void ResetOrphanedPin(UWidgetBlueprint* WidgetBlueprint, FMVVMBlueprintViewBinding& Binding, FName PinName, bool bSourceToDestination) const;
+	bool CanResetOrphanedPin(const UWidgetBlueprint* WidgetBlueprint, const FMVVMBlueprintViewBinding& Binding, FName PinName, bool bSourceToDestination) const;
+	void ResetOrphanedPin(UWidgetBlueprint* WidgetBlueprint, UMVVMBlueprintViewEvent* Event, FName PinName) const;
+	bool CanResetOrphanedPin(const UWidgetBlueprint* WidgetBlueprint, UMVVMBlueprintViewEvent* Event, FName PinName) const;
 
 	TArray<UE::MVVM::FBindingSource> GetBindableWidgets(const UWidgetBlueprint* WidgetBlueprint) const;
 	TArray<UE::MVVM::FBindingSource> GetAllViewModels(const UWidgetBlueprint* WidgetBlueprint) const;
diff --git a/Engine/Source/Editor/UMGEditor/Private/WidgetBlueprintCompiler.cpp b/Engine/Source/Editor/UMGEditor/Private/WidgetBlueprintCompiler.cpp
--- a/Engine/Source/Editor/UMGEditor/Private/WidgetBlueprintCompiler.cpp
+++ b/Engine/Source/Editor/UMGEditor/Private/WidgetBlueprintCompiler.cpp
@@ -71,6 +71,23 @@ EKismetCompileType::Type FWidgetBlueprintCompilerContext::FCreateVariableContext
 }
 
 
+//////////////////////////////////////////////////////////////////////////
+// FWidgetBlueprintCompiler::FCreateFunctionContext
+FWidgetBlueprintCompilerContext::FCreateFunctionContext::FCreateFunctionContext(FWidgetBlueprintCompilerContext& InContext)
+	: Context(InContext)
+{}
+
+void FWidgetBlueprintCompilerContext::FCreateFunctionContext::AddGeneratedFunctionGraph(UEdGraph* Graph) const
+{
+	Context.GeneratedFunctionGraphs.Add(Graph);
+}
+
+UWidgetBlueprintGeneratedClass* FWidgetBlueprintCompilerContext::FCreateFunctionContext::GetGeneratedClass() const
+{
+	return Context.NewWidgetBlueprintClass;
+}
+
+
 //////////////////////////////////////////////////////////////////////////
 // FWidgetBlueprintCompiler
 FWidgetBlueprintCompiler::FWidgetBlueprintCompiler()
@@ -164,9 +181,9 @@ UEdGraphSchema_K2* FWidgetBlueprintCompilerContext::CreateSchema()
 
 void FWidgetBlueprintCompilerContext::CreateFunctionList()
 {
-	UWidgetBlueprintExtension::ForEachExtension(WidgetBlueprint(), [this](UWidgetBlueprintExtension* InExtension)
+	UWidgetBlueprintExtension::ForEachExtension(WidgetBlueprint(), [Self = this](UWidgetBlueprintExtension* InExtension)
 		{
-			InExtension->CreateFunctionList();
+			InExtension->CreateFunctionList(FCreateFunctionContext(*Self));
 		});
 
 	Super::CreateFunctionList();
diff --git a/Engine/Source/Editor/UMGEditor/Public/WidgetBlueprintCompiler.h b/Engine/Source/Editor/UMGEditor/Public/WidgetBlueprintCompiler.h
--- a/Engine/Source/Editor/UMGEditor/Public/WidgetBlueprintCompiler.h
+++ b/Engine/Source/Editor/UMGEditor/Public/WidgetBlueprintCompiler.h
@@ -113,6 +113,18 @@ public:
 		FWidgetBlueprintCompilerContext& Context;
 	};
 
+	struct UMGEDITOR_API FCreateFunctionContext
+	{
+	public:
+		void AddGeneratedFunctionGraph(UEdGraph*) const;
+		UWidgetBlueprintGeneratedClass* GetGeneratedClass() const;
+
+	private:
+		friend FWidgetBlueprintCompilerContext;
+		FCreateFunctionContext(FWidgetBlueprintCompilerContext& InContext);
+		FWidgetBlueprintCompilerContext& Context;
+	};
+
 protected:
 	void FixAbandonedWidgetTree(UWidgetBlueprint* WidgetBP);
 
diff --git a/Engine/Source/Editor/UMGEditor/Public/WidgetBlueprintExtension.h b/Engine/Source/Editor/UMGEditor/Public/WidgetBlueprintExtension.h
--- a/Engine/Source/Editor/UMGEditor/Public/WidgetBlueprintExtension.h
+++ b/Engine/Source/Editor/UMGEditor/Public/WidgetBlueprintExtension.h
@@ -72,7 +72,7 @@ protected:
 	 */
 	virtual void HandleBeginCompilation(FWidgetBlueprintCompilerContext& InCreationContext) {}
 
-	virtual void HandleCreateFunctionList() {}
+	virtual void HandleCreateFunctionList(const FWidgetBlueprintCompilerContext::FCreateFunctionContext& InCreationContext) {}
 	virtual void HandleCleanAndSanitizeClass(UWidgetBlueprintGeneratedClass* ClassToClean, UObject* OldCDO) {}
 	virtual TArray<UObject*> HandleSaveSubObjectsFromCleanAndSanitizeClass(UWidgetBlueprintGeneratedClass* ClassToClean) { return TArray<UObject*>();  }
 	virtual void HandleCreateClassVariablesFromBlueprint(const FWidgetBlueprintCompilerContext::FCreateVariableContext& Context) {}
@@ -95,9 +95,9 @@ private:
 		HandleBeginCompilation(InCreationContext);
 	}
 
-	void CreateFunctionList()
+	void CreateFunctionList(const FWidgetBlueprintCompilerContext::FCreateFunctionContext& InCreationContext)
 	{
-		HandleCreateFunctionList();
+		HandleCreateFunctionList(InCreationContext);
 	}
 
 	void CleanAndSanitizeClass(UWidgetBlueprintGeneratedClass* ClassToClean, UObject* OldCDO)
