[Backout] - CL29797661

diff --git a/Engine/Plugins/Runtime/ModelViewViewModel/Source/ModelViewViewModelBlueprint/Private/MVVMViewBlueprintCompiler.h b/Engine/Plugins/Runtime/ModelViewViewModel/Source/ModelViewViewModelBlueprint/Private/MVVMViewBlueprintCompiler.h
--- a/Engine/Plugins/Runtime/ModelViewViewModel/Source/ModelViewViewModelBlueprint/Private/MVVMViewBlueprintCompiler.h
+++ b/Engine/Plugins/Runtime/ModelViewViewModel/Source/ModelViewViewModelBlueprint/Private/MVVMViewBlueprintCompiler.h
@@ -3,12 +3,9 @@
 #pragma once
 
 #include "Bindings/MVVMCompiledBindingLibraryCompiler.h"
-#include "MVVMBlueprintView.h"
 #include "MVVMBlueprintViewEvent.h"
 #include "MVVMBlueprintViewModelContext.h"
-#include "Templates/ValueOrError.h"
 #include "Types/MVVMFieldVariant.h"
-#include "UObject/StrongObjectPtr.h"
 #include "WidgetBlueprintCompiler.h"
 
 struct FMVVMBlueprintPropertyPath;
@@ -18,8 +15,6 @@ struct FMVVMBlueprintViewBinding;
 class FWidgetBlueprintCompilerContext;
 class UEdGraph;
 class UMVVMBlueprintView;
-class UMVVMBlueprintViewConversionFunction;
-class UMVVMBlueprintViewEvent;
 class UMVVMViewClass;
 class UWidgetBlueprintGeneratedClass;
 namespace  UE::MVVM
@@ -33,19 +28,15 @@ namespace UE::MVVM::Private
 struct FMVVMViewBlueprintCompiler
 {
 private:
-	struct FCompilerBindingSource;
-	struct FCompilerUserWidgetProperty;
-	struct FCompilerViewModelSetter;
-	struct FCompilerViewModelCreatorContext;
-	struct FCompilerSourceViewModelDynamicCreatorContext;
-	struct FCompilerNotifyFieldId;
-	struct FGeneratedReadFieldPathContext;
-	struct FGeneratedWriteFieldPathContext;
+	struct FCompilerUserWidgetPropertyContext;
+	struct FCompilerSourceCreatorContext;
+	struct FCompiledBinding;
 	struct FCompilerBinding;
-	struct FCompilerEvent;
+	struct FBindingSourceContext;
+	struct FBindingDestinationContext;
 
 public:
-	FMVVMViewBlueprintCompiler(FWidgetBlueprintCompilerContext& InCreationContext, UMVVMBlueprintView* BlueprintView);
+	FMVVMViewBlueprintCompiler(FWidgetBlueprintCompilerContext& InCreationContext);
 
 	FWidgetBlueprintCompilerContext& GetCompilerContext()
 	{
@@ -55,321 +46,236 @@ public:
 	void AddExtension(UWidgetBlueprintGeneratedClass* Class, UMVVMViewClass* ViewExtension);
 	void CleanOldData(UWidgetBlueprintGeneratedClass* ClassToClean, UObject* OldCDO);
 
-	/** Generate variable and public function in the Skeleton class and the generated class */
-	void CreateVariables(const FWidgetBlueprintCompilerContext::FCreateVariableContext& Context);
 	/** Generate function that are hidden from the user (not on the Skeleton class). */
-	void CreateFunctions(const FWidgetBlueprintCompilerContext::FCreateFunctionContext& Context);
+	void CreateFunctions(UMVVMBlueprintView* BlueprintView);
+	/** Generate variable in the Skeleton class */
+	void CreateVariables(const FWidgetBlueprintCompilerContext::FCreateVariableContext& Context, UMVVMBlueprintView* BlueprintView);
 
 	/** Add all the field path and the bindings to the library compiler. */
-	bool PreCompile(UWidgetBlueprintGeneratedClass* Class);
+	bool PreCompile(UWidgetBlueprintGeneratedClass* Class, UMVVMBlueprintView* BlueprintView);
 	/** Compile the library and fill the view and viewclass */
-	bool Compile(UWidgetBlueprintGeneratedClass* Class, UMVVMViewClass* ViewExtension);
+	bool Compile(UWidgetBlueprintGeneratedClass* Class, UMVVMBlueprintView* BlueprintView, UMVVMViewClass* ViewExtension);
 
 
 	static void TestGenerateSetter(const UBlueprint* Context, FStringView ObjectName, FStringView FieldPath, FStringView FunctionName);
 
 private:
-	bool AreStepsValid() const
-	{
-		return bIsCreateVariableStepValid && bIsCreateFunctionsStepValid && bIsPreCompileStepValid && bIsCompileStepValid;
-	}
+	void CreateWidgetMap(const FWidgetBlueprintCompilerContext::FCreateVariableContext& Context, UMVVMBlueprintView* BlueprintView);
+	void CreateSourceLists(const FWidgetBlueprintCompilerContext::FCreateVariableContext& Context, UMVVMBlueprintView* BlueprintView);
+	void CreateFunctionsDeclaration(const FWidgetBlueprintCompilerContext::FCreateVariableContext& Context, UMVVMBlueprintView* BlueprintView);
+	void CreateIntermediateGraphFunctions(const FWidgetBlueprintCompilerContext::FCreateVariableContext& Context, UMVVMBlueprintView* BlueprintView);
+	void CreateBindingDestinationContexts(UMVVMBlueprintView* BlueprintView);
 
-	// CreateVariables
-	void CreateWidgetMap(const FWidgetBlueprintCompilerContext::FCreateVariableContext& Context);
-	void CreateBindingList(const FWidgetBlueprintCompilerContext::FCreateVariableContext& Context);
-	void CreateEventList(const FWidgetBlueprintCompilerContext::FCreateVariableContext& Context);
-	void CreateRequiredProperties(const FWidgetBlueprintCompilerContext::FCreateVariableContext& Context);
-	void CreatePublicFunctionsDeclaration(const FWidgetBlueprintCompilerContext::FCreateVariableContext& Context);
-
-	// CreateFuntions
-	void CategorizeBindings(const FWidgetBlueprintCompilerContext::FCreateFunctionContext& Context);
-	void CategorizeEvents(const FWidgetBlueprintCompilerContext::FCreateFunctionContext& Context);
-	void CreateWriteFieldContexts(const FWidgetBlueprintCompilerContext::FCreateFunctionContext& Context);
-	void CreateViewModelSetters(const FWidgetBlueprintCompilerContext::FCreateFunctionContext& Context);
-	void CreateIntermediateGraphFunctions(const FWidgetBlueprintCompilerContext::FCreateFunctionContext& Context);
-
-	// PreCompile
-	void FixCompilerBindingSelfSource(UWidgetBlueprintGeneratedClass* Class);
-	void AddWarningForPropertyWithMVVMAndLegacyBinding(UWidgetBlueprintGeneratedClass* Class);
-	void FixWriteFieldPathContext(UWidgetBlueprintGeneratedClass* Class);
-	void CreateReadFieldContexts(UWidgetBlueprintGeneratedClass* Class);
-	void CreateCreatorContentFromBindingSource(UWidgetBlueprintGeneratedClass* Class);
-	void PreCompileViewModelCreatorContexts(UWidgetBlueprintGeneratedClass* Class);
-	void PreCompileBindings(UWidgetBlueprintGeneratedClass* Class);
-	void PreCompileEvents(UWidgetBlueprintGeneratedClass* Class);
-
-	// Compile
-	void CompileViewModelCreatorContexts(const FCompiledBindingLibraryCompiler::FCompileResult& CompileResult, UWidgetBlueprintGeneratedClass* Class, UMVVMViewClass* ViewExtension);
-	void CompileBindings(const FCompiledBindingLibraryCompiler::FCompileResult& CompileResult, UWidgetBlueprintGeneratedClass* Class, UMVVMViewClass* ViewExtension);
-	void CompileEvents(const FCompiledBindingLibraryCompiler::FCompileResult& CompileResult, UWidgetBlueprintGeneratedClass* Class, UMVVMViewClass* ViewExtension);
+	bool PreCompileBindingSources(UWidgetBlueprintGeneratedClass* Class, UMVVMBlueprintView* BlueprintView);
+	bool CompileBindingSources(const FCompiledBindingLibraryCompiler::FCompileResult& CompileResult, UWidgetBlueprintGeneratedClass* Class, UMVVMBlueprintView* BlueprintView, UMVVMViewClass* ViewExtension);
 
-private:
-	/**
-	 * List of all the sources needed by the view to register the bindings.
-	 * They could be viewmodel, widget or any properties on the UserWidget.
-	 * They could also be a viewmodel in a long path.
-	 * It may not have an associated property (dynamicviewmodel).
-	 * It may only have OneTime binding
-	 */
-	struct FCompilerBindingSource
-	{
-		enum class EType
-		{
-			ViewModel = 0,
-			DynamicViewmodel = 1,
-			Widget = 2,
-			Self = 3,
-		};
-		const UClass* AuthoritativeClass = nullptr;
-		FName Name;
-		EType Type;
-		bool bIsOptional = false;
-	};
-	TArray<TSharedRef<FCompilerBindingSource>> NeededBindingSources;
+	bool PreCompileSourceCreators(UWidgetBlueprintGeneratedClass* Class, UMVVMBlueprintView* BlueprintView);
+	bool CompileSourceCreators(const FCompiledBindingLibraryCompiler::FCompileResult& CompileResult, UWidgetBlueprintGeneratedClass* Class, UMVVMBlueprintView* BlueprintView, UMVVMViewClass* ViewExtension);
+
+	bool PreCompileBindings(UWidgetBlueprintGeneratedClass* Class, UMVVMBlueprintView* BlueprintView);
+	bool CompileBindings(const FCompiledBindingLibraryCompiler::FCompileResult& CompileResult, UWidgetBlueprintGeneratedClass* Class, UMVVMBlueprintView* BlueprintView, UMVVMViewClass* ViewExtension);
+	
+	bool PreCompileEvents(UWidgetBlueprintGeneratedClass* Class, UMVVMBlueprintView* BlueprintView);
+	bool CompileEvents(const FCompiledBindingLibraryCompiler::FCompileResult& CompileResult, UWidgetBlueprintGeneratedClass* Class, UMVVMBlueprintView* BlueprintView, UMVVMViewClass* ViewExtension);
 
+	const FCompilerSourceCreatorContext* FindViewModelSource(FGuid Id) const;
+
+	void AddMessageForBinding(FMVVMBlueprintViewBinding& Binding, UMVVMBlueprintView* BlueprintView, const FText& MessageText, EBindingMessageType MessageType, FName ArgumentName = FName()) const;
+	void AddMessageForEvent(UMVVMBlueprintViewEvent* Event, const FText& MessageText, UMVVMBlueprintViewEvent::EMessageType MessageType, FName ArgumentName) const;
+	void AddErrorForViewModel(const FMVVMBlueprintViewModelContext& ViewModel, const FText& Message) const;
+
+	TValueOrError<FBindingSourceContext, FText> CreateBindingSourceContext(const UMVVMBlueprintView* BlueprintView, const UWidgetBlueprintGeneratedClass* Class, const FMVVMBlueprintPropertyPath& PropertyPath, bool bIsOneTimeBinding);
+	TArray<FMVVMConstFieldVariant> CreateBindingDestinationPath(const UMVVMBlueprintView* BlueprintView, const UWidgetBlueprintGeneratedClass* Class, const FMVVMBlueprintPropertyPath& PropertyPath) const;
+	TValueOrError<FCompiledBinding, FText> CreateCompiledBinding(const UWidgetBlueprintGeneratedClass* Class, TArrayView<const UE::MVVM::FMVVMConstFieldVariant> GetterFields, TArrayView<const UE::MVVM::FMVVMConstFieldVariant> SetterFields, const UFunction* ConversionFunction, bool bIsComplexBinding);
+
+	static TArray<FMVVMConstFieldVariant> CreatePropertyPath(const UClass* Class, FName PropertyName, TArray<FMVVMConstFieldVariant> Properties);
+	static bool IsPropertyPathValid(const UBlueprint* Context, TArrayView<const FMVVMConstFieldVariant> PropertyPath);
+	static bool CanBeSetInNative(TArrayView<const FMVVMConstFieldVariant> PropertyPath);
+
+private:
 	/**
-	 * Describe a Property that need to be added (if it doesn't already exist).
+	 * Describe the Property that need to be added (if it doesn't already exist)
 	 * They can be a Widget, a viewmodel, or any object owned by the UserWidget.
-	 * They can be source or destination.
 	 */
-	struct FCompilerUserWidgetProperty
+	struct FCompilerUserWidgetPropertyContext
 	{
-		enum class ECreationType
-		{
-			None,
-			CreateIfDoesntExist,
-			CreateOnlyIfDoesntExist,
-		};
-		UClass* AuthoritativeClass = nullptr;
-		FName Name;
+		UClass* Class = nullptr;
+		FName PropertyName;
 		FText DisplayName;
 		FString CategoryName;
 		FString BlueprintSetter;
-		ECreationType CreationType = ECreationType::None;
+		FMVVMConstFieldVariant Field;
+		// If the class is a viewmodel on the userwidget, what is the id of that viewmodel
+		FGuid ViewModelId;
+
 		bool bExposeOnSpawn = false;
 		bool bPrivate = false;
-		bool bReadOnly = false;
-
-		const FProperty* Property = nullptr;
 	};
-	TArray<FCompilerUserWidgetProperty> NeededUserWidgetProperties;
+	TArray<FCompilerUserWidgetPropertyContext> CompilerUserWidgetPropertyContexts;
 
-	/**
-	 * Describe a ViewModel generated setter function.
+	/** 
+	 * Describe the data initialize the view's properties/viewmodels.
 	 */
-	struct FCompilerViewModelSetter
+	enum class ECompilerSourceCreatorType
 	{
-		const UClass* Class = nullptr;
-		FName PropertyName;
-		FString BlueprintSetter;
-		FText DisplayName;
-
-		UEdGraph* SetterGraph = nullptr;
+		ViewModel, // added manually
+		ViewModelDynamic, // added by a long binding chain
 	};
-	TArray<FCompilerViewModelSetter> ViewModelSettersToGenerate;
-
-	/** 
-	 * Describe the data to initialize the view's properties/viewmodels.
-	 */
-	struct FCompilerViewModelCreatorContext
+	struct FCompilerSourceCreatorContext
 	{
 		FMVVMBlueprintViewModelContext ViewModelContext;
-		TSharedPtr<FCompilerBindingSource> Source;
-
-		TSharedPtr<FCompilerSourceViewModelDynamicCreatorContext> DynamicContext;
-		FCompiledBindingLibraryCompiler::FFieldPathHandle ReadPropertyPathHandle;
+		FCompiledBindingLibraryCompiler::FFieldPathHandle ReadPropertyPath;
+		ECompilerSourceCreatorType Type = ECompilerSourceCreatorType::ViewModel;
+		FString SetterFunctionName;
+		UEdGraph* SetterGraph = nullptr;
+		FName DynamicParentSourceName;
 	};
-	TArray<FCompilerViewModelCreatorContext> ViewModelCreatorContexts;
+	TArray<FCompilerSourceCreatorContext> CompilerSourceCreatorContexts;
 
-	/**
-	 * Describe the data to initialize the view's properties/widget.
-	 */
-	struct FCompilerWidgetCreatorContext
+	/** The compiled binding from the binding compiler. */
+	struct FCompiledBinding
 	{
-		TSharedPtr<FCompilerBindingSource> Source;
-		bool bSelfReference = false;
+		FCompiledBindingLibraryCompiler::FBindingHandle BindingHandle;
 
-		FCompiledBindingLibraryCompiler::FFieldPathHandle ReadPropertyPathHandle;
+		// At runtime, OnFieldValueChanged
+		//if the ConversionFunction exist
+		// then we read the Source
+		// then we call ConversionFunction with the source as an input
+		// then we copy the result of the ConversionFunction to DestinationWrite
+		//else if the ConversionFunction is complex
+		// then we call ConversionFunction without input (SourceRead can be invalid)
+		// then we copy the result of the ConversionFunction to DestinationWrite
+		//else
+		// then we copy SourceRead to DestiantionWrite.
+		FCompiledBindingLibraryCompiler::FFieldPathHandle SourceRead;
+		FCompiledBindingLibraryCompiler::FFieldPathHandle DestinationWrite;
+		FCompiledBindingLibraryCompiler::FFieldPathHandle ConversionFunction;
+
+		bool bIsConversionFunctionComplex = false;
 	};
-	TArray<FCompilerWidgetCreatorContext> WidgetCreatorContexts;
 
-	/**
-	 * Describe the data to initialize a viewmodel in a long path.
-	 * The viewmodel is not added in the BlueprintView but it needs to be dynamic added.
+	/** 
+	 * Describe a binding for the compiler.
+	 * 2 FCompilerBinding will be generated if the binding is TwoWays.
 	 */
-	struct FCompilerSourceViewModelDynamicCreatorContext
+	enum class ECompilerBindingType
 	{
-		TSharedPtr<FCompilerBindingSource> Source;
-		TSharedPtr<FCompilerBindingSource> ParentSource; // a dynamic always has a parent
-
-		FFieldNotificationId NotificationId;
-		FCompiledBindingLibraryCompiler::FFieldIdHandle NotificationIdLibraryCompilerHandle;
+		PropertyBinding, // normal binding
+		ViewModelDynamic, // added by a long binding chain
 	};
-	TArray<TSharedRef<FCompilerSourceViewModelDynamicCreatorContext>> SourceViewModelDynamicCreatorContexts;
+	struct FCompilerBinding
+	{
+		ECompilerBindingType Type = ECompilerBindingType::PropertyBinding;
+		int32 BindingIndex = INDEX_NONE;
+		int32 UserWidgetPropertyContextIndex = INDEX_NONE;
+		int32 SourceCreatorContextIndex = INDEX_NONE;
+		int32 ComplexConversionFunctionContextIndex = INDEX_NONE;
+		bool bSourceIsUserWidget = false;
+		bool bFieldIdNeeded = false;
+		bool bIsForwardBinding = false;
+		bool bNeedsValidSource = true;	// if false, the binding is a onetime or executed manually (onclick event).
+
+		FCompiledBindingLibraryCompiler::FFieldIdHandle FieldIdHandle;
+		FCompiledBinding CompiledBinding;
+		FName DynamicViewModelName;
+	};
+	TArray<FCompilerBinding> CompilerBindings;
 
 	/**
-	 * The field id for a specific ReadFieldPathContext
+	 * Binding that do not requires a source to execute.
+	 * They could be a OneTime ComplexConversionFunction or an event (onclick).
 	 */
-	struct FCompilerNotifyFieldId
+	struct FSimpleBindingContext
 	{
-		TArray<FGuid> BindingEditorKeys;
-		TArray<FGuid> EventKeys;
-
-		FFieldNotificationId NotificationId;
-		TSharedPtr<FCompilerBindingSource> Source;
-		TSharedPtr<FCompilerSourceViewModelDynamicCreatorContext> ViewModelDynamic;
+		int32 BindingIndex = INDEX_NONE;
+		bool bIsForwardBinding = false;
 
-		FCompiledBindingLibraryCompiler::FFieldIdHandle LibraryCompilerHandle;
+		int32 ComplexConversionFunctionContextIndex = INDEX_NONE;
 	};
-	TArray<TSharedPtr<FCompilerNotifyFieldId>> NotificationFields;
+	TArray<FSimpleBindingContext> SimpleBindingContexts;
 
 	/**
-	 * The source path we need to read from.
-	 * Can be any EMVVMBindingMode (OneTime, OneWay, ...)
+	 * Source path of a binding that contains a FieldId that we can register/bind to
 	 */
-	struct FGeneratedReadFieldPathContext
+	struct FBindingSourceContext
 	{
-		TArray<TWeakPtr<FCompilerBinding>> UsedByBindings;
-		TArray<TWeakPtr<FCompilerEvent>> UsedByEvents;
-
-		TSharedPtr<FCompilerBindingSource> OptionalSource; // Can be invalid if it's an event
-		TArray<UE::MVVM::FMVVMConstFieldVariant> GeneratedFields; // the string path converted into field
-		TArray<UE::MVVM::FMVVMConstFieldVariant> SkeletalGeneratedFields; // the field path converted with getter and setter
-
-		TSharedPtr<FCompilerNotifyFieldId> NotificationField;
-		FCompiledBindingLibraryCompiler::FFieldPathHandle LibraryCompilerHandle;
+		int32 BindingIndex = INDEX_NONE;
+		bool bIsForwardBinding = false;
+
+		const UClass* SourceClass = nullptr;
+		// The property that are registering to.
+		FFieldNotificationId FieldId;
+		// The path always start with the UserWidget as self.
+		// Viewmodel.Field.SubProperty.SubProperty
+		// or Widget.Field.SubProperty.SubProperty
+		// or Field.SubProperty.SubProperty
+		TArray<UE::MVVM::FMVVMConstFieldVariant> PropertyPath;
+
+		// Complex binding do not need to add the PropertyPath to the binding info.
+		int32 ComplexConversionFunctionContextIndex = INDEX_NONE;
+		// The source if it's a property
+		int32 UserWidgetPropertyContextIndex = INDEX_NONE;
+		// The source if it's a property
+		int32 SourceCreatorContextIndex = INDEX_NONE;
+		// The source is the UserWidget itself
+		bool bIsRootWidget = false;
 	};
-	TArray<TSharedRef<FGeneratedReadFieldPathContext>> GeneratedReadFieldPaths;
+	TArray<FBindingSourceContext> BindingSourceContexts;
 
 	/**
-	 * Destination path we need to write to.
-	 * Only if the bindings/events have a destination.
+	 * Destination path of a binding.
 	 * The info needs to be validate before we generate the functions list.
 	 */
-	struct FGeneratedWriteFieldPathContext
+	struct FBindingDestinationContext
 	{
-		TArray<TWeakPtr<FCompilerBinding>> UsedByBindings;
-		TArray<TWeakPtr<FCompilerEvent>> UsedByEvents;
+		int32 BindingIndex = INDEX_NONE;
+		bool bIsForwardBinding = false;
 
-		TArray<UE::MVVM::FMVVMConstFieldVariant> GeneratedFields; // the string path converted into field
-		TArray<UE::MVVM::FMVVMConstFieldVariant> SkeletalGeneratedFields; // the field path converted with getter and setter
-		EMVVMBlueprintFieldPathSource GeneratedFrom = EMVVMBlueprintFieldPathSource::None;
 		bool bCanBeSetInNative = true;
-		bool bUseByNativeBinding = false;
-
 		FName GeneratedFunctionName;
-		FCompiledBindingLibraryCompiler::FFieldPathHandle LibraryCompilerHandle;
+
 	};
-	TArray<TSharedRef<FGeneratedWriteFieldPathContext>> GeneratedWriteFieldPaths;
+	TArray<FBindingDestinationContext> BindingDestinationContexts;
 
 	/**
-	 * The list of all the valid binding to iterates on.
+	 * Complex Conversion function shared data
 	 */
-	struct FCompilerBinding
+	struct FComplexConversionFunctionContext
 	{
-		struct FKey
-		{
-			int32 ViewBindingIndex = INDEX_NONE;
-			bool bIsForwardBinding = false;
-			FKey() = default;
-			FKey(int32 InIndex, bool bInForward)
-				: ViewBindingIndex(InIndex)
-				, bIsForwardBinding(bInForward)
-			{}
-			bool operator==(const FKey& Other) const
-			{
-				return ViewBindingIndex == Other.ViewBindingIndex && bIsForwardBinding == Other.bIsForwardBinding;
-			}
-		};
-		enum class EType : int32
-		{
-			Unknown = 0, // was not evaluated yet
-			Invalid = -1, // evaluated and not valid
-			Assignment = 1, //Destination=Source
-			SimpleConversionFunction = 2, //Destination=Function(Source)
-			ComplexConversionFunction = 3, //Destination=Function(SourceA, SourceB)
-			//Function = 4, //Function(SourceA, SourceB)
-		};
-		FKey Key;
-		EType Type = EType::Unknown;
-		bool bIsOneTimeBinding = false;
-
-		TArray<TSharedPtr<FGeneratedReadFieldPathContext>> ReadPaths;
-		TSharedPtr<FGeneratedWriteFieldPathContext> WritePath;
-		TWeakObjectPtr<UMVVMBlueprintViewConversionFunction> ConversionFunction = nullptr;
-
-		FCompiledBindingLibraryCompiler::FBindingHandle BindingHandle;
-		FCompiledBindingLibraryCompiler::FFieldPathHandle ConversionFunctionHandle;
+		int32 BindingIndex = INDEX_NONE;
+		bool bIsForwardBinding = false;
+
+		// More than one binding can use the same conversion function.
+		// Only one should execute the binding at initialization.
+		bool bExecAtInitGenerated = false;
+		// More than one binding can use the same conversion function.
+		// At least one must have a valid source to trigger the binding.
+		//If not, then it should be a FSimpleBindingContext
+		bool bHasValidFieldId = false;
+		bool bNeedsValidSource = true;
 	};
-	TArray<TSharedRef<FCompilerBinding>> ValidBindings;
+	TArray<FComplexConversionFunctionContext> ComplexConversionFunctionContexts;
 
 	/**
-	 * The list of all the valid binding to iterates on.
+	 * Context for multicast delegate
 	 */
-	struct FCompilerEvent
+	struct FEventDelegateContext
 	{
-		enum class EType : int32
-		{
-			Unknown = 0, // was not evaluated yet
-			Invalid = -1, // evaluated and not valid
-			Valid = 1,
-		};
-		TWeakObjectPtr<UMVVMBlueprintViewEvent> Event = nullptr;
-		EType Type = EType::Unknown;
-
-		TArray<TSharedPtr<FGeneratedReadFieldPathContext>> ReadPaths;
-		TSharedPtr<FGeneratedWriteFieldPathContext> WritePath;
-
-		TSharedPtr<FGeneratedWriteFieldPathContext> DelegateFieldPath;
+		UMVVMBlueprintViewEvent* BlueprintEvent = nullptr;
 
+		FCompiledBindingLibraryCompiler::FFieldPathHandle DelegateFieldPathHandle;
 		FName GeneratedGraphName;
 		FName SourceName;
-		FCompiledBindingLibraryCompiler::FFieldPathHandle DelegateFieldPathHandle;
 	};
-	TArray<TSharedRef<FCompilerEvent>> ValidEvents;
+	TArray<FEventDelegateContext> EventDelegateContexts;
 
 	TMap<FName, UWidget*> WidgetNameToWidgetPointerMap;
 	FWidgetBlueprintCompilerContext& WidgetBlueprintCompilerContext;
-	TStrongObjectPtr<UMVVMBlueprintView> BlueprintView = nullptr;
 	FCompiledBindingLibraryCompiler BindingLibraryCompiler;
-	bool bIsCreateVariableStepValid = true;
-	bool bIsCreateFunctionsStepValid = true;
-	bool bIsPreCompileStepValid = true;
-	bool bIsCompileStepValid = true;
-
-private:
-	enum class EMessageType
-	{
-		Info = 0, Warning = 1, Error = 2
-	};
-	void AddMessage(const FText& MessageText, EMessageType MessageType) const;
-	void AddMessages(TArrayView<TWeakPtr<FCompilerBinding>> Bindings, TArrayView<TWeakPtr<FCompilerEvent>> Events, const FText& MessageText, EMessageType MessageType) const;
-	void AddMessageForBinding(const TSharedPtr<FCompilerBinding>& Binding, const FText& MessageText, EMessageType MessageType, FName ArgumentName) const;
-	void AddMessageForBinding(const FMVVMBlueprintViewBinding& Binding, const FText& MessageText, EMessageType MessageType, FName ArgumentName) const;
-	void AddMessageForEvent(const TSharedPtr<FCompilerEvent>& Event, const FText& MessageText, EMessageType MessageType, FName ArgumentName) const;
-	void AddMessageForEvent(const UMVVMBlueprintViewEvent* Event, const FText& MessageText, EMessageType MessageType, FName ArgumentName) const;
-	void AddMessageForViewModel(const FMVVMBlueprintViewModelContext& ViewModel, const FText& Message, EMessageType MessageType) const;
-	void AddMessageForViewModel(const FText& ViewModelDisplayName, const FText& Message, EMessageType MessageType) const;
-
-	struct FGetFieldsResult
-	{
-		TSharedPtr<FCompilerBindingSource> Source;
-		TArray<FMVVMConstFieldVariant> GeneratedFields;
-	};
-	struct FCreateFieldsResult
-	{
-		TSharedPtr<FCompilerBindingSource> Source;
-		TArray<UE::MVVM::FMVVMConstFieldVariant> GeneratedFields;
-		TArray<UE::MVVM::FMVVMConstFieldVariant> SkeletalGeneratedFields;
-	};
-	TValueOrError<FGetFieldsResult, FText> GetFields(const UWidgetBlueprintGeneratedClass* Class, const FMVVMBlueprintPropertyPath& PropertyPath) const;
-	TValueOrError<FCreateFieldsResult, FText> CreateFieldContext(const UWidgetBlueprintGeneratedClass* Class, const FMVVMBlueprintPropertyPath& PropertyPath, bool bForSourceReading) const;
-	TValueOrError<FCompilerNotifyFieldId, FText> CreateNotifyFieldId(const UWidgetBlueprintGeneratedClass* Class, const TSharedPtr<FGeneratedReadFieldPathContext>& ReadFieldContext, const FMVVMBlueprintViewBinding& Binding);
-
-
-	static TArray<FMVVMConstFieldVariant> GetFields(const UClass* Class, FName PropertyName, TArray<FMVVMConstFieldVariant> Properties);
-	static bool IsPropertyPathValid(const UBlueprint* Context, TArrayView<const FMVVMConstFieldVariant> PropertyPath);
-	static bool CanBeSetInNative(TArrayView<const FMVVMConstFieldVariant> PropertyPath);
+	bool bAreSourcesCreatorValid = true;
+	bool bAreSourceContextsValid = true;
+	bool bAreBindingsValid = true;
+	bool bAreEventsValid = true;
 };
 
 } //namespace
diff --git a/Engine/Plugins/Runtime/ModelViewViewModel/Source/ModelViewViewModelBlueprint/Private/MVVMWidgetBlueprintExtension_View.cpp b/Engine/Plugins/Runtime/ModelViewViewModel/Source/ModelViewViewModelBlueprint/Private/MVVMWidgetBlueprintExtension_View.cpp
--- a/Engine/Plugins/Runtime/ModelViewViewModel/Source/ModelViewViewModelBlueprint/Private/MVVMWidgetBlueprintExtension_View.cpp
+++ b/Engine/Plugins/Runtime/ModelViewViewModel/Source/ModelViewViewModelBlueprint/Private/MVVMWidgetBlueprintExtension_View.cpp
@@ -104,7 +104,7 @@ void UMVVMWidgetBlueprintExtension_View::HandleBeginCompilation(FWidgetBlueprint
 	if (BlueprintView)
 	{
 		BlueprintView->ResetBindingMessages();
-		CurrentCompilerContext = MakePimpl<UE::MVVM::Private::FMVVMViewBlueprintCompiler>(InCreationContext, GetBlueprintView());
+		CurrentCompilerContext = MakePimpl<UE::MVVM::Private::FMVVMViewBlueprintCompiler>(InCreationContext);
 	}
 }
 
@@ -130,17 +130,17 @@ void UMVVMWidgetBlueprintExtension_View::HandleCreateClassVariablesFromBlueprint
 {
 	Super::HandleCreateClassVariablesFromBlueprint(Context);
 
-	CurrentCompilerContext->CreateVariables(Context);
+	CurrentCompilerContext->CreateVariables(Context, GetBlueprintView());
 }
 
 
-void UMVVMWidgetBlueprintExtension_View::HandleCreateFunctionList(const FWidgetBlueprintCompilerContext::FCreateFunctionContext& Context)
+void UMVVMWidgetBlueprintExtension_View::HandleCreateFunctionList()
 {
-	Super::HandleCreateFunctionList(Context);
+	Super::HandleCreateFunctionList();
 
 	if (CurrentCompilerContext)
 	{
-		CurrentCompilerContext->CreateFunctions(Context);
+		CurrentCompilerContext->CreateFunctions(BlueprintView);
 	}
 }
 
@@ -155,7 +155,7 @@ void UMVVMWidgetBlueprintExtension_View::HandleFinishCompilingClass(UWidgetBluep
 	{
 		UMVVMViewClass* ViewExtension = nullptr;
 		bool bCompiled = false;
-		if (CurrentCompilerContext->PreCompile(Class))
+		if (CurrentCompilerContext->PreCompile(Class, BlueprintView))
 		{
 			FName ClassName = "ViewClass";
 			if (UObject* PreviousObj = StaticFindObjectFastInternal(nullptr, Class, ClassName, true))
@@ -166,7 +166,7 @@ void UMVVMWidgetBlueprintExtension_View::HandleFinishCompilingClass(UWidgetBluep
 				PreviousObj->Rename(*TrashName.ToString(), GetTransientPackage(), RenameFlags);
 			}
 			ViewExtension = NewObject<UMVVMViewClass>(Class);
-			bCompiled = CurrentCompilerContext->Compile(Class, ViewExtension);
+			bCompiled = CurrentCompilerContext->Compile(Class, BlueprintView, ViewExtension);
 		}
 
 		if (bCompiled)
diff --git a/Engine/Plugins/Runtime/ModelViewViewModel/Source/ModelViewViewModelBlueprint/Public/Bindings/MVVMCompiledBindingLibraryCompiler.h b/Engine/Plugins/Runtime/ModelViewViewModel/Source/ModelViewViewModelBlueprint/Public/Bindings/MVVMCompiledBindingLibraryCompiler.h
--- a/Engine/Plugins/Runtime/ModelViewViewModel/Source/ModelViewViewModelBlueprint/Public/Bindings/MVVMCompiledBindingLibraryCompiler.h
+++ b/Engine/Plugins/Runtime/ModelViewViewModel/Source/ModelViewViewModelBlueprint/Public/Bindings/MVVMCompiledBindingLibraryCompiler.h
@@ -158,7 +158,7 @@ public:
 	TValueOrError<FFieldPathHandle, FText> AddFieldPath(TArrayView<const UE::MVVM::FMVVMConstFieldVariant> FieldPath, bool bRead);
 
 	/** */
-	TValueOrError<FFieldPathHandle, FText> AddObjectFieldPath(TArrayView<const UE::MVVM::FMVVMConstFieldVariant> FieldPath, const UClass* ExpectedType, bool bRead);
+	TValueOrError<FFieldPathHandle, FText> AddObjectFieldPath(TArrayView<const UE::MVVM::FMVVMConstFieldVariant> FieldPath, UClass* ExpectedType, bool bRead);
 
 	/** */
 	TValueOrError<FFieldPathHandle, FText> AddConversionFunctionFieldPath(const UClass* SourceClass, const UFunction* Function);
diff --git a/Engine/Plugins/Runtime/ModelViewViewModel/Source/ModelViewViewModelBlueprint/Public/MVVMBlueprintPin.h b/Engine/Plugins/Runtime/ModelViewViewModel/Source/ModelViewViewModelBlueprint/Public/MVVMBlueprintPin.h
--- a/Engine/Plugins/Runtime/ModelViewViewModel/Source/ModelViewViewModelBlueprint/Public/MVVMBlueprintPin.h
+++ b/Engine/Plugins/Runtime/ModelViewViewModel/Source/ModelViewViewModelBlueprint/Public/MVVMBlueprintPin.h
@@ -7,19 +7,8 @@
 #include "MVVMBlueprintPin.generated.h"
 
 class UWidgetBlueprint;
-class UEdGraphNode;
 class UEdGraphPin;
 
-/**
-*
-*/
-UENUM()
-enum class EMVVMBlueprintPinStatus : uint8
-{
-	Valid,
-	Orphaned,
-};
-
 /**
 *
 */
@@ -29,9 +18,6 @@ struct MODELVIEWVIEWMODELBLUEPRINT_API FMVVMBlueprintPin
 	GENERATED_BODY()
 
 private:
-	UPROPERTY(VisibleAnywhere, Category = "Viewmodel")
-	FGuid PinId;
-
 	UPROPERTY(VisibleAnywhere, Category = "Viewmodel")
 	FName PinName;
 
@@ -50,14 +36,6 @@ private:
 	UPROPERTY(VisibleAnywhere, Category = "Viewmodel")
 	TObjectPtr<class UObject> DefaultObject;
 
-	/** The pin is split. */
-	UPROPERTY(VisibleAnywhere, Category = "Viewmodel")
-	bool bSplit = false;
-
-	/** The pin could not be set. */
-	UPROPERTY(VisibleAnywhere, Category = "Viewmodel")
-	mutable EMVVMBlueprintPinStatus Status = EMVVMBlueprintPinStatus::Valid;
-
 public:
 	FMVVMBlueprintPin() = default;
 	FMVVMBlueprintPin(FName PinName);
@@ -67,22 +45,10 @@ public:
 		return PinName;
 	}
 
-	/** The pin is split into its different components. */
-	bool IsSplit() const
-	{
-		return bSplit;
-	}
-
-	/** The pin could not be assigned to the graph pin. */
-	EMVVMBlueprintPinStatus GetStatus() const
-	{
-		return Status;
-	}
-
 	/** Are we using the path. */
 	bool UsedPathAsValue() const
 	{
-		return !bSplit && Path.IsValid();
+		return Path.IsValid();
 	}
 
 	const FMVVMBlueprintPropertyPath& GetPath() const
@@ -97,13 +63,8 @@ public:
 	void SetDefaultValue(const FString& Value);
 	void SetPath(const FMVVMBlueprintPropertyPath& Value);
 
-	static bool IsInputPin(const UEdGraphPin* Pin);
-	static TArray<FMVVMBlueprintPin> CopyAndReturnMissingPins(UBlueprint* Blueprint, UEdGraphNode* GraphNode, const TArray<FMVVMBlueprintPin>& Pins);
-	static TArray<FMVVMBlueprintPin> CreateFromNode(UBlueprint* Blueprint, UEdGraphNode* GraphNode);
-	static FMVVMBlueprintPin CreateFromPin(const UBlueprint* Blueprint, const UEdGraphPin* Pin);
-
-	void CopyTo(const UBlueprint* WidgetBlueprint, UEdGraphNode* Node) const;
-	UEdGraphPin* FindGraphPin(UEdGraphNode* Node) const;
+	static FMVVMBlueprintPin CreateFromPin(const UBlueprint* WidgetBlueprint, const UEdGraphPin* Pin);
+	void CopyTo(const UBlueprint* WidgetBlueprint, UEdGraphPin* Pin) const;
 
 private:
 	void Reset();
diff --git a/Engine/Plugins/Runtime/ModelViewViewModel/Source/ModelViewViewModelBlueprint/Public/MVVMBlueprintViewBinding.h b/Engine/Plugins/Runtime/ModelViewViewModel/Source/ModelViewViewModelBlueprint/Public/MVVMBlueprintViewBinding.h
--- a/Engine/Plugins/Runtime/ModelViewViewModel/Source/ModelViewViewModelBlueprint/Public/MVVMBlueprintViewBinding.h
+++ b/Engine/Plugins/Runtime/ModelViewViewModel/Source/ModelViewViewModelBlueprint/Public/MVVMBlueprintViewBinding.h
@@ -74,7 +74,7 @@ public:
 
 	void GenerateWrapper(UBlueprint* Blueprint);
 	void SavePinValues(UBlueprint* Blueprint);
-	void DeprecateViewConversionFunction(UBlueprint* Blueprint, FMVVMBlueprintViewBinding& Owner);
+	void DeprecateViewConversionFunction(UBlueprint* Blueprint);
 
 public:
 PRAGMA_DISABLE_DEPRECATION_WARNINGS
diff --git a/Engine/Plugins/Runtime/ModelViewViewModel/Source/ModelViewViewModelBlueprint/Public/MVVMBlueprintViewConversionFunction.h b/Engine/Plugins/Runtime/ModelViewViewModel/Source/ModelViewViewModelBlueprint/Public/MVVMBlueprintViewConversionFunction.h
--- a/Engine/Plugins/Runtime/ModelViewViewModel/Source/ModelViewViewModelBlueprint/Public/MVVMBlueprintViewConversionFunction.h
+++ b/Engine/Plugins/Runtime/ModelViewViewModel/Source/ModelViewViewModelBlueprint/Public/MVVMBlueprintViewConversionFunction.h
@@ -5,7 +5,6 @@
 #include "MVVMBlueprintPin.h"
 #include "MVVMBlueprintView.h"
 #include "Engine/MemberReference.h"
-#include "UObject/WeakObjectPtrTemplates.h"
 
 #include "MVVMBlueprintViewConversionFunction.generated.h"
 
@@ -13,8 +12,6 @@ class UK2Node;
 class UEdGraphPin;
 class UEdGraph;
 
-struct FEdGraphEditAction;
-
 /**
  *
  */
@@ -23,43 +20,39 @@ class MODELVIEWVIEWMODELBLUEPRINT_API UMVVMBlueprintViewConversionFunction : pub
 {
 	GENERATED_BODY()
 
-public:
-	static bool IsValidConversionFunction(const UBlueprint* WidgetBlueprint, const UFunction* Function);
-
 public:
 	/** @return the conversion function uses at runtime. The wrapper function if complex or GetFunction is simple. */
 	const UFunction* GetCompiledFunction(const UClass* SelfContext) const;
 
 	/** @return the conversion function uses at runtime. The wrapper function if complex or GetFunction is simple. */
-	FName GetCompiledFunctionName(const UClass* SelfContext) const;
+	FName GetCompiledFunctionName() const;
 
 	/** @return the conversion function. */
-	TVariant<const UFunction*, TSubclassOf<UK2Node>> GetConversionFunction(const UBlueprint* SelfContext) const;
+	TVariant<const UFunction*, TSubclassOf<UK2Node>> GetConversionFunction(const UClass* SelfContext) const;
 
-	/** Set the function. Generate a Graph. */
-	void InitializeFromFunction(UBlueprint* SelfContext, FName GraphName, const UFunction* Function);
+	/** Set the function. Generate a Graph is the conversion function is complex. */
+	void InitFromFunction(UBlueprint* SelfContext, const UFunction* Function);
 
-	// For deprecation
-	void Deprecation_InitializeFromWrapperGraph(UBlueprint* SelfContext, UEdGraph* Graph);
+	/** Set the function. Generate a Graph is the conversion function is complex. */
+	void InitFromFunction(UBlueprint* SelfContext, const UFunction* Function, FName GraphName);
 
 	// For deprecation
-	void Deprecation_InitializeFromMemberReference(UBlueprint* SelfContext, FName GraphName, FMemberReference MemberReference, const FMVVMBlueprintPropertyPath& Source);
+	void InitializeFromWrapperGraph(UBlueprint* SelfContext, UEdGraph* Graph);
 
 	// For deprecation
-	void Deprecation_SetWrapperGraphName(UBlueprint* Context, FName GraphName, const FMVVMBlueprintPropertyPath& Source);
-
-	/**
-	 * The conversion is valid.
-	 * The function was valid when created but may not be anymore.
-	 * It doesn't check if the source and destination are valid.
-	 */
-	bool IsValid(const UBlueprint* SelfContext) const;
+	void InitializeFromMemberReference(UBlueprint* SelfContext, FMemberReference MemberReference);
 
 	/** The function has more than one argument and requires a wrapper or it uses a FunctionNode. */
-	bool NeedsWrapperGraph(const UBlueprint* SelfContext) const;
+	bool NeedsWrapperGraph() const
+	{
+		return !GraphName.IsNone();
+	}
 
-	/** The wrapper Graph is generated on load/compile and is not saved. */
-	bool IsWrapperGraphTransient() const;
+	/** The wrapper Graph is generated on domains and is not saved. */
+	bool IsWrapperGraphTransient() const
+	{
+		return NeedsWrapperGraph()  && bWrapperGraphTransient;
+	}
 
 	/** Return the wrapper graph, if it exists. */
 	UEdGraph* GetWrapperGraph() const
@@ -81,10 +74,10 @@ public:
 	 * If needed, create the graph and all the nodes for that graph when compiling.
 	 * Returns the existing one, if one was created from GetOrCreateWrapperGraph.
 	 */
-	UEdGraph* GetOrCreateIntermediateWrapperGraph(FKismetCompilerContext& Context);
+	UEdGraph* GetOrCreateIntermediateWrapperGraph(FKismetCompilerContext& Context) const;
 
 	/** If needed, create the graph and all the nodes for that graph. */
-	UEdGraph* GetOrCreateWrapperGraph(UBlueprint* Blueprint);
+	UEdGraph* GetOrCreateWrapperGraph(UBlueprint* Blueprint) const;
 
 	/**
 	 * The conversion function is going to be removed from the Blueprint.
@@ -96,7 +89,7 @@ public:
 	 * Returns the pin from the graph.
 	 * Create the graph and all the nodes for that graph if the graph doesn't exist and it's needed.
 	 */
-	UEdGraphPin* GetOrCreateGraphPin(UBlueprint* Blueprint, FName PinName);
+	UEdGraphPin* GetOrCreateGraphPin(UBlueprint* Blueprint, FName PinName) const;
 
 	TArrayView<const FMVVMBlueprintPin> GetPins() const
 	{
@@ -109,17 +102,10 @@ public:
 	/** Generates SavedPins from the wrapper graph, if it exists. */
 	void SavePinValues(UBlueprint* Blueprint);
 
-	FSimpleMulticastDelegate OnWrapperGraphModified;
-
 private:
-	void HandleGraphChanged(const FEdGraphEditAction& Action, TWeakObjectPtr<UBlueprint> Context);
-	void HandleUserDefinedPinRenamed(UK2Node* InNode, FName OldPinName, FName NewPinName, TWeakObjectPtr<UBlueprint> WeakBlueprint);
-	void SetCachedWrapperGraph(UBlueprint* Blueprint, UEdGraph* CachedGraph, UK2Node* CachedNode);
-	UEdGraph* GetOrCreateWrapperGraphInternal(FKismetCompilerContext& Context, const UFunction* Function);
-	UEdGraph* GetOrCreateWrapperGraphInternal(UBlueprint* Blueprint, const UFunction* Function);
-	bool NeedsWrapperGraphInternal(const UClass* SkeletalSelfContext) const;
-	void LoadPinValuesInternal(UBlueprint* Blueprint);
-	void CreateWrapperGraphName();
+	UEdGraph* GetOrCreateWrapperGraphInternal(FKismetCompilerContext& Context, const UFunction* Function) const;
+	UEdGraph* GetOrCreateWrapperGraphInternal(UBlueprint* Blueprint, const UFunction* Function) const;
+	void LoadPinValuesInternal(UBlueprint* Blueprint) const;
 	void Reset();
 
 private:
@@ -157,7 +143,4 @@ private:
 	
 	UPROPERTY(Transient, DuplicateTransient)
 	mutable TObjectPtr<UK2Node> CachedWrapperNode;
-
-	FDelegateHandle OnGraphChangedHandle;
-	FDelegateHandle OnUserDefinedPinRenamedHandle;
 };
diff --git a/Engine/Plugins/Runtime/ModelViewViewModel/Source/ModelViewViewModelBlueprint/Public/MVVMBlueprintViewEvent.h b/Engine/Plugins/Runtime/ModelViewViewModel/Source/ModelViewViewModelBlueprint/Public/MVVMBlueprintViewEvent.h
--- a/Engine/Plugins/Runtime/ModelViewViewModel/Source/ModelViewViewModelBlueprint/Public/MVVMBlueprintViewEvent.h
+++ b/Engine/Plugins/Runtime/ModelViewViewModel/Source/ModelViewViewModelBlueprint/Public/MVVMBlueprintViewEvent.h
@@ -7,7 +7,6 @@
 
 #include "MVVMBlueprintViewEvent.generated.h"
 
-struct FEdGraphEditAction;
 class UEdGraph;
 class UK2Node;
 class UWidgetBlueprint;
@@ -95,12 +94,10 @@ public:
 	// To set a pin when loading the asset (no graph generation)
 	void SetPinPathNoGraphGeneration(FName PinName, const FMVVMBlueprintPropertyPath& Path);
 
-	FSimpleMulticastDelegate OnWrapperGraphModified;
-
 public:
 	TArray<FText> GetCompilationMessages(EMessageType InMessageType) const;
 	bool HasCompilationMessage(EMessageType InMessageType) const;
-	void AddCompilationToBinding(FMessage MessageToAdd) const;
+	void AddCompilationToBinding(FMessage MessageToAdd);
 	void ResetCompilationMessages();
 
 public:
@@ -122,10 +119,7 @@ public:
 private:
 	static const UFunction* GetEventSignature(const UWidgetBlueprint* WidgetBlueprint, const FMVVMBlueprintPropertyPath& PropertyPath);
 	const UFunction* GetEventSignature() const;
-	void HandleGraphChanged(const FEdGraphEditAction& Action);
-	void HandleUserDefinedPinRenamed(UK2Node* InNode, FName OldPinName, FName NewPinName);
 	UWidgetBlueprint* GetWidgetBlueprintInternal() const;
-	void SetCachedWrapperGraphInternal(UEdGraph* Graph, UK2Node* Node);
 	UEdGraph* CreateWrapperGraphInternal();
 	UEdGraph* CreateWrapperGraphInternal(const UFunction* Signature, const UFunction* Function);
 	UEdGraph* CreateWrapperGraphInternal(const UFunction* Signature, const FProperty* Property);
@@ -148,14 +142,11 @@ private:
 	UPROPERTY(VisibleAnywhere, Category = "Viewmodel")
 	FName GraphName;
 
-	mutable TArray<FMessage> Messages;
+	TArray<FMessage> Messages;
 
 	UPROPERTY(Transient, DuplicateTransient)
 	mutable TObjectPtr<UEdGraph> CachedWrapperGraph;
 
 	UPROPERTY(Transient, DuplicateTransient)
 	mutable TObjectPtr<UK2Node> CachedWrapperNode;
-
-	FDelegateHandle OnGraphChangedHandle;
-	FDelegateHandle OnUserDefinedPinRenamedHandle;
 };
diff --git a/Engine/Plugins/Runtime/ModelViewViewModel/Source/ModelViewViewModelBlueprint/Public/MVVMWidgetBlueprintExtension_View.h b/Engine/Plugins/Runtime/ModelViewViewModel/Source/ModelViewViewModelBlueprint/Public/MVVMWidgetBlueprintExtension_View.h
--- a/Engine/Plugins/Runtime/ModelViewViewModel/Source/ModelViewViewModelBlueprint/Public/MVVMWidgetBlueprintExtension_View.h
+++ b/Engine/Plugins/Runtime/ModelViewViewModel/Source/ModelViewViewModelBlueprint/Public/MVVMWidgetBlueprintExtension_View.h
@@ -56,7 +56,7 @@ public:
 	virtual void HandleBeginCompilation(FWidgetBlueprintCompilerContext& InCreationContext) override;
 	virtual void HandleCleanAndSanitizeClass(UWidgetBlueprintGeneratedClass* ClassToClean, UObject* OldCDO) override;
 	virtual void HandleCreateClassVariablesFromBlueprint(const FWidgetBlueprintCompilerContext::FCreateVariableContext& Context) override;
-	virtual void HandleCreateFunctionList(const FWidgetBlueprintCompilerContext::FCreateFunctionContext& Context) override;
+	virtual void HandleCreateFunctionList() override;
 	virtual void HandleFinishCompilingClass(UWidgetBlueprintGeneratedClass* Class) override;
 	virtual void HandleEndCompilation() override;
 	virtual FSearchData HandleGatherSearchData(const UBlueprint* OwningBlueprint) const override;
diff --git a/Engine/Plugins/Runtime/ModelViewViewModel/Source/ModelViewViewModelEditor/Private/Widgets/SMVVMViewModelPanel.cpp b/Engine/Plugins/Runtime/ModelViewViewModel/Source/ModelViewViewModelEditor/Private/Widgets/SMVVMViewModelPanel.cpp
--- a/Engine/Plugins/Runtime/ModelViewViewModel/Source/ModelViewViewModelEditor/Private/Widgets/SMVVMViewModelPanel.cpp
+++ b/Engine/Plugins/Runtime/ModelViewViewModel/Source/ModelViewViewModelEditor/Private/Widgets/SMVVMViewModelPanel.cpp
@@ -287,7 +287,7 @@ void SMVVMViewModelPanel::Construct(const FArguments& InArgs, TSharedPtr<FWidget
 				SNew(SWarningOrErrorBox)
 				.Visibility(this, &SMVVMViewModelPanel::GetWarningPanelVisibility)
 				.MessageStyle(EMessageStyle::Warning)
-				.Message(this, &SMVVMViewModelPanel::GetWarningMessage)
+				.Message(LOCTEXT("InitializationPanelWarningDescription", "The view will not initialize automatically. It was manually set in the View Settings."))
 				[
 					SNew(SButton)
 					.OnClicked(this, &SMVVMViewModelPanel::HandleDisableWarningPanel)
@@ -886,10 +886,6 @@ EVisibility SMVVMViewModelPanel::GetWarningPanelVisibility() const
 	{
 		if (UMVVMBlueprintView* WidgetBlueprint = WeakBlueprintView.Get())
 		{
-			if (WidgetBlueprint->GetNumBindings() == 0 && WidgetBlueprint->GetViewModels().Num() != 0)
-			{
-				return EVisibility::Visible;
-			}
 			if (!WidgetBlueprint->GetSettings()->bInitializeSourcesOnConstruct || !WidgetBlueprint->GetSettings()->bInitializeBindingsOnConstruct)
 			{
 				return EVisibility::Visible;
@@ -901,26 +897,6 @@ EVisibility SMVVMViewModelPanel::GetWarningPanelVisibility() const
 }
 
 
-FText SMVVMViewModelPanel::GetWarningMessage() const
-{
-	if (!bDisableWarningPanel)
-	{
-		if (UMVVMBlueprintView* WidgetBlueprint = WeakBlueprintView.Get())
-		{
-			if (WidgetBlueprint->GetNumBindings() == 0 && WidgetBlueprint->GetViewModels().Num() != 0)
-			{
-				return LOCTEXT("NoBindingWarningDescription", "No view will be created for this widget because there are no bindings. Your viewmodels won't be initialized.");
-			}
-			if (!WidgetBlueprint->GetSettings()->bInitializeSourcesOnConstruct || !WidgetBlueprint->GetSettings()->bInitializeBindingsOnConstruct)
-			{
-				return LOCTEXT("InitializationPanelWarningDescription", "The view will not initialize automatically. It was manually set in the View Settings.");
-			}
-		}
-	}
-	return FText::GetEmpty();
-}
-
-
 FReply SMVVMViewModelPanel::HandleDisableWarningPanel()
 {
 	bDisableWarningPanel = true;
diff --git a/Engine/Plugins/Runtime/ModelViewViewModel/Source/ModelViewViewModelEditor/Private/Widgets/SMVVMViewModelPanel.h b/Engine/Plugins/Runtime/ModelViewViewModel/Source/ModelViewViewModelEditor/Private/Widgets/SMVVMViewModelPanel.h
--- a/Engine/Plugins/Runtime/ModelViewViewModel/Source/ModelViewViewModelEditor/Private/Widgets/SMVVMViewModelPanel.h
+++ b/Engine/Plugins/Runtime/ModelViewViewModel/Source/ModelViewViewModelEditor/Private/Widgets/SMVVMViewModelPanel.h
@@ -98,7 +98,6 @@ private:
 	TSharedRef<SWidget> HandleAddViewModelContextMenu();
 
 	EVisibility GetWarningPanelVisibility() const;
-	FText GetWarningMessage() const;
 	FReply HandleDisableWarningPanel();
 
 	void CreateCommandList();
diff --git a/Engine/Source/Editor/UMGEditor/Private/WidgetBlueprintCompiler.cpp b/Engine/Source/Editor/UMGEditor/Private/WidgetBlueprintCompiler.cpp
--- a/Engine/Source/Editor/UMGEditor/Private/WidgetBlueprintCompiler.cpp
+++ b/Engine/Source/Editor/UMGEditor/Private/WidgetBlueprintCompiler.cpp
@@ -71,23 +71,6 @@ EKismetCompileType::Type FWidgetBlueprintCompilerContext::FCreateVariableContext
 }
 
 
-//////////////////////////////////////////////////////////////////////////
-// FWidgetBlueprintCompiler::FCreateFunctionContext
-FWidgetBlueprintCompilerContext::FCreateFunctionContext::FCreateFunctionContext(FWidgetBlueprintCompilerContext& InContext)
-	: Context(InContext)
-{}
-
-void FWidgetBlueprintCompilerContext::FCreateFunctionContext::AddGeneratedFunctionGraph(UEdGraph* Graph) const
-{
-	Context.GeneratedFunctionGraphs.Add(Graph);
-}
-
-UWidgetBlueprintGeneratedClass* FWidgetBlueprintCompilerContext::FCreateFunctionContext::GetGeneratedClass() const
-{
-	return Context.NewWidgetBlueprintClass;
-}
-
-
 //////////////////////////////////////////////////////////////////////////
 // FWidgetBlueprintCompiler
 FWidgetBlueprintCompiler::FWidgetBlueprintCompiler()
@@ -181,9 +164,9 @@ UEdGraphSchema_K2* FWidgetBlueprintCompilerContext::CreateSchema()
 
 void FWidgetBlueprintCompilerContext::CreateFunctionList()
 {
-	UWidgetBlueprintExtension::ForEachExtension(WidgetBlueprint(), [Self = this](UWidgetBlueprintExtension* InExtension)
+	UWidgetBlueprintExtension::ForEachExtension(WidgetBlueprint(), [this](UWidgetBlueprintExtension* InExtension)
 		{
-			InExtension->CreateFunctionList(FCreateFunctionContext(*Self));
+			InExtension->CreateFunctionList();
 		});
 
 	Super::CreateFunctionList();
diff --git a/Engine/Source/Editor/UMGEditor/Public/WidgetBlueprintCompiler.h b/Engine/Source/Editor/UMGEditor/Public/WidgetBlueprintCompiler.h
--- a/Engine/Source/Editor/UMGEditor/Public/WidgetBlueprintCompiler.h
+++ b/Engine/Source/Editor/UMGEditor/Public/WidgetBlueprintCompiler.h
@@ -113,18 +113,6 @@ public:
 		FWidgetBlueprintCompilerContext& Context;
 	};
 
-	struct UMGEDITOR_API FCreateFunctionContext
-	{
-	public:
-		void AddGeneratedFunctionGraph(UEdGraph*) const;
-		UWidgetBlueprintGeneratedClass* GetGeneratedClass() const;
-
-	private:
-		friend FWidgetBlueprintCompilerContext;
-		FCreateFunctionContext(FWidgetBlueprintCompilerContext& InContext);
-		FWidgetBlueprintCompilerContext& Context;
-	};
-
 protected:
 	void FixAbandonedWidgetTree(UWidgetBlueprint* WidgetBP);
 
diff --git a/Engine/Source/Editor/UMGEditor/Public/WidgetBlueprintExtension.h b/Engine/Source/Editor/UMGEditor/Public/WidgetBlueprintExtension.h
--- a/Engine/Source/Editor/UMGEditor/Public/WidgetBlueprintExtension.h
+++ b/Engine/Source/Editor/UMGEditor/Public/WidgetBlueprintExtension.h
@@ -72,7 +72,7 @@ protected:
 	 */
 	virtual void HandleBeginCompilation(FWidgetBlueprintCompilerContext& InCreationContext) {}
 
-	virtual void HandleCreateFunctionList(const FWidgetBlueprintCompilerContext::FCreateFunctionContext& InCreationContext) {}
+	virtual void HandleCreateFunctionList() {}
 	virtual void HandleCleanAndSanitizeClass(UWidgetBlueprintGeneratedClass* ClassToClean, UObject* OldCDO) {}
 	virtual TArray<UObject*> HandleSaveSubObjectsFromCleanAndSanitizeClass(UWidgetBlueprintGeneratedClass* ClassToClean) { return TArray<UObject*>();  }
 	virtual void HandleCreateClassVariablesFromBlueprint(const FWidgetBlueprintCompilerContext::FCreateVariableContext& Context) {}
@@ -95,9 +95,9 @@ private:
 		HandleBeginCompilation(InCreationContext);
 	}
 
-	void CreateFunctionList(const FWidgetBlueprintCompilerContext::FCreateFunctionContext& InCreationContext)
+	void CreateFunctionList()
 	{
-		HandleCreateFunctionList(InCreationContext);
+		HandleCreateFunctionList();
 	}
 
 	void CleanAndSanitizeClass(UWidgetBlueprintGeneratedClass* ClassToClean, UObject* OldCDO)
