MVVM: Backout change that make widget private when not enabled via the "is variable". Assets need to be fixed before reenabling.

diff --git a/Engine/Plugins/Runtime/ModelViewViewModel/Source/ModelViewViewModelBlueprint/Private/MVVMViewBlueprintCompiler.cpp b/Engine/Plugins/Runtime/ModelViewViewModel/Source/ModelViewViewModelBlueprint/Private/MVVMViewBlueprintCompiler.cpp
--- a/Engine/Plugins/Runtime/ModelViewViewModel/Source/ModelViewViewModelBlueprint/Private/MVVMViewBlueprintCompiler.cpp
+++ b/Engine/Plugins/Runtime/ModelViewViewModel/Source/ModelViewViewModelBlueprint/Private/MVVMViewBlueprintCompiler.cpp
@@ -768,7 +768,8 @@ void FMVVMViewBlueprintCompiler::CreateRequiredProperties(const FWidgetBlueprint
 					SourceVariable.CategoryName = TEXT("Widget");
 					SourceVariable.CreationType = FCompilerUserWidgetProperty::ECreationType::CreateIfDoesntExist;
 					SourceVariable.bExposeOnSpawn = false;
-					SourceVariable.bPrivate = true;
+					//SourceVariable.bPrivate = true; Remove until the data is fix properly
+					SourceVariable.bPrivate = false;
 					SourceVariable.bReadOnly = true;
 
 					WidgetUserPropertyCreated.Add(Widget->GetFName());
