Fix localization waring initroduced with CL 29815759
#issue 399399

diff --git a/Engine/Plugins/Runtime/ModelViewViewModel/Source/ModelViewViewModelBlueprint/Private/MVVMViewBlueprintCompiler.cpp b/Engine/Plugins/Runtime/ModelViewViewModel/Source/ModelViewViewModelBlueprint/Private/MVVMViewBlueprintCompiler.cpp
--- a/Engine/Plugins/Runtime/ModelViewViewModel/Source/ModelViewViewModelBlueprint/Private/MVVMViewBlueprintCompiler.cpp
+++ b/Engine/Plugins/Runtime/ModelViewViewModel/Source/ModelViewViewModelBlueprint/Private/MVVMViewBlueprintCompiler.cpp
@@ -2198,7 +2198,7 @@ void FMVVMViewBlueprintCompiler::PreCompileBindings(UWidgetBlueprintGeneratedCla
 			&& ValidBinding->Type != FCompilerBinding::EType::ComplexConversionFunction
 			&& ValidBinding->Type != FCompilerBinding::EType::SimpleConversionFunction)
 		{
-			AddMessageForBinding(Binding, LOCTEXT("CouldNotCreateFieldId", "The binding is invalid."), EMessageType::Error, FName());
+			AddMessageForBinding(Binding, LOCTEXT("UnsupportedBindingType", "The binding is invalid."), EMessageType::Error, FName());
 			bIsPreCompileStepValid = false;
 			continue;
 		}
