MVVM: fix the issue with deletion of a whole binding group: some bindings from other groups were deleted and some bindings from the target group remained. A reverse loop is needed because we are deleting the elements that BindingsToRemove is pointing to in the original Bindings array.

diff --git a/Engine/Plugins/Runtime/ModelViewViewModel/Source/ModelViewViewModelEditor/Private/Widgets/SMVVMViewBindingListView.cpp b/Engine/Plugins/Runtime/ModelViewViewModel/Source/ModelViewViewModelEditor/Private/Widgets/SMVVMViewBindingListView.cpp
--- a/Engine/Plugins/Runtime/ModelViewViewModel/Source/ModelViewViewModelEditor/Private/Widgets/SMVVMViewBindingListView.cpp
+++ b/Engine/Plugins/Runtime/ModelViewViewModel/Source/ModelViewViewModelEditor/Private/Widgets/SMVVMViewBindingListView.cpp
@@ -2359,8 +2359,9 @@ void SBindingsList::HandleDeleteSelected()
 
 				BlueprintView->Modify();
 
-				for (const FMVVMBlueprintViewBinding* Binding : BindingsToRemove)
+				for (int32 BindingIndex = BindingsToRemove.Num() - 1; BindingIndex > -1; BindingIndex--)
 				{
+					const FMVVMBlueprintViewBinding* Binding = BindingsToRemove[BindingIndex];
 					BlueprintView->RemoveBinding(Binding);
 				}
 				for (UMVVMBlueprintViewEvent* Event : EventsToRemove)
