MVVM: Add function to resolve the PropertyPath.

diff --git a/Engine/Plugins/Runtime/ModelViewViewModel/Source/ModelViewViewModelBlueprint/Private/MVVMPropertyPath.cpp b/Engine/Plugins/Runtime/ModelViewViewModel/Source/ModelViewViewModelBlueprint/Private/MVVMPropertyPath.cpp
--- a/Engine/Plugins/Runtime/ModelViewViewModel/Source/ModelViewViewModelBlueprint/Private/MVVMPropertyPath.cpp
+++ b/Engine/Plugins/Runtime/ModelViewViewModel/Source/ModelViewViewModelBlueprint/Private/MVVMPropertyPath.cpp
@@ -259,6 +259,54 @@ TArray<UE::MVVM::FMVVMConstFieldVariant> FMVVMBlueprintPropertyPath::GetFields(c
 }
 
 
+TArray<UE::MVVM::FMVVMConstFieldVariant> FMVVMBlueprintPropertyPath::GetCompleteFields(const UBlueprint* InSelfContext) const
+{
+	TArray<UE::MVVM::FMVVMConstFieldVariant> Result;
+	Result.Reserve(Paths.Num() + 1);
+
+	UClass* ContextClass = InSelfContext->SkeletonGeneratedClass ? InSelfContext->SkeletonGeneratedClass : InSelfContext->GeneratedClass;
+	switch (GetSource(InSelfContext))
+	{
+	case EMVVMBlueprintFieldPathSource::ViewModel:
+	{
+		Result.AddDefaulted();
+		for (const TObjectPtr<UBlueprintExtension>& Extension : InSelfContext->GetExtensions())
+		{
+			if (Extension && Extension->GetClass() == UMVVMWidgetBlueprintExtension_View::StaticClass())
+			{
+				if (UMVVMBlueprintView* View = CastChecked<UMVVMWidgetBlueprintExtension_View>(Extension)->GetBlueprintView())
+				{
+					if (const FMVVMBlueprintViewModelContext* Viewmodel = View->FindViewModel(GetViewModelId()))
+					{
+						Result[0] = UE::MVVM::FMVVMConstFieldVariant(ContextClass->FindPropertyByName(Viewmodel->GetViewModelName()));
+					}
+				}
+			}
+		}
+		break;
+	}
+	case EMVVMBlueprintFieldPathSource::SelfContext:
+		break;
+	case EMVVMBlueprintFieldPathSource::Widget:
+	{
+		FProperty* WidgetProperty = ContextClass->FindPropertyByName(GetWidgetName());
+		Result.Add(UE::MVVM::FMVVMConstFieldVariant(WidgetProperty));
+		break;
+	}
+	default:
+		check(false);
+		break;
+	}
+
+	for (const FMVVMBlueprintFieldPath& Path : Paths)
+	{
+		Result.Add(Path.GetField(ContextClass));
+	}
+
+	return Result;
+}
+
+
 FString FMVVMBlueprintPropertyPath::GetPropertyPath(const UClass* InSelfContext) const
 {
 	TStringBuilder<512> Result;
diff --git a/Engine/Plugins/Runtime/ModelViewViewModel/Source/ModelViewViewModelBlueprint/Public/MVVMPropertyPath.h b/Engine/Plugins/Runtime/ModelViewViewModel/Source/ModelViewViewModelBlueprint/Public/MVVMPropertyPath.h
--- a/Engine/Plugins/Runtime/ModelViewViewModel/Source/ModelViewViewModelBlueprint/Public/MVVMPropertyPath.h
+++ b/Engine/Plugins/Runtime/ModelViewViewModel/Source/ModelViewViewModelBlueprint/Public/MVVMPropertyPath.h
@@ -131,15 +131,27 @@ public:
 		return Paths;
 	}
 
-	/** Get the binding names, resolves reference deprecation / redirectors / etc before returning */
+	/**
+	 * Get the binding names. Resolves reference deprecation / redirectors.
+	 * Returns Field.SubProperty.SubProperty from ViewModel.Field.SubProperty.SubProperty
+	 */
 	MODELVIEWVIEWMODELBLUEPRINT_API TArray<FName> GetFieldNames(const UClass* SelfContext) const;
 
-	/** Get the binding fields, resolves reference deprecation / redirectors / etc before returning */
+	/**
+	 * Get the binding fields. Resolves reference deprecation / redirectors.
+	 * Returns Field.SubProperty.SubProperty from ViewModel.Field.SubProperty.SubProperty
+	 */
 	MODELVIEWVIEWMODELBLUEPRINT_API TArray<UE::MVVM::FMVVMConstFieldVariant> GetFields(const UClass* SelfContext) const;
 
+	/**
+	 * Get the binding fields. Resolves reference deprecation / redirectors.
+	 * Returns Viewmodel.Field.SubProperty.SubProperty from ViewModel.Field.SubProperty.SubProperty
+	 */
+	MODELVIEWVIEWMODELBLUEPRINT_API TArray<UE::MVVM::FMVVMConstFieldVariant> GetCompleteFields(const UBlueprint* SelfContext) const;
+
 	/**
 	 * Get the full path without the first property name.
-	 * returns Field.SubProperty.SubProperty from ViewModel.Field.SubProeprty.SubProperty
+	 * Returns Field.SubProperty.SubProperty from ViewModel.Field.SubProeprty.SubProperty
 	 */
 	MODELVIEWVIEWMODELBLUEPRINT_API FString GetPropertyPath(const UClass* SelfContext) const;
 
@@ -184,9 +196,6 @@ public:
 			const_cast<FMVVMBlueprintPropertyPath*>(this)->DeprecationUpdateSource(InContext);
 		}
 		return Source;
-		
-
-
 	}
 
 	UE_DEPRECATED(5.4, "Use GetSource instead.")
