MVVM: Add error at the end of the line. Future work with setter and conversion function require space in front of the binding.

diff --git a/Engine/Plugins/Runtime/ModelViewViewModel/Source/ModelViewViewModelEditor/Private/Widgets/SMVVMViewBindingListView.cpp b/Engine/Plugins/Runtime/ModelViewViewModel/Source/ModelViewViewModelEditor/Private/Widgets/SMVVMViewBindingListView.cpp
--- a/Engine/Plugins/Runtime/ModelViewViewModel/Source/ModelViewViewModelEditor/Private/Widgets/SMVVMViewBindingListView.cpp
+++ b/Engine/Plugins/Runtime/ModelViewViewModel/Source/ModelViewViewModelEditor/Private/Widgets/SMVVMViewBindingListView.cpp
@@ -593,19 +593,6 @@ public:
 							.OnCheckStateChanged(this, &SBindingRow::OnIsBindingCompileChanged)
 						]
 
-						+ SHorizontalBox::Slot()
-						.Padding(2, 0)
-						.VAlign(VAlign_Center)
-						.HAlign(HAlign_Left) 
-						.AutoWidth()
-						[
-							SNew(SSimpleButton)
-							.Icon(FAppStyle::Get().GetBrush("Icons.Error"))
-							.Visibility(this, &SBindingRow::GetErrorButtonVisibility)
-							.ToolTipText(this, &SBindingRow::GetErrorButtonToolTip)
-							.OnClicked(this, &SBindingRow::OnErrorButtonClicked)
-						]
-
 						+ SHorizontalBox::Slot()
 						.Padding(4.f, 0.f)
 						.VAlign(VAlign_Center)
@@ -704,6 +691,19 @@ public:
 								.ToolTipText(this, &SBindingRow::GetExecutioModeValueToolTip)
 							]
 						]
+
+						+ SHorizontalBox::Slot()
+						.Padding(2, 0)
+						.VAlign(VAlign_Center)
+						.HAlign(HAlign_Left)
+						.AutoWidth()
+						[
+							SNew(SSimpleButton)
+							.Icon(FAppStyle::Get().GetBrush("Icons.Error"))
+							.Visibility(this, &SBindingRow::GetErrorButtonVisibility)
+							.ToolTipText(this, &SBindingRow::GetErrorButtonToolTip)
+							.OnClicked(this, &SBindingRow::OnErrorButtonClicked)
+						]
 					]
 				]
 			],
@@ -785,7 +785,7 @@ private:
 			UMVVMEditorSubsystem* EditorSubsystem = GEditor->GetEditorSubsystem<UMVVMEditorSubsystem>();
 			UMVVMBlueprintView* BlueprintViewPtr = EditorSubsystem->GetView(WidgetBlueprintWeak.Get());
 			bool HasBindingError = BlueprintViewPtr->HasBindingMessage(GetThisViewBinding()->BindingId, EBindingMessageType::Error) || BlueprintViewPtr->HasBindingMessage(GetThisViewBinding()->BindingId, EBindingMessageType::Warning);
-			return HasBindingError ? EVisibility::Visible : EVisibility::Collapsed;
+			return HasBindingError ? EVisibility::Visible : EVisibility::Hidden;
 		}
 		return EVisibility::Collapsed;
 	}
@@ -1369,19 +1369,6 @@ public:
 							.OnCheckStateChanged(this, &SEventRow::OnIsEventCompileChanged)
 						]
 
-						+ SHorizontalBox::Slot()
-						.Padding(2.0f, 0.0f)
-						.VAlign(VAlign_Center)
-						.HAlign(HAlign_Left)
-						.AutoWidth()
-						[
-							SNew(SSimpleButton)
-							.Icon(FAppStyle::Get().GetBrush("Icons.Error"))
-							.Visibility(this, &SEventRow::GetErrorButtonVisibility)
-							.ToolTipText(this, &SEventRow::GetErrorButtonToolTip)
-							.OnClicked(this, &SEventRow::OnErrorButtonClicked)
-						]
-
 						+ SHorizontalBox::Slot()
 						.Padding(4.0f, 0.0f)
 						.VAlign(VAlign_Center)
@@ -1437,6 +1424,24 @@ public:
 								.OnDragEnter(this, &SEventRow::HandleFieldSelectorDragEnter, false)
 							]
 						]
+
+						+ SHorizontalBox::Slot()
+						[
+							SNew(SSpacer)
+						]
+
+						+SHorizontalBox::Slot()
+						.Padding(2.0f, 0.0f)
+						.VAlign(VAlign_Center)
+						.HAlign(HAlign_Left)
+						.AutoWidth()
+						[
+							SNew(SSimpleButton)
+							.Icon(FAppStyle::Get().GetBrush("Icons.Error"))
+							.Visibility(this, &SEventRow::GetErrorButtonVisibility)
+							.ToolTipText(this, &SEventRow::GetErrorButtonToolTip)
+							.OnClicked(this, &SEventRow::OnErrorButtonClicked)
+						]
 					]
 				]
 			],
@@ -1484,7 +1489,7 @@ private:
 		{
 			bool bHasBindingError = Event->HasCompilationMessage(UMVVMBlueprintViewEvent::EMessageType::Error);
 			bool bHasBindingWarning = Event->HasCompilationMessage(UMVVMBlueprintViewEvent::EMessageType::Warning);
-			return bHasBindingError || bHasBindingWarning ? EVisibility::Visible : EVisibility::Collapsed;
+			return bHasBindingError || bHasBindingWarning ? EVisibility::Visible : EVisibility::Hidden;
 		}
 		return EVisibility::Collapsed;
 	}
@@ -2634,8 +2639,8 @@ TSharedPtr<SWidget> SBindingsList::OnSourceConstructContextMenu()
 				FUIAction RemoveAction;
 				RemoveAction.ExecuteAction = FExecuteAction::CreateSP(this, &SBindingsList::HandleDeleteSelected);
 				RemoveAction.CanExecuteAction = FCanExecuteAction::CreateLambda([bCanRemoveEntry]() { return bCanRemoveEntry; });
-				MenuBuilder.AddMenuEntry(LOCTEXT("RemoveBinding", "Remove Binding"),
-					LOCTEXT("RemoveBindingTooltip", "Remove this binding."),
+				MenuBuilder.AddMenuEntry(LOCTEXT("RemoveBinding", "Remove"),
+					LOCTEXT("RemoveBindingTooltip", "Remove bindings or events."),
 					FSlateIcon(FAppStyle::GetAppStyleSetName(), "Icons.Delete"),
 					RemoveAction);
 			}
