MVVM: When filtering the field in the editor, use the previous container to check if the property is allowed.

diff --git a/Engine/Plugins/Runtime/ModelViewViewModel/Source/ModelViewViewModelEditor/Private/Widgets/SMVVMViewModelBindingListWidget.cpp b/Engine/Plugins/Runtime/ModelViewViewModel/Source/ModelViewViewModelEditor/Private/Widgets/SMVVMViewModelBindingListWidget.cpp
--- a/Engine/Plugins/Runtime/ModelViewViewModel/Source/ModelViewViewModelEditor/Private/Widgets/SMVVMViewModelBindingListWidget.cpp
+++ b/Engine/Plugins/Runtime/ModelViewViewModel/Source/ModelViewViewModelEditor/Private/Widgets/SMVVMViewModelBindingListWidget.cpp
@@ -3,6 +3,7 @@
 #include "SMVVMViewModelBindingListWidget.h"
 
 #include "Bindings/MVVMBindingHelper.h"
+#include "Bindings/MVVMFieldPathHelper.h"
 #include "Blueprint/WidgetTree.h"
 #include "BlueprintEditorSettings.h"
 #include "MVVMBlueprintViewModelContext.h"
@@ -287,6 +288,13 @@ TOptional<const UStruct*> FFieldExpander_Bindable::GetExpandedFunction(const UFu
 			return ObjectProperty->PropertyClass.Get();
 		}
 	}
+	//else if (const FStructProperty* StructProperty = CastField<const FStructProperty>(ReturnProperty))
+	//{
+	//	if (CanExpandScriptStruct(StructProperty))
+	//	{
+	//		return StructProperty->Struct.Get();
+	//	}
+	//}
 	return TOptional<const UStruct*>();
 }
 
@@ -480,30 +488,42 @@ FMVVMBlueprintPropertyPath SSourceBindingList::CreateBlueprintPropertyPath(SProp
 		return FMVVMBlueprintPropertyPath();
 	}
 
-	TSubclassOf<UObject> AccessorClass = WidgetBlueprintPtr ? WidgetBlueprintPtr->SkeletonGeneratedClass : nullptr;
+	const UClass* AccessorClass = WidgetBlueprintPtr->SkeletonGeneratedClass ? WidgetBlueprintPtr->SkeletonGeneratedClass : WidgetBlueprintPtr->GeneratedClass;
 	FMVVMBlueprintPropertyPath PropertyPath;
 	if (FieldPath.Num() > 0)
 	{
 		// Backward, test if the object can be access.
 		//The last property can be a struct variable, inside a struct, inside..., inside an object. 
 		bool bPassFilter = false;
-		for (int32 Index = FieldPath.Num() - 1; Index >= 0; --Index)
+		const UStruct* CurrentContainer = AccessorClass;
+
 		{
-			const FFieldVariant& FieldVariant = FieldPath[Index];
-			const UStruct* OwnerStruct = nullptr;
+			FMVVMBlueprintPropertyPath TempPropertyPath;
+			Source->Key.SetSourceTo(TempPropertyPath);
+			TArray<FMVVMConstFieldVariant> AllFields = TempPropertyPath.GetCompleteFields(WidgetBlueprintPtr);
+			if (AllFields.Num() > 0)
+			{
+				TValueOrError<const UStruct*, void> NewContainerResult = FieldPathHelper::GetFieldAsContainer(AllFields[0]);
+				CurrentContainer = NewContainerResult.HasValue() ? NewContainerResult.GetValue() : nullptr;
+			}
+		}
+
+		for (const FFieldVariant& FieldVariant : FieldPath)
+		{
+			FMVVMConstFieldVariant NewField;
 			FName FieldName;
 			if (const FProperty* Property = FieldVariant.Get<FProperty>())
 			{
-				OwnerStruct = Property->GetOwnerStruct();
+				NewField = FMVVMConstFieldVariant(Property);
 				FieldName = Property->GetFName();
 			}
 			else if (const UFunction* Function = FieldVariant.Get<UFunction>())
 			{
-				OwnerStruct = Function->GetOwnerClass();
+				NewField = FMVVMConstFieldVariant(Function);
 				FieldName = Function->GetFName();
 			}
 
-			if (const UClass* OwnerClass = Cast<const UClass>(OwnerStruct))
+			if (const UClass* OwnerClass = Cast<const UClass>(CurrentContainer))
 			{
 				FMVVMAvailableBinding Binding = UMVVMSubsystem::GetAvailableBinding(OwnerClass, FMVVMBindingName(FieldName), AccessorClass);
 				if (Binding.IsValid())
@@ -529,6 +549,9 @@ FMVVMBlueprintPropertyPath SSourceBindingList::CreateBlueprintPropertyPath(SProp
 				}
 				break;
 			}
+
+			TValueOrError<const UStruct*, void> NewContainerResult = FieldPathHelper::GetFieldAsContainer(NewField);
+			CurrentContainer = NewContainerResult.HasValue() ? NewContainerResult.GetValue() : nullptr;
 		}
 
 		if (bPassFilter)
@@ -541,7 +564,7 @@ FMVVMBlueprintPropertyPath SSourceBindingList::CreateBlueprintPropertyPath(SProp
 			}
 		}
 	}
-	else if(AccessorClass.Get())
+	else if(AccessorClass)
 	{
 		FMVVMBindingName BindingName = Source->Key.ToBindingName(WidgetBlueprintPtr);
 		FMVVMAvailableBinding Binding = UMVVMSubsystem::GetAvailableBinding(AccessorClass, BindingName, AccessorClass);
