MVVM: Setter inside a struct is now supported. PropertyPath like Viewmodel.StructA.StructB.Value. It will use the property setter and requires a blueprint graph.

diff --git a/Engine/Plugins/Runtime/ModelViewViewModel/Source/ModelViewViewModelBlueprint/Private/Bindings/MVVMConversionFunctionHelper.cpp b/Engine/Plugins/Runtime/ModelViewViewModel/Source/ModelViewViewModelBlueprint/Private/Bindings/MVVMConversionFunctionHelper.cpp
--- a/Engine/Plugins/Runtime/ModelViewViewModel/Source/ModelViewViewModelBlueprint/Private/Bindings/MVVMConversionFunctionHelper.cpp
+++ b/Engine/Plugins/Runtime/ModelViewViewModel/Source/ModelViewViewModelBlueprint/Private/Bindings/MVVMConversionFunctionHelper.cpp
@@ -3,6 +3,7 @@
 #include "Bindings/MVVMConversionFunctionHelper.h"
 
 #include "Bindings/MVVMBindingHelper.h"
+#include "Bindings/MVVMFieldPathHelper.h"
 #include "Components/Widget.h"
 #include "Containers/Deque.h"
 #include "MVVMBlueprintView.h"
@@ -10,14 +11,18 @@
 #include "MVVMWidgetBlueprintExtension_View.h"
 #include "UObject/MetaData.h"
 
+#include "EdGraphSchema_K2.h"
 #include "K2Node_BreakStruct.h"
 #include "K2Node_CallFunction.h"
 #include "K2Node_FunctionEntry.h"
 #include "K2Node_FunctionResult.h"
 #include "K2Node_Self.h"
 #include "K2Node_VariableGet.h"
+#include "K2Node_VariableSet.h"
 #include "Kismet2/BlueprintEditorUtils.h"
 
+#define LOCTEXT_NAMESPACE "MVVMConversionFunctionHelper"
+
 namespace UE::MVVM::ConversionFunctionHelper
 {
 
@@ -66,7 +71,13 @@ namespace Private
 		return nullptr;
 	}
 
-	TTuple<UEdGraph*, UK2Node_FunctionEntry*, UK2Node_FunctionResult*> CreateGraph(UBlueprint* Blueprint, FName GraphName, const UFunction* FunctionEntryDefinition, bool bIsConst, bool bIsEditable, bool bAddToBlueprint)
+	struct FCreateGraphResult
+	{
+		UEdGraph* FunctionGraph = nullptr;
+		UK2Node_FunctionEntry* FunctionEntry = nullptr;
+		UK2Node_FunctionResult* FunctionResult = nullptr;
+	};
+	FCreateGraphResult CreateGraph(UBlueprint* Blueprint, FName GraphName, const UFunction* FunctionEntryDefinition, bool bIsConst, bool bIsEditable, bool bAddToBlueprint)
 	{
 		FName UniqueFunctionName = FBlueprintEditorUtils::FindUniqueKismetName(Blueprint, GraphName.ToString());
 		UEdGraph* FunctionGraph = FBlueprintEditorUtils::CreateNewGraph(Blueprint, UniqueFunctionName, UEdGraph::StaticClass(), UEdGraphSchema_K2::StaticClass());
@@ -119,208 +130,27 @@ namespace Private
 		FunctionResult->NodePosY = 0;
 		FunctionResultCreator.Finalize();
 
-		return { FunctionGraph, FunctionEntry, FunctionResult };
-	}
-}
-
-bool RequiresWrapper(const UFunction* ConversionFunction)
-{
-	if (ConversionFunction == nullptr)
-	{
-		return false;
-	}
-
-	TValueOrError<TArray<const FProperty*>, FText> ArgumentsResult = UE::MVVM::BindingHelper::TryGetArgumentsForConversionFunction(ConversionFunction);
-	if (ArgumentsResult.HasValue())
-	{
-		return (ArgumentsResult.GetValue().Num() > 1);
-	}
-	return false;
-}
-
-FName CreateWrapperName(const FMVVMBlueprintViewBinding& Binding, bool bSourceToDestination)
-{
-	TStringBuilder<256> StringBuilder;
-	StringBuilder << TEXT("__");
-	StringBuilder << Binding.GetFName();
-	StringBuilder << (bSourceToDestination ? TEXT("_SourceToDest") : TEXT("_DestToSource"));
-
-	return FName(StringBuilder.ToString());
-}
-
-TPair<UEdGraph*, UK2Node*> CreateGraph(UBlueprint* Blueprint, FName GraphName, const UFunction* Signature, const UFunction* FunctionToWrap, bool bIsConst, bool bTransient)
-{
-	bool bIsEditable = false;
-	bool bAddToBlueprint = !bTransient;
-
-	UEdGraph* FunctionGraph = nullptr;
-	UK2Node_FunctionEntry* FunctionEntry = nullptr;
-	UK2Node_FunctionResult* FunctionResult = nullptr;
-	{
-		TTuple<UEdGraph*, UK2Node_FunctionEntry*, UK2Node_FunctionResult*> NewGraph = Private::CreateGraph(Blueprint, GraphName, Signature, bIsConst, bIsEditable, bAddToBlueprint);
-		FunctionGraph = NewGraph.Get<0>();
-		FunctionEntry = NewGraph.Get<1>();
-		FunctionResult = NewGraph.Get<2>();
-	}
-
-	const FProperty* ReturnProperty = UE::MVVM::BindingHelper::GetReturnProperty(FunctionToWrap);
-
-	// create return value pin
-	if (ReturnProperty)
-	{
-		TSharedPtr<FUserPinInfo> PinInfo = MakeShared<FUserPinInfo>();
-		GetDefault<UEdGraphSchema_K2>()->ConvertPropertyToPinType(ReturnProperty, PinInfo->PinType);
-		PinInfo->PinName = ReturnProperty->GetFName();
-		PinInfo->DesiredPinDirection = EGPD_Input;
-		FunctionResult->UserDefinedPins.Add(PinInfo);
-		FunctionResult->ReconstructNode();
-	}
-	
-	UK2Node_CallFunction* CallFunctionNode = nullptr;
-	{
-		FGraphNodeCreator<UK2Node_CallFunction> CallFunctionCreator(*FunctionGraph);
-		CallFunctionNode = CallFunctionCreator.CreateNode();
-		CallFunctionNode->SetFromFunction(FunctionToWrap);
-		CallFunctionNode->NodePosX = 0;
-		CallFunctionCreator.Finalize();
-		Private::MarkAsConversionFunction(CallFunctionNode, FunctionGraph);
-	}
-
-	// Make link Entry -> CallFunction || Entry -> Return
-	{
-		UEdGraphPin* FunctionEntryThenPin = FunctionEntry->GetThenPin();
-		UEdGraphPin* FunctionResultExecPin = FunctionResult->GetExecPin();
-		
-		if (!CallFunctionNode->IsNodePure())
-		{
-			UEdGraphPin* CallFunctionExecPin = CallFunctionNode->GetExecPin();
-			UEdGraphPin* CallFunctionThenPin = CallFunctionNode->GetThenPin();
-
-			FunctionEntryThenPin->MakeLinkTo(CallFunctionExecPin);
-			CallFunctionThenPin->MakeLinkTo(FunctionResultExecPin);
-
-			CallFunctionNode->NodePosY = 0;
-		}
-		else
-		{
-			FunctionEntryThenPin->MakeLinkTo(FunctionResultExecPin);
-			CallFunctionNode->NodePosY = 100;
-		}
-	}
-
-	if (ReturnProperty)
-	{
-		UEdGraphPin* FunctionReturnPin = CallFunctionNode->FindPin(ReturnProperty->GetName(), EGPD_Output);
-		UEdGraphPin* FunctionResultPin = FunctionResult->FindPin(ReturnProperty->GetFName(), EGPD_Input);
-		check(FunctionResultPin && FunctionReturnPin);
-		FunctionReturnPin->MakeLinkTo(FunctionResultPin);
-	}
-
-	return { FunctionGraph , CallFunctionNode };
-}
-
-
-TPair<UEdGraph*, UK2Node*> CreateGraph(UBlueprint* Blueprint, FName GraphName, const UFunction* Signature, TSubclassOf<UK2Node> NodeType, bool bIsConst, bool bTransient, TFunctionRef<void(UK2Node*)> InitNodeCallback)
-{
-	bool bIsEditable = false;
-	bool bAddToBlueprint = !bTransient;
-
-	UEdGraph* FunctionGraph = nullptr;
-	UK2Node_FunctionEntry* FunctionEntry = nullptr;
-	UK2Node_FunctionResult* FunctionResult = nullptr;
-	{
-		TTuple<UEdGraph*, UK2Node_FunctionEntry*, UK2Node_FunctionResult*> NewGraph = Private::CreateGraph(Blueprint, GraphName, Signature, bIsConst, bIsEditable, bAddToBlueprint);
-		FunctionGraph = NewGraph.Get<0>();
-		FunctionEntry = NewGraph.Get<1>();
-		FunctionResult = NewGraph.Get<2>();
-	}
-
-	UK2Node* CallFunctionNode = nullptr;
-	{
-		FGraphNodeCreator<UK2Node> CallFunctionCreator(*FunctionGraph);
-		CallFunctionNode = CallFunctionCreator.CreateNode(true, NodeType);
-		InitNodeCallback(CallFunctionNode);
-		CallFunctionNode->NodePosX = 0;
-		CallFunctionCreator.Finalize();
-		Private::MarkAsConversionFunction(CallFunctionNode, FunctionGraph);
-	}
-
-	// Make link Entry -> CallFunction || Entry -> Return
-	{
-		UEdGraphPin* FunctionEntryThenPin = FunctionEntry->GetThenPin();
-		UEdGraphPin* FunctionResultExecPin = FunctionResult->GetExecPin();
-
-		if (!CallFunctionNode->IsNodePure())
-		{
-			UEdGraphPin* CallFunctionExecPin = CallFunctionNode->GetExecPin();
-			UEdGraphPin* CallFunctionThenPin = CallFunctionNode->GetThenPin();
-
-			FunctionEntryThenPin->MakeLinkTo(CallFunctionExecPin);
-			CallFunctionThenPin->MakeLinkTo(FunctionResultExecPin);
-
-			CallFunctionNode->NodePosY = 0;
-		}
-		else
-		{
-			FunctionEntryThenPin->MakeLinkTo(FunctionResultExecPin);
-			CallFunctionNode->NodePosY = 100;
-		}
-	}
-
-	return { FunctionGraph , CallFunctionNode };
-}
-
-UK2Node* GetWrapperNode(UEdGraph* Graph)
-{
-	if (Graph == nullptr)
-	{
-		return nullptr;
-	}
-
-	FName ConversionFunctionMetadataKey = Private::ConversionFunctionMetadataKey.Resolve();
-	for (UEdGraphNode* Node : Graph->Nodes)
-	{
-		// check if we've set any metadata on the nodes to figure out which one it is
-		if (Cast<UK2Node>(Node) && Node->GetPackage()->GetMetaData()->HasValue(Node, ConversionFunctionMetadataKey))
-		{
-			return CastChecked<UK2Node>(Node);
-		}
-	}
-
-	if (UK2Node_FunctionResult* FunctionResult = Private::FindFunctionResult(Graph))
-	{
-		for (UEdGraphPin* GraphPin : FunctionResult->Pins)
-		{
-			if (GraphPin->GetFName() != UEdGraphSchema_K2::PN_Execute && GraphPin->LinkedTo.Num() == 1)
-			{
-				if (UK2Node* Node = Cast<UK2Node>(GraphPin->LinkedTo[0]->GetOwningNode()))
-				{
-					Private::MarkAsConversionFunction(Node, Graph);
-					return Node;
-				}
-			}
-		}
+		FCreateGraphResult Result;
+		Result.FunctionGraph = FunctionGraph;
+		Result.FunctionEntry = FunctionEntry;
+		Result.FunctionResult = FunctionResult;
+		return Result;
 	}
 
-	return nullptr;
-}
-
-namespace Private
-{
 	TArray<TTuple<UEdGraphNode*, UEdGraphPin*>> GetPropertyPathGraphNode(const UEdGraphPin* StartPin)
 	{
 		TArray<TTuple<UEdGraphNode*, UEdGraphPin*>> NodesInPath;
 
 		auto AddNode = [&NodesInPath](const UEdGraphPin* Pin)
-		{
-			UEdGraphNode* Result = nullptr;
-			if (Pin->Direction == EGPD_Input && Pin->LinkedTo.Num() == 1 && Pin->PinName != UEdGraphSchema_K2::PN_Execute)
 			{
-				Result = Pin->LinkedTo[0]->GetOwningNode();
-				NodesInPath.Emplace(Result, Pin->LinkedTo[0]);
-			}
-			return Result;
-		};
+				UEdGraphNode* Result = nullptr;
+				if (Pin->Direction == EGPD_Input && Pin->LinkedTo.Num() == 1 && Pin->PinName != UEdGraphSchema_K2::PN_Execute)
+				{
+					Result = Pin->LinkedTo[0]->GetOwningNode();
+					NodesInPath.Emplace(Result, Pin->LinkedTo[0]);
+				}
+				return Result;
+			};
 
 		UEdGraphNode* CurrentNode = AddNode(StartPin);
 		while (CurrentNode)
@@ -341,14 +171,13 @@ namespace Private
 		return NodesInPath;
 	}
 
-
 	FMVVMBlueprintPropertyPath GetPropertyPathForPin(const UBlueprint* Blueprint, const UEdGraphPin* StartPin, bool bSkipResolve)
 	{
 		if (StartPin->Direction != EGPD_Input || StartPin->PinName == UEdGraphSchema_K2::PN_Self)
 		{
 			return FMVVMBlueprintPropertyPath();
 		}
-		
+
 		UMVVMBlueprintView* BlueprintView = GetView(Blueprint);
 		if (BlueprintView == nullptr)
 		{
@@ -358,60 +187,60 @@ namespace Private
 		FMVVMBlueprintPropertyPath ResultPath;
 
 		auto AddRoot = [&ResultPath, BlueprintView, Blueprint, bSkipResolve](FMemberReference& Member)
-		{
-			if (bSkipResolve)
 			{
-				// if the generated class hasn't yet been generated we can blindly forge ahead and try to figure out if it's a widget or a viewmodel
-				if (const FMVVMBlueprintViewModelContext* ViewModel = BlueprintView->FindViewModel(Member.GetMemberName()))
+				if (bSkipResolve)
 				{
-					ResultPath.SetViewModelId(ViewModel->GetViewModelId());
+					// if the generated class hasn't yet been generated we can blindly forge ahead and try to figure out if it's a widget or a viewmodel
+					if (const FMVVMBlueprintViewModelContext* ViewModel = BlueprintView->FindViewModel(Member.GetMemberName()))
+					{
+						ResultPath.SetViewModelId(ViewModel->GetViewModelId());
+					}
+					else
+					{
+						ResultPath.SetWidgetName(Member.GetMemberName());
+						ensure(Member.GetMemberName() != Blueprint->GetFName());
+					}
 				}
 				else
 				{
-					ResultPath.SetWidgetName(Member.GetMemberName());
-					ensure(Member.GetMemberName() != Blueprint->GetFName());
-				}
-			}
-			else
-			{
-				if (const FObjectProperty* Property = CastField<FObjectProperty>(Member.ResolveMember<FProperty>(Blueprint->SkeletonGeneratedClass)))
-				{
-					if (Property->PropertyClass->IsChildOf<UWidget>() || Property->PropertyClass->IsChildOf<UBlueprint>())
+					if (const FObjectProperty* Property = CastField<FObjectProperty>(Member.ResolveMember<FProperty>(Blueprint->SkeletonGeneratedClass)))
 					{
-						ResultPath.SetWidgetName(Property->GetFName());
-						ensure(Property->GetFName() != Blueprint->GetFName());
-					}
-					else if (Property->PropertyClass->ImplementsInterface(UNotifyFieldValueChanged::StaticClass()))
-					{
-						if (const FMVVMBlueprintViewModelContext* ViewModel = BlueprintView->FindViewModel(Property->GetFName()))
+						if (Property->PropertyClass->IsChildOf<UWidget>() || Property->PropertyClass->IsChildOf<UBlueprint>())
+						{
+							ResultPath.SetWidgetName(Property->GetFName());
+							ensure(Property->GetFName() != Blueprint->GetFName());
+						}
+						else if (Property->PropertyClass->ImplementsInterface(UNotifyFieldValueChanged::StaticClass()))
 						{
-							ResultPath.SetViewModelId(ViewModel->GetViewModelId());
+							if (const FMVVMBlueprintViewModelContext* ViewModel = BlueprintView->FindViewModel(Property->GetFName()))
+							{
+								ResultPath.SetViewModelId(ViewModel->GetViewModelId());
+							}
 						}
 					}
 				}
-			}
-		};
+			};
 
 		auto AddPropertyPath = [&ResultPath, Blueprint](FMemberReference& MemberReference)
-		{
-			if (UFunction* Function = MemberReference.ResolveMember<UFunction>(Blueprint->SkeletonGeneratedClass))
-			{
-				ResultPath.AppendPropertyPath(Blueprint, UE::MVVM::FMVVMConstFieldVariant(Function));
-			}
-			else if (const FProperty* Property = MemberReference.ResolveMember<FProperty>(Blueprint->SkeletonGeneratedClass))
 			{
-				ResultPath.AppendPropertyPath(Blueprint, UE::MVVM::FMVVMConstFieldVariant(Property));
-			}
-		};
+				if (UFunction* Function = MemberReference.ResolveMember<UFunction>(Blueprint->SkeletonGeneratedClass))
+				{
+					ResultPath.AppendPropertyPath(Blueprint, UE::MVVM::FMVVMConstFieldVariant(Function));
+				}
+				else if (const FProperty* Property = MemberReference.ResolveMember<FProperty>(Blueprint->SkeletonGeneratedClass))
+				{
+					ResultPath.AppendPropertyPath(Blueprint, UE::MVVM::FMVVMConstFieldVariant(Property));
+				}
+			};
 
 		auto AddBreakNode = [&ResultPath, Blueprint](UScriptStruct* Struct, FName PropertyName)
-		{
-			FProperty* FoundProperty = Struct->FindPropertyByName(PropertyName);
-			if (ensure(FoundProperty))
 			{
-				ResultPath.AppendPropertyPath(Blueprint, UE::MVVM::FMVVMConstFieldVariant(FoundProperty));
-			}
-		};
+				FProperty* FoundProperty = Struct->FindPropertyByName(PropertyName);
+				if (ensure(FoundProperty))
+				{
+					ResultPath.AppendPropertyPath(Blueprint, UE::MVVM::FMVVMConstFieldVariant(FoundProperty));
+				}
+			};
 
 		bool bFirst = true;
 		TArray<TTuple<UEdGraphNode*, UEdGraphPin*>> NodesToSearch = GetPropertyPathGraphNode(StartPin);
@@ -479,69 +308,54 @@ namespace Private
 
 		return ResultPath;
 	}
-}
 
-FMVVMBlueprintPropertyPath GetPropertyPathForPin(const UBlueprint* Blueprint, const UEdGraphPin* Pin, bool bSkipResolve)
-{
-	if (Pin == nullptr || Pin->LinkedTo.Num() == 0)
+	TValueOrError<TArray<UEdGraphPin*>, void> BuildPropertyPath(const UBlueprint* Blueprint, UEdGraph* FunctionGraph, const FMVVMBlueprintPropertyPath& PropertyPath, int32 NumberOfFieldExcludingThePropertyPathSource, FVector2f EndLocation)
 	{
-		return FMVVMBlueprintPropertyPath();
-	}
-
-	return Private::GetPropertyPathForPin(Blueprint, Pin, bSkipResolve);
-}
+		// Add new nodes
+		if (!PropertyPath.IsValid())
+		{
+			return MakeError();
+		}
 
-void SetPropertyPathForPin(const UBlueprint* Blueprint, const FMVVMBlueprintPropertyPath& Path, UEdGraphPin* PathPin)
-{
-	if (PathPin == nullptr)
-	{
-		return;
-	}
+		auto CanNewConnections = [](UEdGraphPin* Pin, UEdGraphPin* PreviousDataPin, const UClass* Context)
+			{
+				return Pin->Direction == EGPD_Input
+					&& Pin->PinName != UEdGraphSchema_K2::PN_Execute
+					&& GetDefault<UEdGraphSchema_K2>()->ArePinsCompatible(PreviousDataPin, Pin, Context);
+			};
 
-	UEdGraphNode* ConversionNode = PathPin->GetOwningNode();
-	UEdGraph* FunctionGraph = ConversionNode ? ConversionNode->GetGraph() : nullptr;
-	const UEdGraphSchema* Schema = FunctionGraph ? FunctionGraph->GetSchema() : nullptr;
+		auto FindNewOutputPin = [](UEdGraphNode* NewNode) -> UEdGraphPin*
+			{
+				// then update our previous pin pointers
+				for (UEdGraphPin* Pin : NewNode->Pins)
+				{
+					if (Pin->Direction == EGPD_Output)
+					{
+						if (Pin->PinName != UEdGraphSchema_K2::PN_Then)
+						{
+							return Pin;
+						}
+					}
+				}
+				return nullptr;
+			};
 
-	UK2Node_FunctionEntry* ConverionFunctionEntry = ConversionNode ? Private::FindFunctionEntry(FunctionGraph) : nullptr;
-	UK2Node_FunctionResult* ConverionFunctionResult = ConversionNode ? Private::FindFunctionResult(FunctionGraph) : nullptr;
+		NumberOfFieldExcludingThePropertyPathSource = FMath::Clamp(NumberOfFieldExcludingThePropertyPathSource, 0, PropertyPath.GetFieldPaths().Num());
 
-	if (!FunctionGraph || !ConversionNode || !ConverionFunctionEntry || !ConverionFunctionResult)
-	{
-		return;
-	}
+		const FVector2f LocationDelta = FVector2f(300.0f, 0.f);
+		FVector2f Location = EndLocation;
+		Location.X -= LocationDelta.X * (NumberOfFieldExcludingThePropertyPathSource + 1);
+		UClass* BlueprintClass = Blueprint->SkeletonGeneratedClass ? Blueprint->SkeletonGeneratedClass : Blueprint->GeneratedClass;
 
-	// Remove previous nodes
-	{
-		TArray<TTuple<UEdGraphNode*, UEdGraphPin*>> AllNodesForPath = Private::GetPropertyPathGraphNode(PathPin);
-		for (const TTuple<UEdGraphNode*, UEdGraphPin*>& Pair : AllNodesForPath)
-		{
-			UEdGraphNode* Node = Pair.Get<UEdGraphNode*>();
-			FunctionGraph->RemoveNode(Node, true);
-		}
-	}
-
-	// Add new nodes
-	if (Path.IsValid())
-	{
-		const int32 ArgumentIndex = ConversionNode->Pins.IndexOfByPredicate([PathPin](const UEdGraphPin* Other){ return Other == PathPin; });
-		if (!ensure(ConversionNode->Pins.IsValidIndex(ArgumentIndex)))
-		{
-			return;
-		}
-
-		TArray<UE::MVVM::FMVVMConstFieldVariant> Fields = Path.GetFields(Blueprint->SkeletonGeneratedClass);
-		float PosX = ConversionNode->NodePosX - 300 * (Fields.Num() + 1);
-		const float PosY = ConversionNode->NodePosY + ArgumentIndex * 100;
-
-		UEdGraphPin* PreviousDataPin = nullptr;
-		UClass* PreviousClass = nullptr;
-
-		// create the root property getter node, ie. the Widget/ViewModel
+		UEdGraphPin* PreviousDataPin = nullptr;
+		UClass* PreviousClass = nullptr;
+		const FProperty* PreviousProperty = nullptr;
+		// create the root property getter node, ie. the Widget/ViewModel
 		{
 			const FProperty* RootProperty = nullptr;
 			bool bCreateSelfNodeForRootProperty = false;
 			{
-				switch (Path.GetSource(Blueprint))
+				switch (PropertyPath.GetSource(Blueprint))
 				{
 				case EMVVMBlueprintFieldPathSource::SelfContext:
 					bCreateSelfNodeForRootProperty = true;
@@ -549,15 +363,16 @@ void SetPropertyPathForPin(const UBlueprint* Blueprint, const FMVVMBlueprintProp
 				case EMVVMBlueprintFieldPathSource::ViewModel:
 				{
 					UMVVMBlueprintView* View = Private::GetView(Blueprint);
-					const FMVVMBlueprintViewModelContext* Context = View ? View->FindViewModel(Path.GetViewModelId()) : nullptr;
+					const FMVVMBlueprintViewModelContext* Context = View ? View->FindViewModel(PropertyPath.GetViewModelId()) : nullptr;
 					RootProperty = Context ? Blueprint->SkeletonGeneratedClass->FindPropertyByName(Context->GetViewModelName()) : nullptr;
 					break;
 				}
 				case EMVVMBlueprintFieldPathSource::Widget:
-					RootProperty = Blueprint->SkeletonGeneratedClass->FindPropertyByName(Path.GetWidgetName());
+					RootProperty = Blueprint->SkeletonGeneratedClass->FindPropertyByName(PropertyPath.GetWidgetName());
 					break;
 				default:
 					check(false);
+					return MakeError();
 				}
 			}
 
@@ -565,70 +380,54 @@ void SetPropertyPathForPin(const UBlueprint* Blueprint, const FMVVMBlueprintProp
 			{
 				FGraphNodeCreator<UK2Node_Self> RootGetterCreator(*FunctionGraph);
 				UK2Node_Self* RootSelfNode = RootGetterCreator.CreateNode();
-				RootSelfNode->NodePosX = PosX;
-				RootSelfNode->NodePosY = PosY;
+				RootSelfNode->NodePosX = Location.X;
+				RootSelfNode->NodePosY = Location.Y;
 				RootGetterCreator.Finalize();
 
 				PreviousDataPin = RootSelfNode->FindPinChecked(UEdGraphSchema_K2::PSC_Self);
 				PreviousClass = Blueprint->SkeletonGeneratedClass;
+				PreviousProperty = RootProperty;
 			}
 			else
 			{
 				if (RootProperty == nullptr)
 				{
 					ensureMsgf(false, TEXT("Could not resolve root property!"));
-					return;
+					return MakeError();
 				}
 
 				FGraphNodeCreator<UK2Node_VariableGet> RootGetterCreator(*FunctionGraph);
 				UK2Node_VariableGet* RootGetterNode = RootGetterCreator.CreateNode();
-				RootGetterNode->NodePosX = PosX;
-				RootGetterNode->NodePosY = PosY;
-				RootGetterNode->VariableReference.SetFromField<FProperty>(RootProperty, true, Blueprint->SkeletonGeneratedClass);
+				RootGetterNode->NodePosX = Location.X;
+				RootGetterNode->NodePosY = Location.Y;
+				RootGetterNode->VariableReference.SetFromField<FProperty>(RootProperty, true, BlueprintClass);
 				RootGetterCreator.Finalize();
 
 				PreviousDataPin = RootGetterNode->Pins[0];
 				PreviousClass = CastField<FObjectProperty>(RootProperty)->PropertyClass;
+				PreviousProperty = RootProperty;
 			}
 
-			PosX += 300;
+			Location += LocationDelta;
 		}
 
+		TArray<UEdGraphPin*> Result;
+		Result.Reserve(NumberOfFieldExcludingThePropertyPathSource);
+		Result.Add(PreviousDataPin);
+
+		TArray<UE::MVVM::FMVVMConstFieldVariant> Fields = PropertyPath.GetFields(BlueprintClass);
+
 		// create all the subsequent nodes in the path
-		const FProperty* PreviousProperty = nullptr;
-		for (int32 Index = 0; Index < Fields.Num(); ++Index)
+		for (int32 Index = 0; Index < NumberOfFieldExcludingThePropertyPathSource; ++Index)
 		{
-			UEdGraphNode* NewNode = nullptr;
 			const UE::MVVM::FMVVMConstFieldVariant& Field = Fields[Index];
-
-			auto CanNewConnections = [Schema](UEdGraphPin* Pin, UEdGraphPin* PreviousDataPin, const UClass* Context)
-			{
-				return Pin->Direction == EGPD_Input
-					&& Pin->PinName != UEdGraphSchema_K2::PN_Execute
-					&& Schema->ArePinsCompatible(PreviousDataPin, Pin, Context);
-			};
-
-			auto FindNewOutputPin = [](UEdGraphNode* NewNode) -> UEdGraphPin*
-			{
-				// then update our previous pin pointers
-				for (UEdGraphPin* Pin : NewNode->Pins)
-				{
-					if (Pin->Direction == EGPD_Output)
-					{
-						if (Pin->PinName != UEdGraphSchema_K2::PN_Then)
-						{
-							return Pin;
-						}
-					}
-				}
-				return nullptr;
-			};
-
+			UEdGraphNode* NewNode = nullptr;
 			UEdGraphPin* NewPreviousDataPin = nullptr;
+			UClass* NewPreviousClass = nullptr;
+			const FProperty* NewPreviousProperty = nullptr;
 			if (Field.IsProperty())
 			{
 				const FProperty* Property = Field.GetProperty();
-
 				// for struct in the middle of a path, we need to use a break node
 				if (const FStructProperty* PreviousStructProperty = CastField<FStructProperty>(PreviousProperty))
 				{
@@ -643,7 +442,7 @@ void SetPropertyPathForPin(const UBlueprint* Blueprint, const FMVVMBlueprintProp
 						MakeStructCreator.Finalize();
 
 						NewNode = FunctionNode;
-						PreviousProperty = UE::MVVM::BindingHelper::GetReturnProperty(Function);
+						NewPreviousProperty = UE::MVVM::BindingHelper::GetReturnProperty(Function);
 					}
 					else
 					{
@@ -654,7 +453,7 @@ void SetPropertyPathForPin(const UBlueprint* Blueprint, const FMVVMBlueprintProp
 						BreakCreator.Finalize();
 
 						NewNode = BreakNode;
-						PreviousProperty = Property;
+						NewPreviousProperty = Property;
 					}
 
 					NewPreviousDataPin = nullptr;
@@ -679,16 +478,21 @@ void SetPropertyPathForPin(const UBlueprint* Blueprint, const FMVVMBlueprintProp
 					GetterCreator.Finalize();
 
 					NewNode = GetterNode;
-					PreviousProperty = Property;
+					NewPreviousProperty = Property;
 					NewPreviousDataPin = FindNewOutputPin(NewNode);
 				}
+				else
+				{
+					ensure(false);
+					return MakeError();
+				}
 			}
 			else if (Field.IsFunction())
 			{
 				if (CastField<FStructProperty>(PreviousProperty))
 				{
 					ensure(false);
-					return;
+					return MakeError();
 				}
 
 				const UFunction* Function = Field.GetFunction();
@@ -700,19 +504,24 @@ void SetPropertyPathForPin(const UBlueprint* Blueprint, const FMVVMBlueprintProp
 				CallFunctionCreator.Finalize();
 
 				NewNode = FunctionNode;
-				PreviousProperty = UE::MVVM::BindingHelper::GetReturnProperty(Function);
+				NewPreviousProperty = UE::MVVM::BindingHelper::GetReturnProperty(Function);
 				NewPreviousDataPin = FindNewOutputPin(NewNode);
 			}
 			else
 			{
 				ensureMsgf(false, TEXT("Invalid path, empty field in path."));
-				return;
+				return MakeError();
 			}
 
+			check(NewNode);
+			NewNode->NodePosX = Location.X;
+			NewNode->NodePosY = Location.Y;
+			Location += LocationDelta;
+
 			if (!NewPreviousDataPin)
 			{
 				ensureMsgf(false, TEXT("A node in path doesn't have a return value."));
-				return;
+				return MakeError();
 			}
 
 			// create new data connections
@@ -724,37 +533,135 @@ void SetPropertyPathForPin(const UBlueprint* Blueprint, const FMVVMBlueprintProp
 				}
 			}
 
-			if (const FObjectProperty* ObjectProperty = CastField<FObjectProperty>(PreviousProperty))
+			if (const FObjectProperty* ObjectProperty = CastField<FObjectProperty>(NewPreviousProperty))
 			{
-				PreviousClass = ObjectProperty->PropertyClass;
+				NewPreviousClass = ObjectProperty->PropertyClass;
 			}
 			else
 			{
-				PreviousClass = nullptr;
+				NewPreviousClass = nullptr;
 			}
 
 			// then update our previous pin pointers
 			PreviousDataPin = NewPreviousDataPin;
+			PreviousClass = NewPreviousClass;
+			PreviousProperty = NewPreviousProperty;
 
-			NewNode->NodePosX = PosX;
-			NewNode->NodePosY = PosY;
+			Result.Add(PreviousDataPin);
+		}
+		return MakeValue(MoveTemp(Result));
+	}
 
-			PosX += 300;
+	struct FCanSetterGraphResult
+	{
+		/**
+		 * true when the path is ObjectA.StructB.StructC.PropertyD
+		 * false when the path is ObjectA.PropertyD
+		 */
+		bool bSplitPin = false;
+		/**
+		 * 1 when the path is ObjectA.StructB.StructC.PropertyD
+		 * 2 when the path is ObjectA.StructB.ObjectC.PropertyD
+		 * Can't be INDEX_NONE if self is the container.
+		 */
+		int32 LocalContainerPathIndex = INDEX_NONE;
+	};
+	TValueOrError<FCanSetterGraphResult, FText> CanCreateSetterGraph(UBlueprint* WidgetBlueprint, const TArrayView<FMVVMConstFieldVariant> Path)
+	{
+		/**
+		 * Different types:
+		 *	(1). The container is a object and it set a property or it call a setter function.
+		 *		propertyA/functionA
+		 *		objectA.objectB.propertyC/functionC
+		 *		objectA.structB.objectC.propertyD/functionD (here we will need to split the pin to get objectC)
+		 *	(2). The container is a struct and we need to create a local variable. We can't have setter function here.
+		 *		objectA.structB.propertyC
+		 *		objectA.structB.structC.propertyD
+		 * We cannot create graph for
+		 *	ObjectA.GetterB, because GetterB is not a valid destination
+		 *	ObjectA.GetterB.PropertyC, if GetterB returns a structure, because we cannot set the local container.
+		 */
+
+		if (Path.Num() == 0)
+		{
+			return MakeError(LOCTEXT("SetterGraph_NoPath", "The path is empty."));
 		}
 
-		// Link the last data pin to the Conversation Function Pin
-		PreviousDataPin->MakeLinkTo(PathPin);
+		// Decide the type.
+		FCanSetterGraphResult Result;
+		for (int32 Index = 0; Index < Path.Num() - 1; ++Index)
+		{
+			const FMVVMConstFieldVariant& Field = Path[Index];
+			TValueOrError<const UStruct*, void> ContainerAsResult = FieldPathHelper::GetFieldAsContainer(Field);
+			if (ContainerAsResult.HasError())
+			{
+				return MakeError(FText::Format(LOCTEXT("SetterGraph_NoContainer", "Can't find the container for field {0}."), FText::FromName(Field.GetName())));
+			}
+
+			// Getter needs to be BlueprintPure. No param.
+			if (!BindingHelper::IsValidForSourceBinding(Field))
+			{
+				return MakeError(FText::Format(LOCTEXT("SetterGraph_InvalidGetter", "The getter for field {0} can't be used."), FText::FromName(Field.GetName())));
+			}
+
+			// Find the container
+			if (const UClass* Class = Cast<const UClass>(ContainerAsResult.GetValue()))
+			{
+				// was a struct, now a class
+				Result = FCanSetterGraphResult();
+				Result.LocalContainerPathIndex = Index;
+			}
+			else if (const UScriptStruct* ScriptStruct = Cast<const UScriptStruct>(ContainerAsResult.GetValue()))
+			{
+				// was a class, now a struct
+				if (!Result.bSplitPin)
+				{
+					Result.LocalContainerPathIndex = Index;
+				}
+				Result.bSplitPin = true;
+			}
+		}
+
+		const FMVVMConstFieldVariant& LastField = Path.Last();
+		if (!BindingHelper::IsValidForDestinationBinding(LastField))
+		{
+			return MakeError(FText::Format(LOCTEXT("SetterGraph_InvalidSetter", "The setter for field {0} can't be used."), FText::FromName(LastField.GetName())));
+		}
+
+		if (Result.bSplitPin)
+		{
+			// Can the local container be set.
+			check(Path.IsValidIndex(Result.LocalContainerPathIndex));
+			const FMVVMConstFieldVariant& Field = Path[Result.LocalContainerPathIndex];
+			if (!BindingHelper::IsValidForDestinationBinding(Field))
+			{
+				return MakeError(FText::Format(LOCTEXT("SetterGraph_ContainerSetter", "The path contains a getter of a struct that can't be set. See field {0}."), FText::FromName(Field.GetName())));
+			}
+		}
+
+		return MakeValue(Result);
 	}
 
-	// Link Then / Exec pin
+	void RemoveNodesFromPin(UEdGraph* FunctionGraph, UEdGraphPin* PathPin)
 	{
-		check(ConversionNode && ConverionFunctionEntry && ConverionFunctionResult);
+		TArray<TTuple<UEdGraphNode*, UEdGraphPin*>> AllNodesForPath = GetPropertyPathGraphNode(PathPin);
+		for (const TTuple<UEdGraphNode*, UEdGraphPin*>& Pair : AllNodesForPath)
+		{
+			UEdGraphNode* Node = Pair.Get<UEdGraphNode*>();
+			FunctionGraph->RemoveNode(Node, true);
+		}
+	}
+
+	void LinkAllNodes(UEdGraph* FunctionGraph, UK2Node_FunctionEntry* FunctionEntry, UEdGraphNode* Wrapper, UK2Node_FunctionResult* FunctionResult)
+	{
+		check(FunctionGraph && FunctionEntry && Wrapper && FunctionResult);
 
-		UEdGraphPin* ThenPin = ConverionFunctionEntry->FindPinChecked(UEdGraphSchema_K2::PN_Then);
+		UEdGraphPin* ThenPin = FunctionEntry->FindPinChecked(UEdGraphSchema_K2::PN_Then);
 		check(ThenPin);
 		ThenPin->BreakAllPinLinks();
 
-		for (UEdGraphPin* Pin : ConversionNode->Pins)
+		// Break Pins for every pins
+		for (UEdGraphPin* Pin : Wrapper->Pins)
 		{
 			TArray<TTuple<UEdGraphNode*, UEdGraphPin*>> AllNodesForPath = Private::GetPropertyPathGraphNode(Pin);
 			for (const TTuple<UEdGraphNode*, UEdGraphPin*>& Pair : AllNodesForPath)
@@ -774,18 +681,427 @@ void SetPropertyPathForPin(const UBlueprint* Blueprint, const FMVVMBlueprintProp
 		}
 
 		// Make pin to the conversion node or to the return node
-		if (UEdGraphPin* CallFunctionExecPin = ConversionNode->FindPin(UEdGraphSchema_K2::PN_Execute))
+		if (UEdGraphPin* CallFunctionExecPin = Wrapper->FindPin(UEdGraphSchema_K2::PN_Execute))
 		{
 			ThenPin->MakeLinkTo(CallFunctionExecPin);
 		}
 		else
 		{
-			UEdGraphPin* FunctionResultExecPin = ConverionFunctionResult->GetExecPin();
+			UEdGraphPin* FunctionResultExecPin = FunctionResult->GetExecPin();
 			ThenPin->MakeLinkTo(FunctionResultExecPin);
 		}
 	}
 }
 
+bool RequiresWrapper(const UFunction* ConversionFunction)
+{
+	if (ConversionFunction == nullptr)
+	{
+		return false;
+	}
+
+	TValueOrError<TArray<const FProperty*>, FText> ArgumentsResult = UE::MVVM::BindingHelper::TryGetArgumentsForConversionFunction(ConversionFunction);
+	if (ArgumentsResult.HasValue())
+	{
+		return (ArgumentsResult.GetValue().Num() > 1);
+	}
+	return false;
+}
+
+FName CreateWrapperName(const FMVVMBlueprintViewBinding& Binding, bool bSourceToDestination)
+{
+	TStringBuilder<256> StringBuilder;
+	StringBuilder << TEXT("__");
+	StringBuilder << Binding.GetFName();
+	StringBuilder << (bSourceToDestination ? TEXT("_SourceToDest") : TEXT("_DestToSource"));
+
+	return FName(StringBuilder.ToString());
+}
+
+TValueOrError<void, FText> CanCreateSetterGraph(UBlueprint* Blueprint, const FMVVMBlueprintPropertyPath& PropertyPath)
+{
+	TArray<UE::MVVM::FMVVMConstFieldVariant> Fields = PropertyPath.GetCompleteFields(Blueprint);
+	TValueOrError<Private::FCanSetterGraphResult, FText> Result = Private::CanCreateSetterGraph(Blueprint, Fields);
+	if (Result.HasError())
+	{
+		return MakeError(Result.StealError());
+	}
+	return MakeValue();
+}
+
+TValueOrError<FCreateSetterGraphResult, FText> CreateSetterGraph(UBlueprint* Blueprint, FName GraphName, const UFunction* Signature, const FMVVMBlueprintPropertyPath& PropertyPath, bool bIsConst, bool bTransient)
+{
+	TArray<UE::MVVM::FMVVMConstFieldVariant> Fields = PropertyPath.GetCompleteFields(Blueprint);
+	TValueOrError<Private::FCanSetterGraphResult, FText> CanCreateSetterGraphResult = Private::CanCreateSetterGraph(Blueprint, Fields);
+	if (CanCreateSetterGraphResult.HasError())
+	{
+		return MakeError(CanCreateSetterGraphResult.StealError());
+	}
+
+	const bool bIsEditable = true;
+	const bool bAddToBlueprint = !bTransient;
+	Private::FCreateGraphResult CreateGraphInternalResult = Private::CreateGraph(Blueprint, GraphName, Signature, bIsConst, bIsEditable, bAddToBlueprint);
+
+	auto RemoveGraph = [Blueprint, bAddToBlueprint, &CreateGraphInternalResult]()
+		{
+			if (CreateGraphInternalResult.FunctionGraph && bAddToBlueprint)
+			{
+				Blueprint->FunctionGraphs.RemoveSingleSwap(CreateGraphInternalResult.FunctionGraph);
+			}
+		};
+
+	if (CreateGraphInternalResult.FunctionGraph == nullptr || CreateGraphInternalResult.FunctionEntry == nullptr)
+	{
+		RemoveGraph();
+		return MakeError(LOCTEXT("SetterGraph_CreateGraphFail", "Can create the graph object."));
+	}
+
+	FCreateSetterGraphResult Result;
+	Result.NewGraph = CreateGraphInternalResult.FunctionGraph;
+	Result.WrappedNode = nullptr;
+
+	// Create the wrapper setter node
+	{
+		check(Fields.IsValidIndex(CanCreateSetterGraphResult.GetValue().LocalContainerPathIndex));
+		int32 WrapperFieldIndex = CanCreateSetterGraphResult.GetValue().bSplitPin ? CanCreateSetterGraphResult.GetValue().LocalContainerPathIndex : (CanCreateSetterGraphResult.GetValue().LocalContainerPathIndex + 1);
+		const FMVVMConstFieldVariant& WrapperField = Fields[WrapperFieldIndex];
+
+		if (WrapperField.IsProperty())
+		{
+			FGraphNodeCreator<UK2Node_VariableSet> CallFunctionCreator(*CreateGraphInternalResult.FunctionGraph);
+			UK2Node_VariableSet* VariableNode = CallFunctionCreator.CreateNode(false);
+			UEdGraphSchema_K2::ConfigureVarNode(VariableNode, WrapperField.GetName(), WrapperField.GetOwner(), Blueprint);
+			VariableNode->NodePosX = 0;
+			CallFunctionCreator.Finalize();
+
+			Result.WrappedNode = VariableNode;
+		}
+		else
+		{
+			check(!CanCreateSetterGraphResult.GetValue().bSplitPin);
+			check(WrapperField.IsFunction());
+
+			FGraphNodeCreator<UK2Node_CallFunction> CallFunctionCreator(*CreateGraphInternalResult.FunctionGraph);
+			UK2Node_CallFunction* CallFunctionNode = CallFunctionCreator.CreateNode();
+			CallFunctionNode->SetFromFunction(WrapperField.GetFunction());
+			CallFunctionNode->NodePosX = 0;
+			CallFunctionCreator.Finalize();
+
+			Result.WrappedNode = CallFunctionNode;
+		}
+		Private::MarkAsConversionFunction(Result.WrappedNode, Result.NewGraph);
+	}
+
+	// Make link Entry -> CallFunction || Entry -> Return
+	if (CreateGraphInternalResult.FunctionResult)
+	{
+		UEdGraphPin* FunctionEntryThenPin = CreateGraphInternalResult.FunctionEntry->GetThenPin();
+		UEdGraphPin* FunctionResultExecPin = CreateGraphInternalResult.FunctionResult->GetExecPin();
+
+		if (Result.WrappedNode->IsNodePure())
+		{
+			FunctionEntryThenPin->MakeLinkTo(FunctionResultExecPin);
+		}
+		else
+		{
+			FunctionEntryThenPin->MakeLinkTo(Result.WrappedNode->GetExecPin());
+			Result.WrappedNode->GetThenPin()->MakeLinkTo(FunctionResultExecPin);
+
+			Result.WrappedNode->NodePosY = 0;
+		}
+	}
+
+	// Build the path to get the container to get to the wrapper.
+	UEdGraphPin* SplitPin = nullptr;
+	if (ensure(Fields.Num() > 0))
+	{
+		UEdGraphPin* WrapperSelfPin = Result.WrappedNode->FindPinChecked(UEdGraphSchema_K2::PSC_Self);
+		int32 ArgumentIndex = Result.WrappedNode->Pins.IndexOfByKey(WrapperSelfPin);
+		const float PosX = Result.WrappedNode->NodePosX;
+		const float PosY = Result.WrappedNode->NodePosY + ArgumentIndex * 100.0f;
+		TValueOrError<TArray<UEdGraphPin*>, void> PropertyPathPinsResult = Private::BuildPropertyPath(Blueprint, CreateGraphInternalResult.FunctionGraph, PropertyPath, CanCreateSetterGraphResult.GetValue().LocalContainerPathIndex, FVector2f(PosX, PosY));
+		if (PropertyPathPinsResult.HasError())
+		{
+			RemoveGraph();
+			return MakeError(LOCTEXT("SetterGraph_BuildPropertyPathFail", "Can't build the path to the setter."));
+		}
+
+		// Link the last data pin to the Conversation Function Pin
+		int32 DeltaIndex = CanCreateSetterGraphResult.GetValue().bSplitPin ? -1 : 0;
+		int32 WrapperPinIndex = PropertyPathPinsResult.GetValue().Num() - 1 + DeltaIndex;
+		UEdGraphPin* ToLinkPin = PropertyPathPinsResult.GetValue()[WrapperPinIndex];
+
+		ToLinkPin->MakeLinkTo(WrapperSelfPin);
+
+		SplitPin = CanCreateSetterGraphResult.GetValue().bSplitPin ? PropertyPathPinsResult.GetValue().Last() : nullptr;
+	}
+
+	auto ForEachPins = [](TArrayView<UEdGraphPin*> Pins, TFunctionRef<void(UEdGraphPin*)> Function)
+		{
+			for (UEdGraphPin* Pin : Pins)
+			{
+				if (Pin->PinName != UEdGraphSchema_K2::PN_Self && Pin->PinName != UEdGraphSchema_K2::PN_Execute && Pin->Direction == EGPD_Input && !Pin->bHidden)
+				{
+					Function(Pin);
+				}
+			}
+		};
+
+	// If the path continue, split the pins and make the link
+	if (CanCreateSetterGraphResult.GetValue().bSplitPin && SplitPin)
+	{
+		// need to break the pin
+		UK2Node* CurrentNode = Result.WrappedNode;
+		check(CurrentNode);
+		TArrayView<UEdGraphPin*> InPins = CurrentNode->Pins;
+		TArrayView<UEdGraphPin*> OutPins = TArrayView<UEdGraphPin*>(&SplitPin, 1);
+		FName PinName;
+		for (int32 Index = CanCreateSetterGraphResult.GetValue().LocalContainerPathIndex; Index < Fields.Num(); ++Index)
+		{
+			const FMVVMConstFieldVariant& NewWrapperField = Fields[Index];
+
+			if (PinName.IsNone())
+			{
+				PinName = NewWrapperField.GetName();
+			}
+			else
+			{
+				PinName = FName(*FString::Printf(TEXT("%s_%s"), *PinName.ToString(), *NewWrapperField.GetName().ToString()));
+				for (int32 SubPinIndex = 0; SubPinIndex < InPins.Num(); ++SubPinIndex)
+				{
+					if (OutPins[SubPinIndex]->GetFName() != PinName)
+					{
+						OutPins[SubPinIndex]->MakeLinkTo(InPins[SubPinIndex]);
+						InPins[SubPinIndex]->bHidden = true;
+					}
+				}
+			}
+
+			bool bLastItem = Index == Fields.Num() -1;
+			if (!bLastItem)
+			{
+				UEdGraphPin** FoundInPinPtr = InPins.FindByPredicate([PinName](const UEdGraphPin* Other) { return Other->GetFName() == PinName && Other->Direction == EGPD_Input && !Other->bHidden; });
+				UEdGraphPin** FoundOutPinPtr = OutPins.FindByPredicate([PinName](const UEdGraphPin* Other) { return Other->GetFName() == PinName && Other->Direction == EGPD_Output && !Other->bHidden; });
+				if (FoundInPinPtr == nullptr
+					|| FoundOutPinPtr == nullptr
+					|| !GetDefault<UEdGraphSchema_K2>()->CanSplitStructPin(**FoundInPinPtr)
+					|| !GetDefault<UEdGraphSchema_K2>()->CanSplitStructPin(**FoundOutPinPtr))
+				{
+					RemoveGraph();
+					return MakeError(FText::Format(LOCTEXT("SetterGraph_CantSplitPin", "The pin {0} can't be split."), FText::FromName(NewWrapperField.GetName())));
+				}
+
+				GetDefault<UEdGraphSchema_K2>()->SplitPin(*FoundInPinPtr, false);
+				GetDefault<UEdGraphSchema_K2>()->SplitPin(*FoundOutPinPtr, false);
+
+				InPins = (*FoundInPinPtr)->SubPins;
+				OutPins = (*FoundOutPinPtr)->SubPins;
+			}
+		}
+	}
+
+	return MakeValue(MoveTemp(Result));
+}
+
+TPair<UEdGraph*, UK2Node*> CreateGraph(UBlueprint* Blueprint, FName GraphName, const UFunction* Signature, const UFunction* FunctionToWrap, bool bIsConst, bool bTransient)
+{
+	bool bIsEditable = false;
+	bool bAddToBlueprint = !bTransient;
+
+	Private::FCreateGraphResult NewGraph = Private::CreateGraph(Blueprint, GraphName, Signature, bIsConst, bIsEditable, bAddToBlueprint);
+
+	const FProperty* ReturnProperty = UE::MVVM::BindingHelper::GetReturnProperty(FunctionToWrap);
+
+	// create return value pin
+	if (ReturnProperty)
+	{
+		TSharedPtr<FUserPinInfo> PinInfo = MakeShared<FUserPinInfo>();
+		GetDefault<UEdGraphSchema_K2>()->ConvertPropertyToPinType(ReturnProperty, PinInfo->PinType);
+		PinInfo->PinName = ReturnProperty->GetFName();
+		PinInfo->DesiredPinDirection = EGPD_Input;
+		NewGraph.FunctionResult->UserDefinedPins.Add(PinInfo);
+		NewGraph.FunctionResult->ReconstructNode();
+	}
+	
+	UK2Node_CallFunction* CallFunctionNode = nullptr;
+	{
+		FGraphNodeCreator<UK2Node_CallFunction> CallFunctionCreator(*NewGraph.FunctionGraph);
+		CallFunctionNode = CallFunctionCreator.CreateNode();
+		CallFunctionNode->SetFromFunction(FunctionToWrap);
+		CallFunctionNode->NodePosX = 0;
+		CallFunctionCreator.Finalize();
+		Private::MarkAsConversionFunction(CallFunctionNode, NewGraph.FunctionGraph);
+	}
+
+	// Make link Entry -> CallFunction || Entry -> Return
+	{
+		UEdGraphPin* FunctionEntryThenPin = NewGraph.FunctionEntry->GetThenPin();
+		UEdGraphPin* FunctionResultExecPin = NewGraph.FunctionResult->GetExecPin();
+		
+		if (!CallFunctionNode->IsNodePure())
+		{
+			UEdGraphPin* CallFunctionExecPin = CallFunctionNode->GetExecPin();
+			UEdGraphPin* CallFunctionThenPin = CallFunctionNode->GetThenPin();
+
+			FunctionEntryThenPin->MakeLinkTo(CallFunctionExecPin);
+			CallFunctionThenPin->MakeLinkTo(FunctionResultExecPin);
+
+			CallFunctionNode->NodePosY = 0;
+		}
+		else
+		{
+			FunctionEntryThenPin->MakeLinkTo(FunctionResultExecPin);
+			CallFunctionNode->NodePosY = 100;
+		}
+	}
+
+	if (ReturnProperty)
+	{
+		UEdGraphPin* FunctionReturnPin = CallFunctionNode->FindPin(ReturnProperty->GetName(), EGPD_Output);
+		UEdGraphPin* FunctionResultPin = NewGraph.FunctionResult->FindPin(ReturnProperty->GetFName(), EGPD_Input);
+		check(FunctionResultPin && FunctionReturnPin);
+		FunctionReturnPin->MakeLinkTo(FunctionResultPin);
+	}
+
+	return { NewGraph.FunctionGraph , CallFunctionNode };
+}
+
+
+TPair<UEdGraph*, UK2Node*> CreateGraph(UBlueprint* Blueprint, FName GraphName, const UFunction* Signature, TSubclassOf<UK2Node> NodeType, bool bIsConst, bool bTransient, TFunctionRef<void(UK2Node*)> InitNodeCallback)
+{
+	bool bIsEditable = false;
+	bool bAddToBlueprint = !bTransient;
+
+	Private::FCreateGraphResult NewGraph = Private::CreateGraph(Blueprint, GraphName, Signature, bIsConst, bIsEditable, bAddToBlueprint);
+
+	UK2Node* CallFunctionNode = nullptr;
+	{
+		FGraphNodeCreator<UK2Node> CallFunctionCreator(*NewGraph.FunctionGraph);
+		CallFunctionNode = CallFunctionCreator.CreateNode(true, NodeType);
+		InitNodeCallback(CallFunctionNode);
+		CallFunctionNode->NodePosX = 0;
+		CallFunctionCreator.Finalize();
+		Private::MarkAsConversionFunction(CallFunctionNode, NewGraph.FunctionGraph);
+	}
+
+	// Make link Entry -> CallFunction || Entry -> Return
+	{
+		UEdGraphPin* FunctionEntryThenPin = NewGraph.FunctionEntry->GetThenPin();
+		UEdGraphPin* FunctionResultExecPin = NewGraph.FunctionResult->GetExecPin();
+
+		if (!CallFunctionNode->IsNodePure())
+		{
+			UEdGraphPin* CallFunctionExecPin = CallFunctionNode->GetExecPin();
+			UEdGraphPin* CallFunctionThenPin = CallFunctionNode->GetThenPin();
+
+			FunctionEntryThenPin->MakeLinkTo(CallFunctionExecPin);
+			CallFunctionThenPin->MakeLinkTo(FunctionResultExecPin);
+
+			CallFunctionNode->NodePosY = 0;
+		}
+		else
+		{
+			FunctionEntryThenPin->MakeLinkTo(FunctionResultExecPin);
+			CallFunctionNode->NodePosY = 100;
+		}
+	}
+
+	return { NewGraph.FunctionGraph , CallFunctionNode };
+}
+
+UK2Node* GetWrapperNode(UEdGraph* Graph)
+{
+	if (Graph == nullptr)
+	{
+		return nullptr;
+	}
+
+	FName ConversionFunctionMetadataKey = Private::ConversionFunctionMetadataKey.Resolve();
+	for (UEdGraphNode* Node : Graph->Nodes)
+	{
+		// check if we've set any metadata on the nodes to figure out which one it is
+		if (Cast<UK2Node>(Node) && Node->GetPackage()->GetMetaData()->HasValue(Node, ConversionFunctionMetadataKey))
+		{
+			return CastChecked<UK2Node>(Node);
+		}
+	}
+
+	if (UK2Node_FunctionResult* FunctionResult = Private::FindFunctionResult(Graph))
+	{
+		for (UEdGraphPin* GraphPin : FunctionResult->Pins)
+		{
+			if (GraphPin->GetFName() != UEdGraphSchema_K2::PN_Execute && GraphPin->LinkedTo.Num() == 1)
+			{
+				if (UK2Node* Node = Cast<UK2Node>(GraphPin->LinkedTo[0]->GetOwningNode()))
+				{
+					Private::MarkAsConversionFunction(Node, Graph);
+					return Node;
+				}
+			}
+		}
+	}
+
+	return nullptr;
+}
+
+FMVVMBlueprintPropertyPath GetPropertyPathForPin(const UBlueprint* Blueprint, const UEdGraphPin* Pin, bool bSkipResolve)
+{
+	if (Pin == nullptr || Pin->LinkedTo.Num() == 0)
+	{
+		return FMVVMBlueprintPropertyPath();
+	}
+
+	return Private::GetPropertyPathForPin(Blueprint, Pin, bSkipResolve);
+}
+
+void SetPropertyPathForPin(const UBlueprint* Blueprint, const FMVVMBlueprintPropertyPath& PropertyPath, UEdGraphPin* PathPin)
+{
+	if (PathPin == nullptr)
+	{
+		return;
+	}
+
+	UEdGraphNode* ConversionNode = PathPin->GetOwningNode();
+	UEdGraph* FunctionGraph = ConversionNode ? ConversionNode->GetGraph() : nullptr;
+	const UEdGraphSchema* Schema = FunctionGraph ? FunctionGraph->GetSchema() : nullptr;
+
+	UK2Node_FunctionEntry* ConverionFunctionEntry = ConversionNode ? Private::FindFunctionEntry(FunctionGraph) : nullptr;
+	UK2Node_FunctionResult* ConverionFunctionResult = ConversionNode ? Private::FindFunctionResult(FunctionGraph) : nullptr;
+
+	if (!FunctionGraph || !ConversionNode || !ConverionFunctionEntry || !ConverionFunctionResult)
+	{
+		return;
+	}
+
+	// Remove previous nodes
+	Private::RemoveNodesFromPin(FunctionGraph, PathPin);
+
+	// Add new nodes
+	if (PropertyPath.IsValid())
+	{
+		const int32 ArgumentIndex = ConversionNode->Pins.IndexOfByPredicate([PathPin](const UEdGraphPin* Other){ return Other == PathPin; });
+		if (!ensure(ConversionNode->Pins.IsValidIndex(ArgumentIndex)))
+		{
+			return;
+		}
+
+		const int32 NumberOfFields = PropertyPath.GetFieldPaths().Num();
+		const float PosX = ConversionNode->NodePosX;
+		const float PosY = ConversionNode->NodePosY + ArgumentIndex * 100;
+		TValueOrError<TArray<UEdGraphPin*>, void> BuildPropertyPathResult = Private::BuildPropertyPath(Blueprint, FunctionGraph, PropertyPath, NumberOfFields, FVector2f(PosX, PosY));
+		if (BuildPropertyPathResult.HasError())
+		{
+			return;
+		}
+
+		// Link the last data pin to the Conversation Function Pin
+		BuildPropertyPathResult.GetValue().Last()->MakeLinkTo(PathPin);
+	}
+
+	// Link Then / Exec pin
+	Private::LinkAllNodes(FunctionGraph, ConverionFunctionEntry, ConversionNode, ConverionFunctionResult);
+}
+
 FMVVMBlueprintPropertyPath GetPropertyPathForArgument(const UBlueprint* WidgetBlueprint, const UK2Node_CallFunction* FunctionNode, FName ArgumentName, bool bSkipResolve)
 {
 	const UEdGraphPin* ArgumentPin = FunctionNode->FindPin(ArgumentName, EGPD_Input);
@@ -826,3 +1142,5 @@ TMap<FName, FMVVMBlueprintPropertyPath> GetAllArgumentPropertyPaths(const UBluep
 }
 
 } //namespace UE::MVVM
+
+#undef LOCTEXT_NAMESPACE
diff --git a/Engine/Plugins/Runtime/ModelViewViewModel/Source/ModelViewViewModelBlueprint/Private/MVVMBlueprintPin.cpp b/Engine/Plugins/Runtime/ModelViewViewModel/Source/ModelViewViewModelBlueprint/Private/MVVMBlueprintPin.cpp
--- a/Engine/Plugins/Runtime/ModelViewViewModel/Source/ModelViewViewModelBlueprint/Private/MVVMBlueprintPin.cpp
+++ b/Engine/Plugins/Runtime/ModelViewViewModel/Source/ModelViewViewModelBlueprint/Private/MVVMBlueprintPin.cpp
@@ -75,7 +75,7 @@ FString FMVVMBlueprintPin::GetValueAsString(const UClass* SelfContext) const
 
 bool FMVVMBlueprintPin::IsInputPin(const UEdGraphPin* Pin)
 {
-	return Pin->PinName != UEdGraphSchema_K2::PN_Self && Pin->PinName != UEdGraphSchema_K2::PN_Execute && Pin->Direction == EGPD_Input && (!Pin->bOrphanedPin || Pin->ShouldSavePinIfOrphaned());
+	return Pin->PinName != UEdGraphSchema_K2::PN_Self && Pin->PinName != UEdGraphSchema_K2::PN_Execute && Pin->Direction == EGPD_Input && (!Pin->bOrphanedPin || Pin->ShouldSavePinIfOrphaned()) && !Pin->bHidden;
 }
 
 TArray<FMVVMBlueprintPin> FMVVMBlueprintPin::CopyAndReturnMissingPins(UBlueprint* Blueprint, UEdGraphNode* GraphNode, const TArray<FMVVMBlueprintPin>& Pins)
diff --git a/Engine/Plugins/Runtime/ModelViewViewModel/Source/ModelViewViewModelBlueprint/Private/MVVMBlueprintViewEvent.cpp b/Engine/Plugins/Runtime/ModelViewViewModel/Source/ModelViewViewModelBlueprint/Private/MVVMBlueprintViewEvent.cpp
--- a/Engine/Plugins/Runtime/ModelViewViewModel/Source/ModelViewViewModelBlueprint/Private/MVVMBlueprintViewEvent.cpp
+++ b/Engine/Plugins/Runtime/ModelViewViewModel/Source/ModelViewViewModelBlueprint/Private/MVVMBlueprintViewEvent.cpp
@@ -5,6 +5,7 @@
 
 #include "Bindings/MVVMConversionFunctionHelper.h"
 #include "Bindings/MVVMBindingHelper.h"
+#include "Bindings/MVVMFieldPathHelper.h"
 #include "Engine/Blueprint.h"
 #include "Kismet2/BlueprintEditorUtils.h"
 #include "MVVMDeveloperProjectSettings.h"
@@ -219,88 +220,21 @@ UEdGraph* UMVVMBlueprintViewEvent::CreateWrapperGraphInternal()
 		return nullptr;
 	}
 
-	if (DestinationPath.GetFieldPaths().Num() > 0)
-	{
-		const FMVVMBlueprintFieldPath& LastPath = DestinationPath.GetFieldPaths().Last();
-		UE::MVVM::FMVVMConstFieldVariant LastField = LastPath.GetField(GetWidgetBlueprintInternal()->SkeletonGeneratedClass);
-		if (LastField.IsFunction() && LastField.GetFunction() != nullptr)
-		{
-			return CreateWrapperGraphInternal(DelegateSignature, LastField.GetFunction());
-		}
-		else if (LastField.IsProperty() && LastField.GetProperty() != nullptr)
-		{
-			return CreateWrapperGraphInternal(DelegateSignature, LastField.GetProperty());
-		}
-	}
-	else if (DestinationPath.GetViewModelId().IsValid())
-	{
-		const FMVVMBlueprintViewModelContext* Context = GetOuterUMVVMBlueprintView()->FindViewModel(DestinationPath.GetViewModelId());
-		const FProperty* ViewModel = Context ? GetWidgetBlueprintInternal()->SkeletonGeneratedClass->FindPropertyByName(Context->GetViewModelName()) : nullptr;
-		if (ViewModel)
-		{
-			return CreateWrapperGraphInternal(DelegateSignature, ViewModel);
-		}
-	}
-	return nullptr;
-}
-
-UEdGraph* UMVVMBlueprintViewEvent::CreateWrapperGraphInternal(const UFunction* DelegateSignature, const UFunction* Function)
-{
 	UWidgetBlueprint* WidgetBlueprint = GetWidgetBlueprintInternal();
 	bool bIsConst = false;
 	bool bTransient = true;
-	TPair<UEdGraph*, UK2Node*> Result = UE::MVVM::ConversionFunctionHelper::CreateGraph(WidgetBlueprint, GraphName, DelegateSignature, Function, bIsConst, bTransient);
-	SetCachedWrapperGraphInternal(Result.Get<0>(), Result.Get<1>());
-
-	// Add SelfNode if needed
-	if (CachedWrapperNode)
+	TValueOrError<UE::MVVM::ConversionFunctionHelper::FCreateSetterGraphResult, FText> CreateSetterGraphResult = UE::MVVM::ConversionFunctionHelper::CreateSetterGraph(WidgetBlueprint, GraphName, DelegateSignature, DestinationPath, bIsConst, bTransient);
+	if (CreateSetterGraphResult.HasError())
 	{
-		if (UEdGraphPin* SelfPin = CachedWrapperNode->FindPin(UEdGraphSchema_K2::PSC_Self))
-		{
-			FMVVMBlueprintPropertyPath PathToSet = DestinationPath;
-			TArray<UE::MVVM::FMVVMConstFieldVariant> FieldPaths = PathToSet.GetFields(WidgetBlueprint->SkeletonGeneratedClass);
-			PathToSet.ResetPropertyPath();
-			for (int32 Index = 0; Index < FieldPaths.Num() - 1; ++Index)
-			{
-				PathToSet.AppendPropertyPath(WidgetBlueprint, FieldPaths[Index]);
-			}
-			;
-			UE::MVVM::ConversionFunctionHelper::SetPropertyPathForPin(WidgetBlueprint, PathToSet, SelfPin);
-		}
+		SetCachedWrapperGraphInternal(nullptr, nullptr);
+		return nullptr;
 	}
-
-	LoadPinValuesInternal();
-	return CachedWrapperGraph;
-}
-
-UEdGraph* UMVVMBlueprintViewEvent::CreateWrapperGraphInternal(const UFunction* DelegateSignature, const FProperty* Property)
-{
-	UWidgetBlueprint* WidgetBlueprint = GetWidgetBlueprintInternal();
-	bool bConst = false;
-	bool bTransient = true;
-	TPair<UEdGraph*, UK2Node*> Result = UE::MVVM::ConversionFunctionHelper::CreateGraph(WidgetBlueprint, GraphName, DelegateSignature, UK2Node_VariableSet::StaticClass(), bConst, bTransient,
-		[WidgetBlueprint, Property](UK2Node* NewNode)
-		{
-			UK2Node_VariableSet* VariableNode = CastChecked<UK2Node_VariableSet>(NewNode);
-			UEdGraphSchema_K2::ConfigureVarNode(VariableNode, Property->GetFName(), Property->GetOwnerStruct(), WidgetBlueprint);
-		});
-	SetCachedWrapperGraphInternal(Result.Get<0>(), Result.Get<1>());
-
-	// Add SelfNode
-	if (CachedWrapperNode)
+	else
 	{
-		FMVVMBlueprintPropertyPath PathToSet = DestinationPath;
-		TArray<UE::MVVM::FMVVMConstFieldVariant> FieldPaths = PathToSet.GetFields(WidgetBlueprint->SkeletonGeneratedClass);
-		PathToSet.ResetPropertyPath();
-		for (int32 Index = 0; Index < FieldPaths.Num() - 1; ++Index)
-		{
-			PathToSet.AppendPropertyPath(WidgetBlueprint, FieldPaths[Index]);
-		}
-		UEdGraphPin* SelfPin = CachedWrapperNode->FindPinChecked(UEdGraphSchema_K2::PSC_Self);
-		UE::MVVM::ConversionFunctionHelper::SetPropertyPathForPin(WidgetBlueprint, PathToSet, SelfPin);
+		SetCachedWrapperGraphInternal(CreateSetterGraphResult.GetValue().NewGraph, CreateSetterGraphResult.GetValue().WrappedNode);
+		LoadPinValuesInternal();
 	}
 
-	LoadPinValuesInternal();
 	return CachedWrapperGraph;
 }
 
diff --git a/Engine/Plugins/Runtime/ModelViewViewModel/Source/ModelViewViewModelBlueprint/Public/Bindings/MVVMConversionFunctionHelper.h b/Engine/Plugins/Runtime/ModelViewViewModel/Source/ModelViewViewModelBlueprint/Public/Bindings/MVVMConversionFunctionHelper.h
--- a/Engine/Plugins/Runtime/ModelViewViewModel/Source/ModelViewViewModelBlueprint/Public/Bindings/MVVMConversionFunctionHelper.h
+++ b/Engine/Plugins/Runtime/ModelViewViewModel/Source/ModelViewViewModelBlueprint/Public/Bindings/MVVMConversionFunctionHelper.h
@@ -5,6 +5,7 @@
 #include "Containers/Map.h"
 #include "Templates/Function.h"
 #include "Templates/SubclassOf.h"
+#include "Templates/ValueOrError.h"
 
 struct FMVVMBlueprintPropertyPath;
 struct FMVVMBlueprintViewBinding;
@@ -17,6 +18,7 @@ class UK2Node;
 class UK2Node_CallFunction;
 class FKismetCompilerContext;
 class UMVVMBlueprintView;
+namespace UE::MVVM { struct FMVVMConstFieldVariant; }
 
 namespace UE::MVVM::ConversionFunctionHelper
 {
@@ -33,7 +35,7 @@ namespace UE::MVVM::ConversionFunctionHelper
 	MODELVIEWVIEWMODELBLUEPRINT_API FMVVMBlueprintPropertyPath GetPropertyPathForPin(const UBlueprint* WidgetBlueprint, const UEdGraphPin* Pin, bool bSkipResolve);
 	
 	/** Set the property path of a given argument in the conversion function. */
-	MODELVIEWVIEWMODELBLUEPRINT_API void SetPropertyPathForPin(const UBlueprint* Blueprint, const FMVVMBlueprintPropertyPath& Path, UEdGraphPin* Pin);
+	MODELVIEWVIEWMODELBLUEPRINT_API void SetPropertyPathForPin(const UBlueprint* Blueprint, const FMVVMBlueprintPropertyPath& PropertyPath, UEdGraphPin* Pin);
 	
 	/** Find the property path of a given argument in the conversion function. */
 	MODELVIEWVIEWMODELBLUEPRINT_API FMVVMBlueprintPropertyPath GetPropertyPathForArgument(const UBlueprint* WidgetBlueprint, const UK2Node_CallFunction* Function, FName ArgumentName, bool bSkipResolve);
@@ -41,6 +43,24 @@ namespace UE::MVVM::ConversionFunctionHelper
 	/** Create the name of the conversion function wrapper this binding should have. */
 	MODELVIEWVIEWMODELBLUEPRINT_API FName CreateWrapperName(const FMVVMBlueprintViewBinding& Binding, bool bSourceToDestination);
 
+	/**
+	 * If we can create a graph to set a property/function.
+	 */
+	MODELVIEWVIEWMODELBLUEPRINT_API TValueOrError<void, FText> CanCreateSetterGraph(UBlueprint* WidgetBlueprint, const FMVVMBlueprintPropertyPath& PropertyPath);
+
+	struct FCreateSetterGraphResult
+	{
+		/** The new graph created. */
+		UEdGraph* NewGraph = nullptr;
+		/** Node that owns the pins. */
+		UK2Node* WrappedNode = nullptr;
+	};
+
+	/**
+	 * Create a graph to set a property/function.
+	 */
+	MODELVIEWVIEWMODELBLUEPRINT_API TValueOrError<FCreateSetterGraphResult, FText> CreateSetterGraph(UBlueprint* WidgetBlueprint, FName GraphName, const UFunction* Signature, const FMVVMBlueprintPropertyPath& PropertyPath, bool bIsConst, bool bTransient);
+
 	/** */
 	MODELVIEWVIEWMODELBLUEPRINT_API TPair<UEdGraph*, UK2Node*> CreateGraph(UBlueprint* WidgetBlueprint, FName GraphName, const UFunction* Signature, const UFunction* FunctionToWrap, bool bIsConst, bool bTransient);
 	
diff --git a/Engine/Plugins/Runtime/ModelViewViewModel/Source/ModelViewViewModelBlueprint/Public/MVVMBlueprintViewEvent.h b/Engine/Plugins/Runtime/ModelViewViewModel/Source/ModelViewViewModelBlueprint/Public/MVVMBlueprintViewEvent.h
--- a/Engine/Plugins/Runtime/ModelViewViewModel/Source/ModelViewViewModelBlueprint/Public/MVVMBlueprintViewEvent.h
+++ b/Engine/Plugins/Runtime/ModelViewViewModel/Source/ModelViewViewModelBlueprint/Public/MVVMBlueprintViewEvent.h
@@ -127,8 +127,6 @@ private:
 	UWidgetBlueprint* GetWidgetBlueprintInternal() const;
 	void SetCachedWrapperGraphInternal(UEdGraph* Graph, UK2Node* Node);
 	UEdGraph* CreateWrapperGraphInternal();
-	UEdGraph* CreateWrapperGraphInternal(const UFunction* Signature, const UFunction* Function);
-	UEdGraph* CreateWrapperGraphInternal(const UFunction* Signature, const FProperty* Property);
 	void LoadPinValuesInternal();
 
 private:
diff --git a/Engine/Plugins/Runtime/ModelViewViewModel/Source/ModelViewViewModelEditor/Private/Widgets/SMVVMEventParameter.cpp b/Engine/Plugins/Runtime/ModelViewViewModel/Source/ModelViewViewModelEditor/Private/Widgets/SMVVMEventParameter.cpp
--- a/Engine/Plugins/Runtime/ModelViewViewModel/Source/ModelViewViewModelEditor/Private/Widgets/SMVVMEventParameter.cpp
+++ b/Engine/Plugins/Runtime/ModelViewViewModel/Source/ModelViewViewModelEditor/Private/Widgets/SMVVMEventParameter.cpp
@@ -37,33 +37,34 @@ void SEventParameter::Construct(const FArguments& InArgs, UWidgetBlueprint* InWi
 
 	bAllowDefault = InArgs._AllowDefault;
 
-	UEdGraphPin* Pin = InArgs._Event->GetOrCreateGraphPin(ParameterName);
-	check(Pin);
-
-	bool bIsBooleanPin = Pin->PinType.PinCategory == UEdGraphSchema_K2::PC_Boolean;
+	bool bIsBooleanPin = false;
 	TSharedRef<SWidget> ValueWidget = SNullWidget::NullWidget;
-	// create a new pin widget so that we can get the default value widget out of it
-	if (TSharedPtr<SGraphPin> PinWidget = FNodeFactory::CreatePinWidget(Pin))
+	if (UEdGraphPin* Pin = InArgs._Event->GetOrCreateGraphPin(ParameterName))
 	{
-		GraphPin = PinWidget;
-		ValueWidget = PinWidget->GetDefaultValueWidget();
-	}
-
-	if (ValueWidget == SNullWidget::NullWidget)
-	{
-		ValueWidget = SNew(STextBlock)
-			.Text(LOCTEXT("DefaultValue", "Default Value"))
-			.TextStyle(FAppStyle::Get(), "HintText");
-	}
-	// booleans are represented by a checkbox which doesn't expand to the minsize we have, so don't put a border around them
-	else if (!bIsBooleanPin)
-	{
-		ValueWidget = SNew(SBorder)
-			.Padding(0.0f)
-			.BorderImage(FMVVMEditorStyle::Get().GetBrush("FunctionParameter.Border"))
-			[
-				ValueWidget
-			];
+		bIsBooleanPin = Pin->PinType.PinCategory == UEdGraphSchema_K2::PC_Boolean;
+		// create a new pin widget so that we can get the default value widget out of it
+		if (TSharedPtr<SGraphPin> PinWidget = FNodeFactory::CreatePinWidget(Pin))
+		{
+			GraphPin = PinWidget;
+			ValueWidget = PinWidget->GetDefaultValueWidget();
+		}
+
+		if (ValueWidget == SNullWidget::NullWidget)
+		{
+			ValueWidget = SNew(STextBlock)
+				.Text(LOCTEXT("DefaultValue", "Default Value"))
+				.TextStyle(FAppStyle::Get(), "HintText");
+		}
+		// booleans are represented by a checkbox which doesn't expand to the minsize we have, so don't put a border around them
+		else if (!bIsBooleanPin)
+		{
+			ValueWidget = SNew(SBorder)
+				.Padding(0.0f)
+				.BorderImage(FMVVMEditorStyle::Get().GetBrush("FunctionParameter.Border"))
+				[
+					ValueWidget
+				];
+		}
 	}
 
 	FMVVMBlueprintPropertyPath Path = OnGetSelectedField();
