MVVM: Event do not save the pin id. The graph is rebuilt and the id is new everytime. Warn the user, at compile, when a pin is orphaned. Fix event's setter on the UserWidget (the index is -1 because of that).

diff --git a/Engine/Plugins/Runtime/ModelViewViewModel/Source/ModelViewViewModelBlueprint/Private/Bindings/MVVMConversionFunctionHelper.cpp b/Engine/Plugins/Runtime/ModelViewViewModel/Source/ModelViewViewModelBlueprint/Private/Bindings/MVVMConversionFunctionHelper.cpp
--- a/Engine/Plugins/Runtime/ModelViewViewModel/Source/ModelViewViewModelBlueprint/Private/Bindings/MVVMConversionFunctionHelper.cpp
+++ b/Engine/Plugins/Runtime/ModelViewViewModel/Source/ModelViewViewModelBlueprint/Private/Bindings/MVVMConversionFunctionHelper.cpp
@@ -562,7 +562,7 @@ namespace Private
 		/**
 		 * 1 when the path is ObjectA.StructB.StructC.PropertyD
 		 * 2 when the path is ObjectA.StructB.ObjectC.PropertyD
-		 * Can't be INDEX_NONE if self is the container.
+		 * Can be INDEX_NONE if self is the container.
 		 */
 		int32 LocalContainerPathIndex = INDEX_NONE;
 	};
@@ -762,8 +762,8 @@ TValueOrError<FCreateSetterGraphResult, FText> CreateSetterGraph(UBlueprint* Blu
 
 	// Create the wrapper setter node
 	{
-		check(Fields.IsValidIndex(CanCreateSetterGraphResult.GetValue().LocalContainerPathIndex));
 		int32 WrapperFieldIndex = CanCreateSetterGraphResult.GetValue().bSplitPin ? CanCreateSetterGraphResult.GetValue().LocalContainerPathIndex : (CanCreateSetterGraphResult.GetValue().LocalContainerPathIndex + 1);
+		check(Fields.IsValidIndex(WrapperFieldIndex));
 		const FMVVMConstFieldVariant& WrapperField = Fields[WrapperFieldIndex];
 
 		if (WrapperField.IsProperty())
diff --git a/Engine/Plugins/Runtime/ModelViewViewModel/Source/ModelViewViewModelBlueprint/Private/MVVMBlueprintPin.cpp b/Engine/Plugins/Runtime/ModelViewViewModel/Source/ModelViewViewModelBlueprint/Private/MVVMBlueprintPin.cpp
--- a/Engine/Plugins/Runtime/ModelViewViewModel/Source/ModelViewViewModelBlueprint/Private/MVVMBlueprintPin.cpp
+++ b/Engine/Plugins/Runtime/ModelViewViewModel/Source/ModelViewViewModelBlueprint/Private/MVVMBlueprintPin.cpp
@@ -16,7 +16,6 @@ FMVVMBlueprintPin::FMVVMBlueprintPin(FName InPinName)
 FMVVMBlueprintPin FMVVMBlueprintPin::CreateFromPin(const UBlueprint* Blueprint, const UEdGraphPin* Pin)
 {
 	FMVVMBlueprintPin Result;
-	Result.PinId = Pin->PinId;
 	Result.PinName = Pin->PinName;
 	Result.Path = UE::MVVM::ConversionFunctionHelper::GetPropertyPathForPin(Blueprint, Pin, true);
 	Result.DefaultObject = Pin->DefaultObject;
@@ -111,7 +110,7 @@ void FMVVMBlueprintPin::CopyTo(const UBlueprint* Blueprint, UEdGraphNode* Node)
 	Status = EMVVMBlueprintPinStatus::Orphaned;
 	if (UEdGraphPin* GraphPin = FindGraphPin(Node))
 	{
-		if (IsInputPin(GraphPin) && !GraphPin->bOrphanedPin)
+		if (IsInputPin(GraphPin))
 		{
 			Status = EMVVMBlueprintPinStatus::Valid;
 			if (bSplit && GraphPin->SubPins.Num() == 0 && GetDefault<UEdGraphSchema_K2>()->CanSplitStructPin(*GraphPin))
@@ -133,7 +132,7 @@ TArray<FMVVMBlueprintPin> FMVVMBlueprintPin::CreateFromNode(UBlueprint* Blueprin
 	Result.Reserve(GraphNode->Pins.Num());
 	for (const UEdGraphPin* GraphPin : GraphNode->Pins)
 	{
-		if (IsInputPin(GraphPin))
+		if (IsInputPin(GraphPin) && !GraphPin->bOrphanedPin)
 		{
 			Result.Add(FMVVMBlueprintPin::CreateFromPin(Blueprint, GraphPin));
 		}
@@ -143,16 +142,11 @@ TArray<FMVVMBlueprintPin> FMVVMBlueprintPin::CreateFromNode(UBlueprint* Blueprin
 
 UEdGraphPin* FMVVMBlueprintPin::FindGraphPin(UEdGraphNode* Node) const
 {
-	if (PinId.IsValid())
-	{
-		return Node->FindPinById(PinId);
-	}
 	return Node->FindPin(PinName, EGPD_Input);
 }
 
 void FMVVMBlueprintPin::Reset()
 {
-	PinId.Invalidate();
 	Path = FMVVMBlueprintPropertyPath();
 	DefaultString.Empty();
 	DefaultText = FText();
diff --git a/Engine/Plugins/Runtime/ModelViewViewModel/Source/ModelViewViewModelBlueprint/Private/MVVMBlueprintViewConversionFunction.cpp b/Engine/Plugins/Runtime/ModelViewViewModel/Source/ModelViewViewModelBlueprint/Private/MVVMBlueprintViewConversionFunction.cpp
--- a/Engine/Plugins/Runtime/ModelViewViewModel/Source/ModelViewViewModelBlueprint/Private/MVVMBlueprintViewConversionFunction.cpp
+++ b/Engine/Plugins/Runtime/ModelViewViewModel/Source/ModelViewViewModelBlueprint/Private/MVVMBlueprintViewConversionFunction.cpp
@@ -428,12 +428,35 @@ void UMVVMBlueprintViewConversionFunction::SetGraphPin(UBlueprint* Blueprint, FN
 
 void UMVVMBlueprintViewConversionFunction::SavePinValues(UBlueprint* Blueprint)
 {
+	SavedPins.Empty();
 	if (CachedWrapperNode)
 	{
 		SavedPins = FMVVMBlueprintPin::CreateFromNode(Blueprint, CachedWrapperNode);
 	}
 }
 
+void UMVVMBlueprintViewConversionFunction::UpdatePinValues(UBlueprint* Blueprint)
+{
+	if (CachedWrapperNode)
+	{
+		TArray<FMVVMBlueprintPin> TmpSavedPins = FMVVMBlueprintPin::CreateFromNode(Blueprint, CachedWrapperNode);
+		SavedPins.RemoveAll([](const FMVVMBlueprintPin& Pin) { return Pin.GetStatus() != EMVVMBlueprintPinStatus::Orphaned; });
+		SavedPins.Append(TmpSavedPins);
+	}
+}
+
+bool UMVVMBlueprintViewConversionFunction::HasOrphanedPin() const
+{
+	for (const FMVVMBlueprintPin& Pin : SavedPins)
+	{
+		if (Pin.GetStatus() == EMVVMBlueprintPinStatus::Orphaned)
+		{
+			return true;
+		}
+	}
+	return false;
+}
+
 void UMVVMBlueprintViewConversionFunction::LoadPinValuesInternal(UBlueprint* Blueprint)
 {
 	if (CachedWrapperNode)
diff --git a/Engine/Plugins/Runtime/ModelViewViewModel/Source/ModelViewViewModelBlueprint/Private/MVVMBlueprintViewEvent.cpp b/Engine/Plugins/Runtime/ModelViewViewModel/Source/ModelViewViewModelBlueprint/Private/MVVMBlueprintViewEvent.cpp
--- a/Engine/Plugins/Runtime/ModelViewViewModel/Source/ModelViewViewModelBlueprint/Private/MVVMBlueprintViewEvent.cpp
+++ b/Engine/Plugins/Runtime/ModelViewViewModel/Source/ModelViewViewModelBlueprint/Private/MVVMBlueprintViewEvent.cpp
@@ -125,6 +125,7 @@ UEdGraphPin* UMVVMBlueprintViewEvent::GetOrCreateGraphPin(FName PinName)
 
 void UMVVMBlueprintViewEvent::SavePinValues()
 {
+	SavedPins.Empty();
 	if (CachedWrapperNode)
 	{
 		UWidgetBlueprint* Blueprint = GetWidgetBlueprintInternal();
@@ -132,6 +133,29 @@ void UMVVMBlueprintViewEvent::SavePinValues()
 	}
 }
 
+void UMVVMBlueprintViewEvent::UpdatePinValues()
+{
+	if (CachedWrapperNode)
+	{
+		UWidgetBlueprint* Blueprint = GetWidgetBlueprintInternal();
+		TArray<FMVVMBlueprintPin> TmpSavedPins = FMVVMBlueprintPin::CreateFromNode(Blueprint, CachedWrapperNode);
+		SavedPins.RemoveAll([](const FMVVMBlueprintPin& Pin){ return Pin.GetStatus() != EMVVMBlueprintPinStatus::Orphaned; });
+		SavedPins.Append(TmpSavedPins);
+	}
+}
+
+bool UMVVMBlueprintViewEvent::HasOrphanedPin() const
+{
+	for (const FMVVMBlueprintPin& Pin : SavedPins)
+	{
+		if (Pin.GetStatus() == EMVVMBlueprintPinStatus::Orphaned)
+		{
+			return true;
+		}
+	}
+	return false;
+}
+
 FMVVMBlueprintPropertyPath UMVVMBlueprintViewEvent::GetPinPath(FName PinName) const
 {
 	const FMVVMBlueprintPin* ViewPin = SavedPins.FindByPredicate([PinName](const FMVVMBlueprintPin& Other) { return PinName == Other.GetName(); });
diff --git a/Engine/Plugins/Runtime/ModelViewViewModel/Source/ModelViewViewModelBlueprint/Private/MVVMViewBlueprintCompiler.cpp b/Engine/Plugins/Runtime/ModelViewViewModel/Source/ModelViewViewModelBlueprint/Private/MVVMViewBlueprintCompiler.cpp
--- a/Engine/Plugins/Runtime/ModelViewViewModel/Source/ModelViewViewModelBlueprint/Private/MVVMViewBlueprintCompiler.cpp
+++ b/Engine/Plugins/Runtime/ModelViewViewModel/Source/ModelViewViewModelBlueprint/Private/MVVMViewBlueprintCompiler.cpp
@@ -972,7 +972,15 @@ void FMVVMViewBlueprintCompiler::CategorizeBindings(const FWidgetBlueprintCompil
 				bIsCreateFunctionsStepValid = false;
 				continue;
 			}
-			ConversionFunction->SavePinValues(WidgetBlueprintCompilerContext.WidgetBlueprint());
+
+			ConversionFunction->UpdatePinValues(WidgetBlueprintCompilerContext.WidgetBlueprint());
+
+			if (ConversionFunction->HasOrphanedPin())
+			{
+				AddMessageForBinding(Binding, LOCTEXT("InvalidConversionFunctionGraphOrphaned", "The conversion function has an orphaned pin."), EMessageType::Warning, FName());
+				bIsCreateFunctionsStepValid = false;
+				continue;
+			}
 
 			if (ConversionFunction->GetPins().Num() == 0)
 			{
@@ -1024,7 +1032,15 @@ void FMVVMViewBlueprintCompiler::CategorizeEvents(const FWidgetBlueprintCompiler
 			bIsCreateFunctionsStepValid = false;
 			continue;
 		}
-		EventPtr->SavePinValues();
+
+		EventPtr->UpdatePinValues();
+
+		if (EventPtr->HasOrphanedPin())
+		{
+			AddMessageForEvent(EventPtr, LOCTEXT("InvalidEventGraphOrphaned", "The event has an orphaned pin."), EMessageType::Warning, FName());
+			bIsCreateFunctionsStepValid = false;
+			continue;
+		}
 
 		Event->Type = FCompilerEvent::EType::Valid;
 	}
diff --git a/Engine/Plugins/Runtime/ModelViewViewModel/Source/ModelViewViewModelBlueprint/Public/MVVMBlueprintPin.h b/Engine/Plugins/Runtime/ModelViewViewModel/Source/ModelViewViewModelBlueprint/Public/MVVMBlueprintPin.h
--- a/Engine/Plugins/Runtime/ModelViewViewModel/Source/ModelViewViewModelBlueprint/Public/MVVMBlueprintPin.h
+++ b/Engine/Plugins/Runtime/ModelViewViewModel/Source/ModelViewViewModelBlueprint/Public/MVVMBlueprintPin.h
@@ -29,9 +29,6 @@ struct MODELVIEWVIEWMODELBLUEPRINT_API FMVVMBlueprintPin
 	GENERATED_BODY()
 
 private:
-	UPROPERTY(VisibleAnywhere, Category = "Viewmodel")
-	FGuid PinId;
-
 	UPROPERTY(VisibleAnywhere, Category = "Viewmodel")
 	FName PinName;
 
diff --git a/Engine/Plugins/Runtime/ModelViewViewModel/Source/ModelViewViewModelBlueprint/Public/MVVMBlueprintViewConversionFunction.h b/Engine/Plugins/Runtime/ModelViewViewModel/Source/ModelViewViewModelBlueprint/Public/MVVMBlueprintViewConversionFunction.h
--- a/Engine/Plugins/Runtime/ModelViewViewModel/Source/ModelViewViewModelBlueprint/Public/MVVMBlueprintViewConversionFunction.h
+++ b/Engine/Plugins/Runtime/ModelViewViewModel/Source/ModelViewViewModelBlueprint/Public/MVVMBlueprintViewConversionFunction.h
@@ -108,6 +108,10 @@ public:
 
 	/** Generates SavedPins from the wrapper graph, if it exists. */
 	void SavePinValues(UBlueprint* Blueprint);
+	/** Keep the orphaned pins. Add the missing pins. */
+	void UpdatePinValues(UBlueprint* Blueprint);
+	/** Keep the orphaned pins. Add the missing pins. */
+	bool HasOrphanedPin() const;
 
 	FSimpleMulticastDelegate OnWrapperGraphModified;
 
diff --git a/Engine/Plugins/Runtime/ModelViewViewModel/Source/ModelViewViewModelBlueprint/Public/MVVMBlueprintViewEvent.h b/Engine/Plugins/Runtime/ModelViewViewModel/Source/ModelViewViewModelBlueprint/Public/MVVMBlueprintViewEvent.h
--- a/Engine/Plugins/Runtime/ModelViewViewModel/Source/ModelViewViewModelBlueprint/Public/MVVMBlueprintViewEvent.h
+++ b/Engine/Plugins/Runtime/ModelViewViewModel/Source/ModelViewViewModelBlueprint/Public/MVVMBlueprintViewEvent.h
@@ -87,6 +87,10 @@ public:
 
 	/** Generates SavedPins from the wrapper graph, if it exists. */
 	void SavePinValues();
+	/** Keep the orphaned pins. Add the missing pins. */
+	void UpdatePinValues();
+	/** Keep the orphaned pins. Add the missing pins. */
+	bool HasOrphanedPin() const;
 
 	UEdGraphPin* GetOrCreateGraphPin(FName PinName);
 
