MVVM: Enable setter flag by default. Add flag to generate a public getter. Change the display logic for the optional flag and the setter flag.

diff --git a/Engine/Plugins/Runtime/ModelViewViewModel/Source/ModelViewViewModelBlueprint/Private/MVVMBlueprintViewModelContext.cpp b/Engine/Plugins/Runtime/ModelViewViewModel/Source/ModelViewViewModelBlueprint/Private/MVVMBlueprintViewModelContext.cpp
--- a/Engine/Plugins/Runtime/ModelViewViewModel/Source/ModelViewViewModelBlueprint/Private/MVVMBlueprintViewModelContext.cpp
+++ b/Engine/Plugins/Runtime/ModelViewViewModel/Source/ModelViewViewModelBlueprint/Private/MVVMBlueprintViewModelContext.cpp
@@ -37,7 +37,7 @@ namespace Private
 static const FName NAME_MVVMAllowedContextCreationType = TEXT("MVVMAllowedContextCreationType");
 static const FName NAME_MVVMDisallowedContextCreationType = TEXT("MVVMDisallowedContextCreationType");
 
-// Find the Allowed and Disallowed from first calss in the hierarchy that defines it.
+// Find the Allowed and Disallowed from first class in the hierarchy that defines it.
 void GenerateAllowedList(const UClass* Class, TArray<FString>& AllowedStrings, TArray<FString>& DisallowedStrings)
 {
 	if (AllowedStrings.Num() == 0 && Class->HasMetaData(NAME_MVVMAllowedContextCreationType))
diff --git a/Engine/Plugins/Runtime/ModelViewViewModel/Source/ModelViewViewModelBlueprint/Private/MVVMViewBlueprintCompiler.cpp b/Engine/Plugins/Runtime/ModelViewViewModel/Source/ModelViewViewModelBlueprint/Private/MVVMViewBlueprintCompiler.cpp
--- a/Engine/Plugins/Runtime/ModelViewViewModel/Source/ModelViewViewModelBlueprint/Private/MVVMViewBlueprintCompiler.cpp
+++ b/Engine/Plugins/Runtime/ModelViewViewModel/Source/ModelViewViewModelBlueprint/Private/MVVMViewBlueprintCompiler.cpp
@@ -675,6 +675,8 @@ void FMVVMViewBlueprintCompiler::CreateRequiredProperties(const FWidgetBlueprint
 		}
 
 		{
+			bool bIsPublicReadable = ViewModelContext.InstancedViewModel == nullptr && ViewModelContext.bCreateGetterFunction;
+
 			FCompilerUserWidgetProperty& SourceVariable = NeededUserWidgetProperties.AddDefaulted_GetRef();
 			SourceVariable.AuthoritativeClass = ViewModelContext.GetViewModelClass();
 			SourceVariable.Name = PropertyName;
@@ -683,7 +685,7 @@ void FMVVMViewBlueprintCompiler::CreateRequiredProperties(const FWidgetBlueprint
 			SourceVariable.BlueprintSetter = SetterFunctionName;
 			SourceVariable.CreationType = FCompilerUserWidgetProperty::ECreationType::CreateOnlyIfDoesntExist;
 			SourceVariable.bExposeOnSpawn = bCreateSetterFunction;
-			SourceVariable.bPrivate = ViewModelContext.InstancedViewModel != nullptr;
+			SourceVariable.bPrivate = !bIsPublicReadable;
 			SourceVariable.bReadOnly = !bCreateSetterFunction;
 		}
 
diff --git a/Engine/Plugins/Runtime/ModelViewViewModel/Source/ModelViewViewModelBlueprint/Public/MVVMBlueprintViewModelContext.h b/Engine/Plugins/Runtime/ModelViewViewModel/Source/ModelViewViewModelBlueprint/Public/MVVMBlueprintViewModelContext.h
--- a/Engine/Plugins/Runtime/ModelViewViewModel/Source/ModelViewViewModelBlueprint/Public/MVVMBlueprintViewModelContext.h
+++ b/Engine/Plugins/Runtime/ModelViewViewModel/Source/ModelViewViewModelBlueprint/Public/MVVMBlueprintViewModelContext.h
@@ -80,6 +80,15 @@ public:
 			{
 				NotifyFieldValueClass = ViewModelClass_DEPRECATED.Get();
 			}
+			if (!bCreateSetterFunction_Deprecation)
+			{
+				bCreateSetterFunction_Deprecation = true;
+				if (CreationType == EMVVMBlueprintViewModelContextCreationType::Manual)
+				{
+					bOptional = true;
+					bCreateSetterFunction = true;
+				}
+			}
 		}
 	}
 
@@ -113,17 +122,17 @@ public:
 	EMVVMBlueprintViewModelContextCreationType CreationType = EMVVMBlueprintViewModelContextCreationType::CreateInstance;
 
 	/** Identifier of an already registered viewmodel. */
-	UPROPERTY(EditAnywhere, Category = "Viewmodel", AdvancedDisplay, meta = (DisplayName = "Global Viewmodel Identifier"))
+	UPROPERTY(EditAnywhere, Category = "Viewmodel", meta = (DisplayName = "Global Viewmodel Identifier"))
 	FName GlobalViewModelIdentifier;
 
 	/** The Path to get the viewmodel instance. */
-	UPROPERTY(EditAnywhere, Category = "Viewmodel", AdvancedDisplay, meta = (DisplayName = "Viewmodel Property Path"))
+	UPROPERTY(EditAnywhere, Category = "Viewmodel", meta = (DisplayName = "Viewmodel Property Path"))
 	FString ViewModelPropertyPath;
 
-	UPROPERTY(EditAnywhere, Category = "Viewmodel", Instanced, AdvancedDisplay, meta = (EditInline))
+	UPROPERTY(EditAnywhere, Category = "Viewmodel", Instanced, meta = (EditInline))
 	TObjectPtr<UMVVMViewModelContextResolver> Resolver = nullptr;
 
-	UPROPERTY(EditAnywhere, Category = "Viewmodel", Instanced, AdvancedDisplay, NoClear, meta = (ShowOnlyInnerProperties))
+	UPROPERTY(EditAnywhere, Category = "Viewmodel", Instanced, NoClear, meta = (ShowOnlyInnerProperties))
 	TObjectPtr<UMVVMBlueprintInstancedViewModelBase> InstancedViewModel;
 
 	/**
@@ -133,6 +142,13 @@ public:
 	UPROPERTY(EditAnywhere, Category = "Viewmodel", AdvancedDisplay, meta = (DisplanName="Create Public Setter"))
 	bool bCreateSetterFunction = false;
 
+	/**
+	 * Generate a public getter for this viewmodel.
+	 * @note Always false when using a Instanced viewmodel.
+	 */
+	UPROPERTY(EditAnywhere, Category = "Viewmodel", AdvancedDisplay, meta = (DisplanName = "Create Public Getter"))
+	bool bCreateGetterFunction = true;
+
 	/**
 	 * Optional. Will not warn if the instance is not set or found.
 	 * @note Always true when the Creation Type is Manual.
@@ -151,6 +167,10 @@ public:
 	/** Can remove the viewmodel in the editor. */
 	UPROPERTY()
 	bool bCanRemove = true;
+
+private:
+	UPROPERTY()
+	bool bCreateSetterFunction_Deprecation = false;
 };
 
 template<>
diff --git a/Engine/Plugins/Runtime/ModelViewViewModel/Source/ModelViewViewModelEditor/Private/Customizations/MVVMBlueprintViewModelContextCustomization.cpp b/Engine/Plugins/Runtime/ModelViewViewModel/Source/ModelViewViewModelEditor/Private/Customizations/MVVMBlueprintViewModelContextCustomization.cpp
--- a/Engine/Plugins/Runtime/ModelViewViewModel/Source/ModelViewViewModelEditor/Private/Customizations/MVVMBlueprintViewModelContextCustomization.cpp
+++ b/Engine/Plugins/Runtime/ModelViewViewModel/Source/ModelViewViewModelEditor/Private/Customizations/MVVMBlueprintViewModelContextCustomization.cpp
@@ -121,6 +121,20 @@ TSharedRef<SWidget> FViewModelPropertyAccessEditor::MakePropertyBindingWidget(TS
 /**
  * 
  */
+namespace Private
+{
+	FMVVMBlueprintViewModelContext* GetViewModelContext(TSharedRef<IPropertyHandle> PropertyHandle)
+	{
+		ensure(CastField<FStructProperty>(PropertyHandle->GetProperty()) && CastField<FStructProperty>(PropertyHandle->GetProperty())->Struct == FMVVMBlueprintViewModelContext::StaticStruct());
+		void* Buffer = nullptr;
+		if (PropertyHandle->GetValueData(Buffer) == FPropertyAccess::Success)
+		{
+			return reinterpret_cast<FMVVMBlueprintViewModelContext*>(Buffer);
+		}
+		return nullptr;
+	}
+}
+
 FBlueprintViewModelContextDetailCustomization::FBlueprintViewModelContextDetailCustomization(TWeakPtr<FWidgetBlueprintEditor> InEditor)
 	: WidgetBlueprintEditor(InEditor)
 {}
@@ -131,16 +145,8 @@ void FBlueprintViewModelContextDetailCustomization::CustomizeChildren(TSharedRef
 	uint32 NumChildren = 0;
 	PropertyHandle->GetNumChildren(NumChildren);
 
-	FMVVMBlueprintViewModelContext* ContextPtr = nullptr;
-	{
-		ensure(CastField<FStructProperty>(PropertyHandle->GetProperty()) && CastField<FStructProperty>(PropertyHandle->GetProperty())->Struct == FMVVMBlueprintViewModelContext::StaticStruct());
-		void* Buffer = nullptr;
-		if (PropertyHandle->GetValueData(Buffer) == FPropertyAccess::Success)
-		{
-			ContextPtr = reinterpret_cast<FMVVMBlueprintViewModelContext*>(Buffer);
-		}
-	}
-
+	ContextHandle = PropertyHandle;
+	FMVVMBlueprintViewModelContext* ContextPtr = Private::GetViewModelContext(PropertyHandle);
 	if (ContextPtr == nullptr)
 	{
 		return;
@@ -171,6 +177,11 @@ void FBlueprintViewModelContextDetailCustomization::CustomizeChildren(TSharedRef
 		NotifyFieldValueClassHandle->SetOnPropertyValueChanged(FSimpleDelegate::CreateSP(this, &FBlueprintViewModelContextDetailCustomization::HandleClassChanged));
 	}
 
+	if (ensure(CreationTypeHandle))
+	{
+		CreationTypeHandle->SetOnPropertyValueChanged(FSimpleDelegate::CreateSP(this, &FBlueprintViewModelContextDetailCustomization::HandleCreationTypeChanged));
+	}
+
 	if (ensure(ViewModelNameHandle))
 	{
 		IDetailPropertyRow& PropertyRow = ChildBuilder.AddProperty(ViewModelNameHandle.ToSharedRef())
@@ -323,26 +334,31 @@ void FBlueprintViewModelContextDetailCustomization::CustomizeChildren(TSharedRef
 			if (ensure(CreateSetterFunctionHandle))
 			{
 				ChildBuilder.AddProperty(CreateSetterFunctionHandle.ToSharedRef())
-					.IsEnabled(bCanEdit)
-					.Visibility(MakeAttributeLambda([ContextPtr]()
+					.IsEnabled(MakeAttributeLambda([bCanEdit, ContextPtr]()
 						{
 							bool bResult = ContextPtr->CreationType != EMVVMBlueprintViewModelContextCreationType::Manual;
-							return bResult ? EVisibility::Visible : EVisibility::Collapsed;
+							return bResult && bCanEdit;
 						}));
 			}
 		}
 
+		TSharedPtr<IPropertyHandle> CreateGetterFunctionHandle = PropertyHandle->GetChildHandle(GET_MEMBER_NAME_CHECKED(FMVVMBlueprintViewModelContext, bCreateGetterFunction), false);
+		if (ensure(CreateGetterFunctionHandle))
+		{
+			ChildBuilder.AddProperty(CreateGetterFunctionHandle.ToSharedRef())
+				.IsEnabled(bCanEdit);
+		}
+
 		TSharedPtr<IPropertyHandle> OptionalHandle = PropertyHandle->GetChildHandle(GET_MEMBER_NAME_CHECKED(FMVVMBlueprintViewModelContext, bOptional), false);
 		if (ensure(OptionalHandle))
 		{
 			ChildBuilder.AddProperty(OptionalHandle.ToSharedRef())
-				.IsEnabled(bCanEdit)
-				.Visibility(MakeAttributeLambda([ContextPtr]()
+				.IsEnabled(MakeAttributeLambda([bCanEdit, ContextPtr]()
 					{
 						bool bResult = ContextPtr->CreationType == EMVVMBlueprintViewModelContextCreationType::GlobalViewModelCollection
 						 || ContextPtr->CreationType == EMVVMBlueprintViewModelContextCreationType::PropertyPath
 						 || ContextPtr->CreationType == EMVVMBlueprintViewModelContextCreationType::Resolver;
-						return bResult ? EVisibility::Visible : EVisibility::Collapsed;
+						return bResult && bCanEdit;
 					}));
 		}
 	}
@@ -395,6 +411,27 @@ FText FBlueprintViewModelContextDetailCustomization::GetClassName() const
 	return LOCTEXT("None", "None");
 }
 
+void FBlueprintViewModelContextDetailCustomization::HandleCreationTypeChanged()
+{
+	uint8 NewValue = 0;
+	if (CreationTypeHandle->GetValue(NewValue) == FPropertyAccess::Success)
+	{
+		if (FMVVMBlueprintViewModelContext* ContextPtr = Private::GetViewModelContext(ContextHandle.ToSharedRef()))
+		{
+			if ((EMVVMBlueprintViewModelContextCreationType)NewValue == EMVVMBlueprintViewModelContextCreationType::Manual)
+			{
+				ContextPtr->bOptional = true;
+				ContextPtr->bCreateSetterFunction = true;
+			}
+			else
+			{
+				ContextPtr->bOptional = false;
+				ContextPtr->bCreateSetterFunction = false;
+			}
+		}
+	}
+}
+
 TSharedRef<SWidget> FBlueprintViewModelContextDetailCustomization::CreateExecutionTypeMenuContent()
 {
 	const bool bCloseAfterSelection = true;
diff --git a/Engine/Plugins/Runtime/ModelViewViewModel/Source/ModelViewViewModelEditor/Private/Customizations/MVVMBlueprintViewModelContextCustomization.h b/Engine/Plugins/Runtime/ModelViewViewModel/Source/ModelViewViewModelEditor/Private/Customizations/MVVMBlueprintViewModelContextCustomization.h
--- a/Engine/Plugins/Runtime/ModelViewViewModel/Source/ModelViewViewModelEditor/Private/Customizations/MVVMBlueprintViewModelContextCustomization.h
+++ b/Engine/Plugins/Runtime/ModelViewViewModel/Source/ModelViewViewModelEditor/Private/Customizations/MVVMBlueprintViewModelContextCustomization.h
@@ -53,6 +53,7 @@ private:
 	FText GetExecutionTypeValueToolTip() const;
 	FText GetClassName() const;
 	void HandleClassChanged();
+	void HandleCreationTypeChanged();
 	TSharedRef<SWidget> HandleClassGetMenuContent();
 	void HandleClassCancelMenu();
 	void HandleClassCommitted(const UClass* SelectedClass);
@@ -63,6 +64,7 @@ private:
 	FViewModelPropertyAccessEditor PropertyAccessEditor;
 	TWeakPtr<FWidgetBlueprintEditor> WidgetBlueprintEditor;
 	TSharedPtr<SComboButton> NotifyFieldValueClassComboButton;
+	TSharedPtr<IPropertyHandle> ContextHandle;
 	TSharedPtr<IPropertyHandle> NotifyFieldValueClassHandle;
 	TSharedPtr<IPropertyHandle> PropertyPathHandle;
 	TSharedPtr<IPropertyHandle> CreationTypeHandle;
