MVVM: Rework for the view runtime data. The new format allows for better error detection and new features in progress. The view owns the bindings and events. The source has a list of the bindings. That saves runtime of memory, the delegate doesn't store the binding id. All bindings, events, sources are now sorted for better incremental build performance.
Add a different "key" structures to prevent mistakes with indexes.
Unregister the delay bindings when the source is released.
Store the FieldId, in the source itself. That reduces the amount of work when loading the view.
Add option to delay the events initialization (like we do for bindings).
The binding to evaluate "long path" view models are now separated from regular bindings.

diff --git a/Engine/Plugins/Runtime/ModelViewViewModel/Source/ModelViewViewModel/Private/Bindings/MVVMCompiledBindingLibrary.cpp b/Engine/Plugins/Runtime/ModelViewViewModel/Source/ModelViewViewModel/Private/Bindings/MVVMCompiledBindingLibrary.cpp
--- a/Engine/Plugins/Runtime/ModelViewViewModel/Source/ModelViewViewModel/Private/Bindings/MVVMCompiledBindingLibrary.cpp
+++ b/Engine/Plugins/Runtime/ModelViewViewModel/Source/ModelViewViewModel/Private/Bindings/MVVMCompiledBindingLibrary.cpp
@@ -42,17 +42,6 @@ UFunction* FMVVMVCompiledFields::GetFunction(FName FunctionName) const
 }
 
 
-UE::FieldNotification::FFieldId FMVVMVCompiledFields::GetFieldId(FName FieldName) const
-{
-	// ScriptStruct can't implements UNotifyFieldValueChanged
-	const UClass* Class = CastChecked<UClass>(ClassOrScriptStruct);
-	check(Class->ImplementsInterface(UNotifyFieldValueChanged::StaticClass()));
-
-	TScriptInterface<INotifyFieldValueChanged> Interface = Class->GetDefaultObject();
-	return Interface->GetFieldNotificationDescriptor().GetField(Class, FieldName);
-}
-
-
 /**
  *
  */
@@ -131,7 +120,6 @@ void FMVVMCompiledBindingLibrary::Load()
 
 	ensureAlwaysMsgf(LoadedProperties.Num() == 0, TEXT("The binding library was loaded more than once."));
 	ensureAlwaysMsgf(LoadedFunctions.Num() == 0, TEXT("The binding library was loaded more than once."));
-	ensureAlwaysMsgf(LoadedFieldIds.Num() == 0, TEXT("The binding library was loaded more than once."));
 
 	for (const FMVVMVCompiledFields& Field : CompiledFields)
 	{
@@ -161,23 +149,13 @@ void FMVVMCompiledBindingLibrary::Load()
 				LoadedFunctions.Emplace(LoadedFunction); // add it even if none to keep the index valid
 			}
 		}
-		{
-			const int32 NumberOfFieldIds = Field.GetFieldIdNum();
-			for (int32 Index = 0; Index < NumberOfFieldIds; ++Index)
-			{
-				FName FieldName = Field.GetFieldIdName(CompiledFieldNames, Index);
-				UE::FieldNotification::FFieldId LoadedFieldId = Field.GetFieldId(FieldName);
-				LoadedFieldIds.Add(LoadedFieldId); // add it even if none to keep the index valid
-				ensureAlwaysMsgf(LoadedFieldId.IsValid(), TEXT("The field id '%s:%s' could not be loaded."), (Field.GetStruct() ? *Field.GetStruct()->GetName() : TEXT("None")), *FieldName.ToString());
-			}
-		}
 	}
 }
 
 
 bool FMVVMCompiledBindingLibrary::IsLoaded() const
 {
-	return LoadedProperties.Num() > 0 || LoadedFunctions.Num() > 0 || LoadedFieldIds.Num() > 0 || CompiledFields.Num() == 0;
+	return LoadedProperties.Num() > 0 || LoadedFunctions.Num() > 0 || CompiledFields.Num() == 0;
 }
 
 
@@ -185,7 +163,6 @@ void FMVVMCompiledBindingLibrary::Unload()
 {
 	LoadedProperties.Empty();
 	LoadedFunctions.Empty();
-	LoadedFieldIds.Empty();
 }
 
 
@@ -595,23 +572,4 @@ TValueOrError<FString, FString> FMVVMCompiledBindingLibrary::FieldPathToString(F
 	}
 }
 
-
-TValueOrError<UE::FieldNotification::FFieldId, void> FMVVMCompiledBindingLibrary::GetFieldId(const FMVVMVCompiledFieldId& InFieldId) const
-{
-#if WITH_EDITORONLY_DATA
-	const bool bIsValidBinding = InFieldId.CompiledBindingLibraryId == CompiledBindingLibraryId;
-	if (!bIsValidBinding)
-	{
-		ensureAlwaysMsgf(false, TEXT("The binding is from a different library."));
-		return MakeError<>();
-	}
-#endif
-
-	if (LoadedFieldIds.IsValidIndex(InFieldId.FieldIdIndex))
-	{
-		return MakeValue(LoadedFieldIds[InFieldId.FieldIdIndex]);
-	}
-	return MakeError<>();
-}
-
 #undef LOCTEXT_NAMESPACE
diff --git a/Engine/Plugins/Runtime/ModelViewViewModel/Source/ModelViewViewModel/Private/Debugging/MVVMDebugging.cpp b/Engine/Plugins/Runtime/ModelViewViewModel/Source/ModelViewViewModel/Private/Debugging/MVVMDebugging.cpp
--- a/Engine/Plugins/Runtime/ModelViewViewModel/Source/ModelViewViewModel/Private/Debugging/MVVMDebugging.cpp
+++ b/Engine/Plugins/Runtime/ModelViewViewModel/Source/ModelViewViewModel/Private/Debugging/MVVMDebugging.cpp
@@ -12,7 +12,7 @@ namespace UE::MVVM
 {
 /** */
 FDebugging::FView::FView(const UMVVMView* InView)
-	: UserWidget(InView->GetUserWidget())
+	: UserWidget(InView->GetOuterUUserWidget())
 	, View(InView)
 {
 }
@@ -40,49 +40,26 @@ void FDebugging::BroadcastViewBeginDestruction(const UMVVMView* InView)
 }
 
 /** */
-FDebugging::FViewSourceValueArgs::FViewSourceValueArgs(const FMVVMViewSource& InViewSource)
-	: ViewSource(InViewSource)
+FDebugging::FViewSourceValueArgs::FViewSourceValueArgs(const FMVVMViewClass_SourceKey& InClass, const FMVVMView_SourceKey& InView)
+	: ClassSource(InClass)
+	, ViewSource(InView)
 {
 }
 
 FDebugging::FViewSourceValueChanged FDebugging::OnViewSourceValueChanged;
 
-void FDebugging::BroadcastViewSourceValueChanged(const UMVVMView* InView, const FMVVMViewSource& ViewSource)
+void FDebugging::BroadcastViewSourceValueChanged(const UMVVMView* InView, const FMVVMViewClass_SourceKey ClassSourceKey, const FMVVMView_SourceKey ViewSourceKey)
 {
-	OnViewSourceValueChanged.Broadcast(InView, FViewSourceValueArgs(ViewSource));
+	OnViewSourceValueChanged.Broadcast(InView, FViewSourceValueArgs(ClassSourceKey, ViewSourceKey));
 }
 
 /** */
-FDebugging::FLibraryBindingRegisteredArgs::FLibraryBindingRegisteredArgs(const FMVVMViewClass_CompiledBinding& InBinding, ERegisterLibraryBindingResult InResult)
+FDebugging::FLibraryBindingExecutedArgs::FLibraryBindingExecutedArgs(FMVVMViewClass_BindingKey InBinding)
 	: Binding(InBinding)
-	, Result(InResult)
 {
 }
 
-FDebugging::FLibraryBindingRegistered FDebugging::OnLibraryBindingRegistered;
-
-void FDebugging::BroadcastLibraryBindingRegistered(const UMVVMView* InView, const FLibraryBindingRegisteredArgs& Args)
-{
-	check(InView);
-	OnLibraryBindingRegistered.Broadcast(FView(InView), Args);
-}
-
-void FDebugging::BroadcastLibraryBindingRegistered(const UMVVMView* InView, const FMVVMViewClass_CompiledBinding& InBinding, ERegisterLibraryBindingResult InResult)
-{
-	if (OnLibraryBindingRegistered.IsBound())
-	{
-		check(InView);
-		OnLibraryBindingRegistered.Broadcast(FView(InView), FLibraryBindingRegisteredArgs(InBinding, InResult));
-	}
-}
-
-/** */
-FDebugging::FLibraryBindingExecutedArgs::FLibraryBindingExecutedArgs(const FMVVMViewClass_CompiledBinding& InBinding)
-	: Binding(InBinding)
-{
-}
-
-FDebugging::FLibraryBindingExecutedArgs::FLibraryBindingExecutedArgs(const FMVVMViewClass_CompiledBinding& InBinding, FMVVMCompiledBindingLibrary::EExecutionFailingReason InResult)
+FDebugging::FLibraryBindingExecutedArgs::FLibraryBindingExecutedArgs(FMVVMViewClass_BindingKey InBinding, FMVVMCompiledBindingLibrary::EExecutionFailingReason InResult)
 	: Binding(InBinding)
 	, FailingReason(InResult)
 {
@@ -96,7 +73,7 @@ void FDebugging::BroadcastLibraryBindingExecuted(const UMVVMView* InView, const
 	OnLibraryBindingExecuted.Broadcast(FView(InView), Args);
 }
 
-void FDebugging::BroadcastLibraryBindingExecuted(const UMVVMView* InView, const FMVVMViewClass_CompiledBinding& InBinding)
+void FDebugging::BroadcastLibraryBindingExecuted(const UMVVMView* InView, FMVVMViewClass_BindingKey InBinding)
 {
 	if (OnLibraryBindingExecuted.IsBound())
 	{
@@ -105,7 +82,7 @@ void FDebugging::BroadcastLibraryBindingExecuted(const UMVVMView* InView, const
 	}
 }
 
-void FDebugging::BroadcastLibraryBindingExecuted(const UMVVMView* InView, const FMVVMViewClass_CompiledBinding& InBinding, FMVVMCompiledBindingLibrary::EExecutionFailingReason InResult)
+void FDebugging::BroadcastLibraryBindingExecuted(const UMVVMView* InView, FMVVMViewClass_BindingKey InBinding, FMVVMCompiledBindingLibrary::EExecutionFailingReason InResult)
 {
 	if (OnLibraryBindingExecuted.IsBound())
 	{
diff --git a/Engine/Plugins/Runtime/ModelViewViewModel/Source/ModelViewViewModel/Private/View/MVVMBindingSubsystem.cpp b/Engine/Plugins/Runtime/ModelViewViewModel/Source/ModelViewViewModel/Private/View/MVVMBindingSubsystem.cpp
--- a/Engine/Plugins/Runtime/ModelViewViewModel/Source/ModelViewViewModel/Private/View/MVVMBindingSubsystem.cpp
+++ b/Engine/Plugins/Runtime/ModelViewViewModel/Source/ModelViewViewModel/Private/View/MVVMBindingSubsystem.cpp
@@ -36,13 +36,13 @@ namespace UE::MVVM::Private
 {
 struct FViewAndBinding
 {
-	FViewAndBinding(const TObjectKey<const UMVVMView>& InView, FMVVMViewDelayedBinding InBinding)
+	FViewAndBinding(const TObjectKey<const UMVVMView>& InView, FMVVMViewClass_BindingKey InBinding)
 		: View(InView)
 		, Binding(InBinding)
 	{
 	}
 	TObjectKey<const UMVVMView> View;
-	FMVVMViewDelayedBinding Binding;
+	FMVVMViewClass_BindingKey Binding;
 
 	bool operator== (const FViewAndBinding& Other) const
 	{
@@ -52,7 +52,7 @@ struct FViewAndBinding
 	friend uint32 GetTypeHash(const FViewAndBinding& Key)
 	{
 		uint32 Value1 = GetTypeHash(Key.View);
-		uint32 Value2 = GetTypeHash(Key.Binding.GetCompiledBindingIndex());
+		uint32 Value2 = GetTypeHash(Key.Binding.GetIndex());
 		return HashCombine(Value1, Value2);
 	}
 };
@@ -62,14 +62,15 @@ void UMVVMBindingSubsystem::HandlePreTick(float DeltaTime)
 {
 	SCOPE_CYCLE_COUNTER(STAT_MVVMBindingTick);
 
-	if (EveryTickBindings.Num() > 0)
+	if (ViewsWithTickBindings.Num() > 0)
 	{
 		FMemMark Mark(FMemStack::Get());
 		TArray<const UMVVMView*, TMemStackAllocator<>> ToTick;
-		ToTick.Reserve(EveryTickBindings.Num());
-		for (TWeakObjectPtr<const UMVVMView> View : EveryTickBindings)
+		ToTick.Reserve(ViewsWithTickBindings.Num());
+		for (TWeakObjectPtr<const UMVVMView> View : ViewsWithTickBindings)
 		{
-			if (const UMVVMView* ViewPtr = View.Get())
+			const UMVVMView* ViewPtr = View.Get();
+			if (ensure(ViewPtr))
 			{
 				ToTick.Add(ViewPtr);
 			}
@@ -77,7 +78,7 @@ void UMVVMBindingSubsystem::HandlePreTick(float DeltaTime)
 
 		for (const UMVVMView* ViewPtr : ToTick)
 		{
-			ViewPtr->ExecuteEveryTickBindings();
+			ViewPtr->ExecuteTickBindings();
 		}
 	}
 
@@ -86,8 +87,8 @@ void UMVVMBindingSubsystem::HandlePreTick(float DeltaTime)
 		TSet<UE::MVVM::Private::FViewAndBinding> AllDelayedBindingsExecutedThisFrame;
 		AllDelayedBindingsExecutedThisFrame.Reserve(DelayedBindings.Num());
 
-		FDelayedMap DelayedBindingsWhileTicking = MoveTemp(DelayedBindings);
-		DelayedBindings = FDelayedMap();
+		FDelayedBindingMap DelayedBindingsWhileTicking = MoveTemp(DelayedBindings);
+		DelayedBindings = FDelayedBindingMap();
 
 		do 
 		{
@@ -96,7 +97,7 @@ void UMVVMBindingSubsystem::HandlePreTick(float DeltaTime)
 				if (const UMVVMView* View = DelayedBindingsPair.Key.ResolveObjectPtr())
 				{
 					ensure(DelayedBindingsPair.Value.Num() > 0);
-					for (const FMVVMViewDelayedBinding& DelayedBinding : DelayedBindingsPair.Value)
+					for (const FMVVMViewClass_BindingKey& DelayedBinding : DelayedBindingsPair.Value)
 					{
 						View->ExecuteDelayedBinding(DelayedBinding);
 						UE::MVVM::Private::FViewAndBinding ViewAndBinding = UE::MVVM::Private::FViewAndBinding(DelayedBindingsPair.Key, DelayedBinding);
@@ -115,7 +116,7 @@ void UMVVMBindingSubsystem::HandlePreTick(float DeltaTime)
 				for (int32 DelayIndex = DelayedBindingItt.Value().Num() - 1; DelayIndex >= 0; --DelayIndex)
 				{
 					// Was it executed this frame
-					const FMVVMViewDelayedBinding& DelayedBinding = DelayedBindingItt.Value()[DelayIndex];
+					const FMVVMViewClass_BindingKey& DelayedBinding = DelayedBindingItt.Value()[DelayIndex];
 					UE::MVVM::Private::FViewAndBinding ViewAndBinding = UE::MVVM::Private::FViewAndBinding(DelayedBindingItt.Key(), DelayedBinding);
 					if (!AllDelayedBindingsExecutedThisFrame.Find(ViewAndBinding))
 					{
@@ -141,21 +142,44 @@ void UMVVMBindingSubsystem::HandlePreTick(float DeltaTime)
 	}
 }
 
-void UMVVMBindingSubsystem::AddViewWithEveryTickBinding(const UMVVMView* InView)
+void UMVVMBindingSubsystem::AddViewWithTickBinding(const UMVVMView* InView)
 {
-	check(!EveryTickBindings.Contains(InView));
+	check(!ViewsWithTickBindings.Contains(InView));
 	ensureMsgf(FSlateApplication::IsInitialized(), TEXT("The Slate Application is not initialized. This is probably because you are running a server. The Delayed and Tick binding will not execute."));
 
-	EveryTickBindings.Add(InView);
+	ViewsWithTickBindings.Add(InView);
 }
 
-void UMVVMBindingSubsystem::RemoveViewWithEveryTickBinding(const UMVVMView* InView)
+void UMVVMBindingSubsystem::RemoveViewWithTickBinding(const UMVVMView* InView)
 {
-	EveryTickBindings.RemoveSingleSwap(InView);
+	ViewsWithTickBindings.RemoveSingleSwap(InView);
 }
 
-void UMVVMBindingSubsystem::AddDelayedBinding(const UMVVMView* View, FMVVMViewDelayedBinding InCompiledBinding)
+void UMVVMBindingSubsystem::AddDelayedBinding(const UMVVMView* View, FMVVMViewClass_BindingKey InCompiledBinding)
 {
 	ensureMsgf(FSlateApplication::IsInitialized(), TEXT("The Slate Application is not initialized. This is probably because you are running a server. The Delayed and Tick binding will not execute."));
 	DelayedBindings.FindOrAdd(View).AddUnique(InCompiledBinding);
 }
+
+void UMVVMBindingSubsystem::RemoveDelayedBindings(const UMVVMView* View)
+{
+	ensureMsgf(FSlateApplication::IsInitialized(), TEXT("The Slate Application is not initialized. This is probably because you are running a server. The Delayed and Tick binding will not execute."));
+	DelayedBindings.Remove(View);
+}
+
+void UMVVMBindingSubsystem::RemoveDelayedBindings(const UMVVMView* View, FMVVMViewClass_SourceKey SourceKey)
+{
+	ensureMsgf(FSlateApplication::IsInitialized(), TEXT("The Slate Application is not initialized. This is probably because you are running a server. The Delayed and Tick binding will not execute."));
+	if (FDelayedBindingList* FoundView = DelayedBindings.Find(View))
+	{
+		const UMVVMViewClass* ViewClass = View->GetViewClass();
+		for (int32 Index = FoundView->Num() - 1; Index >= 0; --Index)
+		{
+			const FMVVMViewClass_Binding& Binding = ViewClass->GetBinding((*FoundView)[Index]);
+			if ((Binding.GetSources() & SourceKey.GetBit()) != 0)
+			{
+				(*FoundView).RemoveAtSwap(Index);
+			}
+		}
+	}
+}
diff --git a/Engine/Plugins/Runtime/ModelViewViewModel/Source/ModelViewViewModel/Private/View/MVVMBindingSubsystem.h b/Engine/Plugins/Runtime/ModelViewViewModel/Source/ModelViewViewModel/Private/View/MVVMBindingSubsystem.h
--- a/Engine/Plugins/Runtime/ModelViewViewModel/Source/ModelViewViewModel/Private/View/MVVMBindingSubsystem.h
+++ b/Engine/Plugins/Runtime/ModelViewViewModel/Source/ModelViewViewModel/Private/View/MVVMBindingSubsystem.h
@@ -4,38 +4,12 @@
 
 #include "Subsystems/EngineSubsystem.h"
 #include "UObject/ObjectKey.h"
+#include "View/MVVMViewClass.h"
 
 #include "MVVMBindingSubsystem.generated.h"
 
 class UMVVMView;
 
-/** */
-struct FMVVMViewDelayedBinding
-{
-	FMVVMViewDelayedBinding(int32 InCompiledBindingIndex)
-		: CompiledBindingIndex(InCompiledBindingIndex)
-	{
-	}
-
-	int32 GetCompiledBindingIndex() const
-	{
-		return CompiledBindingIndex;
-	}
-
-	bool operator==(const FMVVMViewDelayedBinding& Other) const
-	{
-		return CompiledBindingIndex == Other.CompiledBindingIndex;
-	}
-
-	bool operator!=(const FMVVMViewDelayedBinding& Other) const
-	{
-		return CompiledBindingIndex != Other.CompiledBindingIndex;
-	}
-
-private:
-	int32 CompiledBindingIndex = INDEX_NONE;
-};
-
 /** */
 UCLASS(NotBlueprintable, Hidden)
 class UMVVMBindingSubsystem : public UEngineSubsystem
@@ -46,16 +20,18 @@ public:
 	virtual void Initialize(FSubsystemCollectionBase& Collection) override;
 	virtual void Deinitialize() override;
 
-	void AddViewWithEveryTickBinding(const UMVVMView* View);
-	void RemoveViewWithEveryTickBinding(const UMVVMView* View);
+	void AddViewWithTickBinding(const UMVVMView* View);
+	void RemoveViewWithTickBinding(const UMVVMView* View);
 
-	void AddDelayedBinding(const UMVVMView* View, FMVVMViewDelayedBinding CompiledBinding);
+	void AddDelayedBinding(const UMVVMView* View, FMVVMViewClass_BindingKey BindingKey);
+	void RemoveDelayedBindings(const UMVVMView* View);
+	void RemoveDelayedBindings(const UMVVMView* View, FMVVMViewClass_SourceKey SourceKey);
 
 private:
 	void HandlePreTick(float DeltaTIme);
 
-	using FDelayedBindingList = TArray<FMVVMViewDelayedBinding, TInlineAllocator<8>>;
-	using FDelayedMap = TMap<TObjectKey<const UMVVMView>, FDelayedBindingList>;
-	FDelayedMap DelayedBindings;
-	TArray<TWeakObjectPtr<const UMVVMView>> EveryTickBindings;
+	using FDelayedBindingList = TArray<FMVVMViewClass_BindingKey, TInlineAllocator<8>>;
+	using FDelayedBindingMap = TMap<TObjectKey<const UMVVMView>, FDelayedBindingList>;
+	FDelayedBindingMap DelayedBindings;
+	TArray<TWeakObjectPtr<const UMVVMView>> ViewsWithTickBindings;
 };
diff --git a/Engine/Plugins/Runtime/ModelViewViewModel/Source/ModelViewViewModel/Private/View/MVVMView.cpp b/Engine/Plugins/Runtime/ModelViewViewModel/Source/ModelViewViewModel/Private/View/MVVMView.cpp
--- a/Engine/Plugins/Runtime/ModelViewViewModel/Source/ModelViewViewModel/Private/View/MVVMView.cpp
+++ b/Engine/Plugins/Runtime/ModelViewViewModel/Source/ModelViewViewModel/Private/View/MVVMView.cpp
@@ -36,13 +36,13 @@ void UMVVMView::ConstructView(const UMVVMViewClass* InClassExtension)
 	ensure(ClassExtension == nullptr);
 	ClassExtension = InClassExtension;
 
-	Sources.Reserve(ClassExtension->GetViewModelCreators().Num());
-	for (const FMVVMViewClass_SourceCreator& Item : ClassExtension->GetViewModelCreators())
+	check(Sources.Num() == 0);
+	int32 SourceNum = ClassExtension->GetSources().Num();
+	Sources.SetNum(SourceNum);
+	for (int32 Index = 0; Index < SourceNum; ++Index)
 	{
-		FMVVMViewSource& NewCreatedSource = Sources.AddDefaulted_GetRef();
-		NewCreatedSource.Source = nullptr;
-		NewCreatedSource.SourceName = Item.GetSourceName();
-		NewCreatedSource.bCreatedSource = true;
+		FMVVMView_Source& Source = Sources[Index];
+		Source.ClassKey = FMVVMViewClass_SourceKey(Index);
 	}
 }
 
@@ -52,11 +52,15 @@ void UMVVMView::Construct()
 	check(ClassExtension);
 	check(bConstructed == false);
 
-	if (ClassExtension->InitializeSourcesOnConstruct())
+	if (ClassExtension->DoesInitializeSourcesOnConstruct())
 	{
 		InitializeSources();
 	}
-	InitializeEvents();
+
+	if (ClassExtension->DoesInitializeEventsOnConstruct())
+	{
+		InitializeEvents();
+	}
 
 	bConstructed = true;
 
@@ -89,48 +93,31 @@ void UMVVMView::InitializeSources()
 
 	SCOPE_CYCLE_COUNTER(STAT_UMG_Viewmodel_InitializeSources);
 
+	check(bHasDefaultTickBinding == false && NumberOfSourceWithTickBinding == 0);
+
 	// Init Sources/ViewModel instances
-	UUserWidget* UserWidget = GetUserWidget();
-	const TArrayView<const FMVVMViewClass_SourceCreator> AllViewModelCreators = ClassExtension->GetViewModelCreators();
-	for (int32 Index = 0; Index < AllViewModelCreators.Num(); ++Index)
+	const int32 NumberOfSources = ClassExtension->GetSources().Num();
+	check(Sources.Num() == NumberOfSources);
+	check(NumberOfSources <= 64); // the max number of source the bitfield can hold.
+
+	for (int32 Index = 0; Index < NumberOfSources; ++Index)
 	{
-		const FMVVMViewClass_SourceCreator& Item = AllViewModelCreators[Index];
-		FMVVMViewSource& ViewSource = Sources[Index];
-		check(ViewSource.SourceName == Item.GetSourceName());
-		check(ViewSource.bCreatedSource == true);
+		InitializeSource(FMVVMView_SourceKey(Index));
 
-		if (!ViewSource.bSetManually)
+		// If DoesInitializeSourcesOnConstruct false but DoesInitializeBindingsOnConstruct, then we need to initialized the bindings when the source are initialized
+		if (ClassExtension->DoesInitializeBindingsOnConstruct())
 		{
-			UObject* NewSource = Item.CreateInstance(ClassExtension, this, UserWidget);
-			if (NewSource)
-			{
-				if (Item.IsSourceAUserWidgetProperty())
-				{
-					FObjectPropertyBase* FoundObjectProperty = FindFProperty<FObjectPropertyBase>(UserWidget->GetClass(), Item.GetSourceName());
-					if (ensureAlwaysMsgf(FoundObjectProperty, TEXT("The compiler should have added the property")))
-					{
-						if (ensure(NewSource->GetClass()->IsChildOf(FoundObjectProperty->PropertyClass)))
-						{
-							FoundObjectProperty->SetObjectPropertyValue_InContainer(UserWidget, NewSource);
-							ViewSource.bAssignedToUserWidgetProperty = true;
-						}
-					}
-				}
-			}
-			ViewSource.Source = NewSource;
-
-#if UE_WITH_MVVM_DEBUGGING
-			UE::MVVM::FDebugging::BroadcastViewSourceValueChanged(this, ViewSource);
-#endif
+			//note: Run the bindings. The binding from Source A can influence how the Source B is initialized/evaluated.
+			InitializeSourceBindings(FMVVMView_SourceKey(Index), false);
 		}
 	}
 
-	bSourcesInitialized = true;
-
-	if (ClassExtension->InitializeBindingsOnConstruct())
+	if (ClassExtension->DoesInitializeBindingsOnConstruct())
 	{
-		InitializeBindings();
+		InitializeSourceBindingsCommon();
 	}
+
+	bSourcesInitialized = true;
 }
 
 
@@ -143,118 +130,90 @@ void UMVVMView::UninitializeSources()
 
 	SCOPE_CYCLE_COUNTER(STAT_UMG_Viewmodel_UninitializeSources);
 
-	UninitializeInternal(true);
-}
+	if (bBindingsInitialized)
+	{
+		UninitializeBindings();
+	}
 
+	bSourcesInitialized = false;
 
-void UMVVMView::InitializeBindings()
-{
-	if (!bSourcesInitialized)
+	const int32 NumberOfSources = Sources.Num();
+	for (int32 Index = 0; Index < NumberOfSources; ++Index)
 	{
-		InitializeSources();
+		UninitializeSource(FMVVMView_SourceKey(Index));
 	}
+}
 
-	if (bBindingsInitialized)
+
+void UMVVMView::InitializeSource(FMVVMView_SourceKey SourceKey)
+{
+	FMVVMView_Source& ViewSource = Sources[SourceKey.GetIndex()];
+	const FMVVMViewClass_Source& ClassSource = ClassExtension->GetSource(ViewSource.ClassKey);
+
+	check(ViewSource.bSourceInitialized == false);
+
+	// Is it already set by something external before we had the change to set it.
+	if (!ViewSource.bSetManually)
 	{
-		return;
-	}
+		UUserWidget* UserWidget = GetUserWidget();
+		UObject* NewSource = ClassSource.GetOrCreateInstance(ClassExtension, this, UserWidget);
+		InitializeSourceInternal(NewSource, ViewSource.ClassKey, ClassSource, ViewSource);
 
-	SCOPE_CYCLE_COUNTER(STAT_UMG_Viewmodel_InitializeBindings);
+#if UE_WITH_MVVM_DEBUGGING
+		UE::MVVM::FDebugging::BroadcastViewSourceValueChanged(this, ViewSource.ClassKey, SourceKey);
+#endif
+	}
 
-	bHasEveryTickBinding = false;
+	ViewSource.bSourceInitialized = true;
+}
 
-	const TArrayView<const FMVVMViewClass_CompiledBinding> CompiledBindings = ClassExtension->GetCompiledBindings();
-	ensure(RegisteredLibraryBindings.IsEmpty());
-	RegisteredLibraryBindings.Reset(CompiledBindings.Num());
-	RegisteredLibraryBindings.AddDefaulted(CompiledBindings.Num());
 
-	if (CompiledBindings.Num())
+void UMVVMView::InitializeSourceInternal(UObject* NewSource, FMVVMViewClass_SourceKey SourceKey, const FMVVMViewClass_Source& ClassSource, FMVVMView_Source& ViewSource)
+{
+	if (NewSource)
 	{
-		FName PreviousSourceName = CompiledBindings[0].GetSourceName();
-		for (int32 Index = 0; Index < CompiledBindings.Num(); ++Index)
+		ValidSources |= SourceKey.GetBit();
+		if (ClassSource.RequireSettingUserWidgetProperty())
 		{
-			const FMVVMViewClass_CompiledBinding& Binding = CompiledBindings[Index];
-
-			// A binding execution might call SetViewModel.
-			//The bBindingsInitialized for that source needs to be true to re execute all the bindings.
-			if (PreviousSourceName != Binding.GetSourceName())
+			UUserWidget* UserWidget = GetUserWidget();
+			FObjectPropertyBase* FoundObjectProperty = FindFProperty<FObjectPropertyBase>(UserWidget->GetClass(), ClassSource.GetUserWidgetPropertyName());
+			if (ensureAlwaysMsgf(FoundObjectProperty, TEXT("The compiler should have added the property")))
 			{
-				if (FMVVMViewSource* FoundSource = FindViewSource(PreviousSourceName))
+				if (ensure(NewSource->GetClass()->IsChildOf(FoundObjectProperty->PropertyClass)))
 				{
-					FoundSource->bBindingsInitialized = true;
+					FoundObjectProperty->SetObjectPropertyValue_InContainer(UserWidget, NewSource);
+					ViewSource.bAssignedToUserWidgetProperty = true;
 				}
-				PreviousSourceName = Binding.GetSourceName();
-			}
-
-			if (Binding.IsEnabledByDefault())
-			{
-				EnableLibraryBinding(Binding, Index, false, false);
 			}
 		}
 	}
-
-	for (FMVVMViewSource& Source : Sources)
-	{
-		Source.bBindingsInitialized = true;
-	}
-
-	bBindingsInitialized = true;
-
-	if (bHasEveryTickBinding)
-	{
-		GEngine->GetEngineSubsystem<UMVVMBindingSubsystem>()->AddViewWithEveryTickBinding(this);
-	}
-}
-
-
-void UMVVMView::UninitializeBindings()
-{
-	if (!bBindingsInitialized)
+	else
 	{
-		return;
+		ValidSources &= ~SourceKey.GetBit();
 	}
 
-	SCOPE_CYCLE_COUNTER(STAT_UMG_Viewmodel_UninitializeBindings);
-
-	UninitializeInternal(false);
+	ViewSource.Source = NewSource;
 }
 
 
-void UMVVMView::UninitializeInternal(bool bUninitializeSources)
+void UMVVMView::UninitializeSource(FMVVMView_SourceKey SourceKey)
 {
-	SCOPE_CYCLE_COUNTER(STAT_UMG_Viewmodel_UninitializeBindings);
+	FMVVMView_Source& ViewSource = Sources[SourceKey.GetIndex()];
+	const FMVVMViewClass_Source& ClassSource = ClassExtension->GetSource(ViewSource.ClassKey);
 
-	bool bUninitializeBindings = bBindingsInitialized;
-
-	bBindingsInitialized = false;
-	bSourcesInitialized = !bUninitializeSources;
-
-	UUserWidget* UserWidget = GetUserWidget();
-	check(UserWidget);
-
-	const TArrayView<const FMVVMViewClass_SourceCreator> AllViewModelCreators = ClassExtension->GetViewModelCreators();
-	for (int32 Index = 0; Index < AllViewModelCreators.Num(); ++Index)
+	if (ViewSource.bSourceInitialized)
 	{
-		const FMVVMViewClass_SourceCreator& Item = AllViewModelCreators[Index];
-		FMVVMViewSource& ViewSource = Sources[Index];
-		ViewSource.bBindingsInitialized = false;
-
-		if (ViewSource.RegisteredCount > 0 && ViewSource.Source && bUninitializeBindings)
-		{
-			TScriptInterface<INotifyFieldValueChanged> SourceAsInterface = ViewSource.Source;
-			checkf(SourceAsInterface.GetInterface(), TEXT("It was added as a INotifyFieldValueChanged. It should still be."));
-			SourceAsInterface->RemoveAllFieldValueChangedDelegates(this);
-		}
-
-		// For GC release any object used by the view
-		if (!ViewSource.bSetManually && bUninitializeSources)
+		ViewSource.bSourceInitialized = false;
+		if (!ViewSource.bSetManually)
 		{
-			Item.DestroyInstance(ViewSource.Source, this);
+			ClassSource.ReleaseInstance(ViewSource.Source, this);
+			ViewSource.bSourceInitialized = false;
 
 			ViewSource.Source = nullptr;
 			if (ViewSource.bAssignedToUserWidgetProperty)
 			{
-				FObjectPropertyBase* FoundObjectProperty = FindFProperty<FObjectPropertyBase>(UserWidget->GetClass(), ViewSource.SourceName);
+				UUserWidget* UserWidget = GetUserWidget();
+				FObjectPropertyBase* FoundObjectProperty = FindFProperty<FObjectPropertyBase>(UserWidget->GetClass(), ClassSource.GetUserWidgetPropertyName());
 				if (ensureAlwaysMsgf(FoundObjectProperty, TEXT("The compiler should have added the property")))
 				{
 					FoundObjectProperty->SetObjectPropertyValue_InContainer(UserWidget, nullptr);
@@ -262,744 +221,672 @@ void UMVVMView::UninitializeInternal(bool bUninitializeSources)
 				ViewSource.bAssignedToUserWidgetProperty = false;
 			}
 
+			ValidSources &= ~ViewSource.ClassKey.GetBit();
+
 #if UE_WITH_MVVM_DEBUGGING
-			UE::MVVM::FDebugging::BroadcastViewSourceValueChanged(this, ViewSource);
+			UE::MVVM::FDebugging::BroadcastViewSourceValueChanged(this, ViewSource.ClassKey, SourceKey);
 #endif
 		}
 	}
+}
 
-	RegisteredLibraryBindings.Reset();
 
-	if (bHasEveryTickBinding)
+void UMVVMView::InitializeBindings()
+{
+	if (!bSourcesInitialized)
 	{
-		ensureMsgf(bUninitializeBindings, TEXT("The "));
-		GEngine->GetEngineSubsystem<UMVVMBindingSubsystem>()->RemoveViewWithEveryTickBinding(this);
+		InitializeSources();
 	}
-	bHasEveryTickBinding = false;
-}
 
+	if (bBindingsInitialized)
+	{
+		return;
+	}
 
-TScriptInterface<INotifyFieldValueChanged> UMVVMView::GetViewModel(FName ViewModelName) const
-{
-	const FMVVMViewSource* RegisteredSource = FindViewSource(ViewModelName);
-	return RegisteredSource ? RegisteredSource->Source : nullptr;
+	SCOPE_CYCLE_COUNTER(STAT_UMG_Viewmodel_InitializeBindings);
+
+	check(NumberOfSourceWithTickBinding == 0);
+
+	for (int32 SourceIndex = 0; SourceIndex < Sources.Num(); ++SourceIndex)
+	{
+		InitializeSourceBindings(FMVVMView_SourceKey(SourceIndex), false);
+	}
+
+	InitializeSourceBindingsCommon();
 }
 
 
-bool UMVVMView::SetViewModel(FName ViewModelName, TScriptInterface<INotifyFieldValueChanged> NewValue)
+void UMVVMView::InitializeSourceBindingsCommon()
 {
-	return SetSourceInternal(ViewModelName, NewValue, false);
+	ensure(bHasDefaultTickBinding == false);
+
+	if (ValidSources != 0) // If one source is property initialized, we assumed that at least one binding was registered.
+	{
+		static IConsoleVariable* CVarDefaultExecutionMode = IConsoleManager::Get().FindConsoleVariable(TEXT("MVVM.DefaultExecutionMode"));
+		bHasDefaultTickBinding = ensure(CVarDefaultExecutionMode) ? (EMVVMExecutionMode)CVarDefaultExecutionMode->GetInt() == EMVVMExecutionMode::Tick : false;
+	}
+
+	if (NumberOfSourceWithTickBinding == 0 && bHasDefaultTickBinding)
+	{
+		GEngine->GetEngineSubsystem<UMVVMBindingSubsystem>()->AddViewWithTickBinding(this);
+	}
+
+	bBindingsInitialized = true;
 }
 
 
-bool UMVVMView::SetViewModelByClass(TScriptInterface<INotifyFieldValueChanged> NewValue)
+void UMVVMView::UninitializeBindings()
 {
-	if (!NewValue.GetObject())
+	if (!bBindingsInitialized)
 	{
-		UE::MVVM::FMessageLog Log(GetUserWidget());
-		Log.Error(LOCTEXT("SetViewModelInvalidObject", "The new viewmodel is null."));
-		return false;
+		return;
 	}
 
-	if (ClassExtension == nullptr)
+	SCOPE_CYCLE_COUNTER(STAT_UMG_Viewmodel_UninitializeBindings);
+
+	bBindingsInitialized = false;
+
+	const TArrayView<const FMVVMViewClass_Source> ClassSources = ClassExtension->GetSources();
+	for (int32 SourceIndex = 0; SourceIndex < ClassSources.Num(); ++SourceIndex)
 	{
-		UE::MVVM::FMessageLog Log(GetUserWidget());
-		Log.Error(LOCTEXT("SetViewModelInvalidClass", "The view is not constructed."));
-		return false;
+		const FMVVMViewClass_Source& ClassSource = ClassSources[SourceIndex];
+		FMVVMView_Source& ViewSource = Sources[SourceIndex];
+		UninitializeSourceBindings(FMVVMViewClass_SourceKey(SourceIndex), ClassSource, ViewSource);
+	}
+
+	// Remove all delayed bindings from the 
+	UMVVMBindingSubsystem* BindingSubsystem = GEngine->GetEngineSubsystem<UMVVMBindingSubsystem>();
+	BindingSubsystem->RemoveDelayedBindings(this);
+
+	if (NumberOfSourceWithTickBinding == 0 && bHasDefaultTickBinding)
+	{
+		BindingSubsystem->RemoveViewWithTickBinding(this);
+		bHasDefaultTickBinding = false;
 	}
+}
+
+
+void UMVVMView::InitializeSourceBindings(FMVVMView_SourceKey SourceKey, bool bRunAllBindings)
+{
+	FMVVMView_Source& ViewSource = Sources[SourceKey.GetIndex()];
+	const FMVVMViewClass_Source& ClassSource = ClassExtension->GetSource(ViewSource.ClassKey);
 
-	FName ViewModelName;
-	for (const FMVVMViewClass_SourceCreator& SourceCreator : ClassExtension->GetViewModelCreators())
+	const bool bIsPointerValid = ViewSource.Source != nullptr;
+	const bool bIsBitfieldValid = (ViewSource.ClassKey.GetBit() & ValidSources) != 0;
+	ensureMsgf(bIsPointerValid == bIsBitfieldValid, TEXT("The source %s should be valid."), *ClassSource.GetName().ToString());
+	if (bIsPointerValid)
 	{
-		if (NewValue.GetObject()->GetClass()->IsChildOf(SourceCreator.GetSourceClass()))
+		// Register the delegates
+		if (ClassSource.GetFieldIds().Num() > 0)
 		{
-			if (ViewModelName.IsNone())
+			TScriptInterface<INotifyFieldValueChanged> NotifyFieldValueChanged = ViewSource.Source->Implements<UNotifyFieldValueChanged>() ? TScriptInterface<INotifyFieldValueChanged>(ViewSource.Source) : TScriptInterface<INotifyFieldValueChanged>();
+			if (NotifyFieldValueChanged.GetInterface() != nullptr)
 			{
-				ViewModelName = SourceCreator.GetSourceName();
+				for (const FMVVMViewClass_FieldId& FieldId : ClassSource.GetFieldIds())
+				{
+					if (ensureMsgf(FieldId.IsValid(), TEXT("Invalid field. It failed somewhere in the compiler.")))
+					{
+						UE::FieldNotification::FFieldMulticastDelegate::FDelegate Delegate = UE::FieldNotification::FFieldMulticastDelegate::FDelegate::CreateUObject(this, &UMVVMView::HandledLibraryBindingValueChanged);
+						NotifyFieldValueChanged->AddFieldValueChangedDelegate(FieldId.GetFieldId(), MoveTemp(Delegate));
+						++ViewSource.RegisteredCount;
+					}
+				}
 			}
-			else
+		}
+
+		// Run the bindings
+		for (const FMVVMViewClass_SourceBinding& SourceBinding : ClassSource.GetBindings())
+		{
+			if (SourceBinding.ExecuteAtInitialization() || bRunAllBindings)
 			{
-				UE::MVVM::FMessageLog Log(GetUserWidget());
-				Log.Error(FText::Format(LOCTEXT("SetViewModelViewModelNotUnique", "More than one viewmodel match the type of the passed instance. Make sure there exists only one viewmodel of type {0} in the widget blueprint, or try SetViewModel by name."), FText::FromName(NewValue.GetObject()->GetClass()->GetFName())));
-				return false;
+				const FMVVMViewClass_Binding& ClassBinding = ClassExtension->GetBinding(SourceBinding.GetBindingKey());
+				ExecuteBindingImmediately(ClassBinding, SourceBinding.GetBindingKey());
+
+				ensureMsgf(SourceBinding.GetFieldId().IsValid() == ClassBinding.IsOneWay(), TEXT("If the binding has no id, then it must be a OneTime."));
 			}
 		}
-	}
 
-	if (ViewModelName.IsNone())
+		if (ClassSource.HasTickBindings())
+		{
+			++NumberOfSourceWithTickBinding;
+			if (NumberOfSourceWithTickBinding == 1 && !bHasDefaultTickBinding)
+			{
+				GEngine->GetEngineSubsystem<UMVVMBindingSubsystem>()->AddViewWithTickBinding(this);
+			}
+		}
+	}
+	else if (!ClassSource.IsOptional())
 	{
 		UE::MVVM::FMessageLog Log(GetUserWidget());
-		Log.Error(LOCTEXT("SetViewModelViewModelNotFound", "A created viewmodel matching the class of the passed instance could not be found."));
-		return false;
+		Log.Error(FText::Format(LOCTEXT("InitializeBindingFailInvalidSource", "The source '{0}' is invalid and could not initialize the bindings.")
+			, FText::FromName(ClassSource.GetName())
+		));
 	}
-
-	return SetViewModel(ViewModelName, NewValue);
+	ViewSource.bBindingsInitialized = true;
 }
 
 
-bool UMVVMView::EvaluateSourceCreator(int32 SourceCreatorIndex)
+void UMVVMView::UninitializeSourceBindings(FMVVMViewClass_SourceKey SourceKey, const FMVVMViewClass_Source& ClassSource, FMVVMView_Source& ViewSource)
 {
-	if (!ClassExtension->GetViewModelCreators().IsValidIndex(SourceCreatorIndex))
+	bool bResult = false;
+
+	// The view can be bBindingsInitialized but the pointer could invalid then and not actual bindgins was registered.
+	if (ViewSource.bBindingsInitialized && ViewSource.RegisteredCount > 0)
 	{
-		UE::MVVM::FMessageLog Log(GetUserWidget());
-		Log.Error(FText::Format(LOCTEXT("ReevaluateDynamicSourceViewModelNameNotFound", "Internal error. The dynamic source {0} cound not be found."), FText::AsNumber(SourceCreatorIndex)));
-		return false;
+		if (ViewSource.Source)
+		{
+			TScriptInterface<INotifyFieldValueChanged> SourceAsInterface = ViewSource.Source;
+			checkf(SourceAsInterface.GetInterface(), TEXT("It was added as a INotifyFieldValueChanged. It should still be."));
+			SourceAsInterface->RemoveAllFieldValueChangedDelegates(this);
+		}
+		else
+		{
+			UE::MVVM::FMessageLog Log(GetUserWidget());
+			Log.Error(FText::Format(LOCTEXT("UninitializeBindingFailInvalidSource", "The source '{0}' is invalid and could not uninitialize the bindings.")
+				, FText::FromName(ClassSource.GetName())
+			));
+		}
+
+		if (ClassSource.HasTickBindings())
+		{
+			--NumberOfSourceWithTickBinding;
+			check(NumberOfSourceWithTickBinding >= 0);
+			if (NumberOfSourceWithTickBinding == 0 && !bHasDefaultTickBinding)
+			{
+				GEngine->GetEngineSubsystem<UMVVMBindingSubsystem>()->RemoveViewWithTickBinding(this);
+			}
+		}
 	}
 
-	UUserWidget* UserWidget = GetUserWidget();
-	const TArrayView<const FMVVMViewClass_SourceCreator> ViewModelCreators = ClassExtension->GetViewModelCreators();
-	const FMVVMViewClass_SourceCreator& SourceCreator = ViewModelCreators[SourceCreatorIndex];
-	UObject* NewSource = SourceCreator.CreateInstance(ClassExtension, this, UserWidget);
-	ensureMsgf((NewSource == nullptr || NewSource->GetClass()->ImplementsInterface(UNotifyFieldValueChanged::StaticClass())), TEXT("The source has implement the interface. It should be check at compile time."));
-	return SetSourceInternal(SourceCreator.GetSourceName(), NewSource, true);
+	ViewSource.bBindingsInitialized = false;
 }
 
 
-bool UMVVMView::SetSourceInternal(FName ViewModelName, TScriptInterface<INotifyFieldValueChanged> NewValue, bool bForDynamicSource)
+void UMVVMView::ExecuteBindingImmediately(const FMVVMViewClass_Binding& ClassBinding, FMVVMViewClass_BindingKey KeyForLog) const
 {
-	SCOPE_CYCLE_COUNTER(STAT_UMG_Viewmodel_SetSource);
-
-	if (ViewModelName.IsNone())
+	check(ClassExtension);
+	check(GetUserWidget())
+	if ((ClassBinding.GetSources() & ValidSources) == ClassBinding.GetSources())
 	{
-		UE::MVVM::FMessageLog Log(GetUserWidget());
-		Log.Error(LOCTEXT("SetViewModelInvalidName", "The viewmodel name is empty."));
-		return false;
-	}
+	
+		// All the source are valid. Run the binding.
+		FMVVMCompiledBindingLibrary::EConversionFunctionType FunctionType = ClassBinding.GetBinding().HasComplexConversionFunction() ? FMVVMCompiledBindingLibrary::EConversionFunctionType::Complex : FMVVMCompiledBindingLibrary::EConversionFunctionType::Simple;
+		TValueOrError<void, FMVVMCompiledBindingLibrary::EExecutionFailingReason> ExecutionResult = ClassExtension->GetBindingLibrary().Execute(GetUserWidget(), ClassBinding.GetBinding(), FunctionType);
 
-	if (ClassExtension == nullptr)
-	{
-		UE::MVVM::FMessageLog Log(GetUserWidget());
-		Log.Error(LOCTEXT("SetViewModelInvalidClass", "The view is not constructed."));
-		return false;
-	}
+#if UE_WITH_MVVM_DEBUGGING
+		if (ExecutionResult.HasError())
+		{
+			UE::MVVM::FMessageLog Log(GetUserWidget());
+			Log.Error(FText::Format(LOCTEXT("ExecuteBindingFailGeneric", "The binding '{0}' was not executed. {1}.")
+				, FText::FromString(ClassBinding.ToString(ClassExtension, FMVVMViewClass_Binding::FToStringArgs::Short()))
+				, FMVVMCompiledBindingLibrary::LexToText(ExecutionResult.GetError())
+			));
 
-	const int32 AllCreatedSourcesIndex = ClassExtension->GetViewModelCreators().IndexOfByPredicate([ViewModelName](const FMVVMViewClass_SourceCreator& Item)
+			UE::MVVM::FDebugging::BroadcastLibraryBindingExecuted(this, KeyForLog, ExecutionResult.GetError());
+		}
+		else
 		{
-			return Item.GetSourceName() == ViewModelName;
-		});
-		
-	if (AllCreatedSourcesIndex == INDEX_NONE)
-	{
-		UE::MVVM::FMessageLog Log(GetUserWidget());
-		Log.Error(LOCTEXT("SetViewModelViewModelNameNotFound", "The viewmodel name could not be found."));
-		return false;
+			if (bLogBinding)
+			{
+				UE::MVVM::FMessageLog Log(GetUserWidget());
+				Log.Info(FText::Format(LOCTEXT("ExecuteBindingGeneric", "Execute binding '{0}'.")
+					, FText::FromString(ClassBinding.ToString(ClassExtension, FMVVMViewClass_Binding::FToStringArgs::All()))
+				));
+			}
+			UE::MVVM::FDebugging::BroadcastLibraryBindingExecuted(this, KeyForLog);
+		}
+#else
+		if (ExecutionResult.HasError())
+		{
+			UE::MVVM::FMessageLog Log(GetUserWidget());
+			Log.Error(FText::Format(LOCTEXT("ExecuteBindingFailGeneric", "The binding '{0}' was not executed. {1}.")
+				, FText::AsNumber(KeyForLog.GetIndex())
+				, FMVVMCompiledBindingLibrary::LexToText(ExecutionResult.GetError())
+			));
+		}
+		else if (bLogBinding)
+		{
+			UE::MVVM::FMessageLog Log(GetUserWidget());
+			Log.Info(FText::Format(LOCTEXT("ExecuteBindingGeneric", "Execute binding '{0}'.")
+				, FText::AsNumber(KeyForLog.GetIndex())
+			));
+		}
+#endif
 	}
-
-	const TArrayView<const FMVVMViewClass_SourceCreator> ViewModelCreators = ClassExtension->GetViewModelCreators();
-	const FMVVMViewClass_SourceCreator& SourceCreator = ViewModelCreators[AllCreatedSourcesIndex];
-
-	if (bForDynamicSource && !SourceCreator.CanBeEvaluated())
+	else if ((ClassBinding.GetSources() & ClassExtension->GetOptionalSources()) != ClassBinding.GetSources())
 	{
+#if UE_WITH_MVVM_DEBUGGING
 		UE::MVVM::FMessageLog Log(GetUserWidget());
-		Log.Error(FText::Format(LOCTEXT("SetViewModelCannotBeEvalaute", "The new viewmodel {0} cannot be evaluated again."), FText::FromName(ViewModelName)));
-		return false;
-	}
-	else if (!bForDynamicSource && !SourceCreator.CanBeSet())
-	{
+		Log.Error(FText::Format(LOCTEXT("ExecuteBindingFailInvalidSource", "The binding '{0}' was not executed. There are invalid sources.")
+			, FText::FromString(ClassBinding.ToString(ClassExtension, FMVVMViewClass_Binding::FToStringArgs::Short()))
+		));
+		UE::MVVM::FDebugging::BroadcastLibraryBindingExecuted(this, KeyForLog, FMVVMCompiledBindingLibrary::EExecutionFailingReason::InvalidSource);
+#else
 		UE::MVVM::FMessageLog Log(GetUserWidget());
-		Log.Error(FText::Format(LOCTEXT("SetViewModelCannotBeSet", "The new viewmodel {0} cannot be set."), FText::FromName(ViewModelName)));
-		return false;
+		Log.Error(FText::Format(LOCTEXT("ExecuteBindingFailInvalidSource", "The binding '{0}' was not executed. There are invalid sources.")
+			, FText::AsNumber(KeyForLog.GetIndex())
+		));
+#endif
 	}
-			
-	if (NewValue.GetObject()
-		&& !NewValue.GetObject()->GetClass()->IsChildOf(SourceCreator.GetSourceClass()))
+}
+
+
+bool UMVVMView::EvaluateSource(FMVVMViewClass_SourceKey SourceKey)
+{
+	if (!SourceKey.IsValid())
 	{
 		UE::MVVM::FMessageLog Log(GetUserWidget());
-		Log.Error(LOCTEXT("SetViewModelInvalidValueType", "The new viewmodel is not of the expected type."));
-		UE_LOG(LogMVVM, Error, TEXT("The viewmodel name '%s' is invalid for the view '%s'"), *ViewModelName.ToString(), *GetFullName());
+		Log.Error(FText::Format(LOCTEXT("ReevaluateDynamicSourceViewModelNameNotFound", "Internal error. The dynamic source {0} cound not be found."), FText::AsNumber(SourceKey.GetIndex())));
 		return false;
 	}
-			
-	if (!Sources.IsValidIndex(AllCreatedSourcesIndex) || Sources[AllCreatedSourcesIndex].SourceName != ViewModelName)
+
+	UUserWidget* UserWidget = GetUserWidget();
+	const FMVVMViewClass_Source& ClassSource = ClassExtension->GetSource(SourceKey);
+	UObject* NewSource = ClassSource.GetOrCreateInstance(ClassExtension, this, UserWidget);
+	ensureMsgf((NewSource == nullptr || NewSource->GetClass()->ImplementsInterface(UNotifyFieldValueChanged::StaticClass())), TEXT("The source has implement the interface. It should be check at compile time."));
+	bool bResult = SetSourceInternal(SourceKey, NewSource, true);
+	if (bResult)
 	{
 		UE::MVVM::FMessageLog Log(GetUserWidget());
-		Log.Error(LOCTEXT("SetViewModelInvalidCreatedSource", "An internal error occurs. The new viewmodel is not found."));
-		UE_LOG(LogMVVM, Error, TEXT("The viewmodel name '%s' is invalid for the view '%s'"), *ViewModelName.ToString(), *GetFullName());
-		return false;
+		Log.Error(FText::Format(LOCTEXT("ExecuteBindingFailEvaluate", "The evaluate source '{0}' fail execution.")
+			, FText::FromName(ClassSource.GetName())));
 	}
+	return bResult;
+}
 
-	FMVVMViewSource& ViewSource = Sources[AllCreatedSourcesIndex];
-	UObject* PreviousValue = Sources[AllCreatedSourcesIndex].Source;
-	if (PreviousValue != NewValue.GetObject())
-	{
-		const TArrayView<const FMVVMViewClass_CompiledBinding> CompiledBindings = ClassExtension->GetCompiledBindings();
-		const TArrayView<const FMVVMViewClass_CompiledEvent> CompiledEvents = ClassExtension->GetCompiledEvents();
 
-		// Unregister any bindings from that source
-		if (Sources[AllCreatedSourcesIndex].bBindingsInitialized)
-		{
-			for (int32 Index = 0; Index < CompiledBindings.Num(); ++Index)
-			{
-				const FMVVMViewClass_CompiledBinding& Binding = CompiledBindings[Index];
-				if (IsLibraryBindingEnabled(Index) && Binding.GetSourceName() == ViewModelName)
-				{
-					DisableLibraryBinding(Binding, Index);
-				}
-			}
-		}
+namespace UE::MVVM::Private
+{
+	using FRecursiveDetctionElement = TTuple<const UObject*, FMVVMViewClass_BindingKey>;
+	TArray<FRecursiveDetctionElement> RecursiveDetector;
+}
 
-		// Unregister any events from that source
-		if (bConstructed)
-		{
-			for (int32 Index = 0; Index < CompiledEvents.Num(); ++Index)
-			{
-				const FMVVMViewClass_CompiledEvent& CompiledEvent = CompiledEvents[Index];
-				if (IsLibraryEventEnabled(Index) && CompiledEvent.GetSourceName() == ViewModelName)
-				{
-					DisableLibraryEvent(CompiledEvent, Index);
-				}
-			}
-		}
 
-		if (!ViewSource.bSetManually && ViewSource.Source)
+void UMVVMView::HandledLibraryBindingValueChanged(UObject* InSource, UE::FieldNotification::FFieldId InFieldId)
+{
+	SCOPE_CYCLE_COUNTER(STAT_UMG_Viewmodel_ExecuteBinding_ValueChanged);
+
+	check(InSource);
+	check(InFieldId.IsValid());
+
+	if (ensure(ClassExtension))
+	{
+		int32 ViewSourceIndex = Sources.IndexOfByPredicate([InSource](const FMVVMView_Source& Other){ return Other.Source == InSource; });
+		if (ViewSourceIndex == INDEX_NONE)
 		{
-			SourceCreator.DestroyInstance(ViewSource.Source, this);
+			ensureMsgf(false, TEXT("No source was found. That means the source was removed but the binding was not."));
+			return;
 		}
-
-		ViewSource.Source = NewValue.GetObject();
-		ViewSource.bSetManually = !bForDynamicSource;
-		if (SourceCreator.IsSourceAUserWidgetProperty())
+		const FMVVMView_SourceKey ViewSourceKey = FMVVMView_SourceKey(ViewSourceIndex);
+		const FMVVMView_Source& ViewSource = Sources[ViewSourceIndex];
+		ensure(ViewSource.bBindingsInitialized);
+		if (!ViewSource.bBindingsInitialized || !ViewSource.bSourceInitialized)
 		{
-			FObjectPropertyBase* FoundObjectProperty = FindFProperty<FObjectPropertyBase>(GetUserWidget()->GetClass(), ViewModelName);
-			if (ensure(FoundObjectProperty))
-			{
-				FoundObjectProperty->SetObjectPropertyValue_InContainer(GetUserWidget(), NewValue.GetObject());
-				ViewSource.bAssignedToUserWidgetProperty = true;
-			}
+			// we do not want to run a binding while we are initializing the bindings.
+			return;
 		}
 
-		bool bPreviousEveryTickBinding = bHasEveryTickBinding;
-		bHasEveryTickBinding = false;
+		const FMVVMViewClass_SourceKey ClassSourceKey = ViewSource.ClassKey;
+		const FMVVMViewClass_Source& ClassSource = ClassExtension->GetSource(ClassSourceKey);
 
-		// Register back any bindings that was previously enabled
-		if (Sources[AllCreatedSourcesIndex].bBindingsInitialized && NewValue.GetObject())
+		// Run all evaluates
+		if (ClassSource.HasEvaluateBindings())
 		{
-			// Enabled the default bindings
-			for (int32 Index = 0; Index < CompiledBindings.Num(); ++Index)
+			for (const FMVVMViewClass_EvaluateSource& ClassEvaluate : ClassExtension->GetEvaluateSources())
 			{
-				const FMVVMViewClass_CompiledBinding& Binding = CompiledBindings[Index];
-				if (Binding.IsEnabledByDefault())
+				if (ClassEvaluate.GetParentSource() == ClassSourceKey && ClassEvaluate.GetFieldId().GetFieldName() == InFieldId.GetName())
 				{
-					// Binding on this viewmodel
-					if (Binding.GetSourceName() == ViewModelName)
-					{
-						// The binding needs execution, but maybe it's a complex binding and not all the sources from that binding are initialized.
-						//This occurs only at initialization when all the bindings from source A exectuted and a binding from source B call SetViewModel on A.
-						// A should reexecute its bindings but only if all the sources from the complex conversion initialized. We do not know that. Wait for the next frame and executed it. 
-						const bool bForceDelayed = Binding.IsConversionFunctionComplex() && !bBindingsInitialized;
-						EnableLibraryBinding(Binding, Index, true, bForceDelayed);
-
-						// Bindings that depends on this binding
-						if (Binding.IsEvaluateSourceCreatorBinding())
-						{
-							int32 ParentSourceIndex = ClassExtension->GetViewModelCreators().IndexOfByPredicate([ViewModelName](const FMVVMViewClass_SourceCreator& Item)
-								{
-									return Item.GetParentSourceName() == ViewModelName;
-								});
-							if (ParentSourceIndex != INDEX_NONE)
-							{
-								check(ParentSourceIndex != AllCreatedSourcesIndex)
-								EvaluateSourceCreator(ParentSourceIndex); // it will deactivate previous binding 
-							}
-						}
-					}
+					EvaluateSource(ClassEvaluate.GetSource());
 				}
 			}
 		}
 
-		// Register back any events that was previously enabled
-		if (bConstructed && NewValue.GetObject())
+		// Run all bindings
+		for (const FMVVMViewClass_SourceBinding& SourceBinding : ClassSource.GetBindings())
 		{
-			// Enabled the default events
-			for (int32 Index = 0; Index < CompiledEvents.Num(); ++Index)
+			if (SourceBinding.GetFieldId().GetFieldName() == InFieldId.GetName())
 			{
-				const FMVVMViewClass_CompiledEvent& CompiledEvent = CompiledEvents[Index];
-				if (CompiledEvent.GetSourceName() == ViewModelName)
+				const FMVVMViewClass_Binding& ClassBinding = ClassExtension->GetBinding(SourceBinding.GetBindingKey());
+				if (ensure(ClassBinding.IsOneWay()))
 				{
-					EnableLibraryEvent(CompiledEvent, Index);
-				}
-			}
-		}
+					const EMVVMExecutionMode ExecutionMode = ClassBinding.GetExecuteMode();
+					if (ExecutionMode == EMVVMExecutionMode::Immediate)
+					{
+						// Test for recursivity
+						const UMVVMView* Self = this;
+						if (UE::MVVM::Private::RecursiveDetector.FindByPredicate([Self, InCompiledBindingIndex = SourceBinding.GetBindingKey()](const UE::MVVM::Private::FRecursiveDetctionElement& Element)
+							{
+								return Element.Get<0>() == Self && Element.Get<1>() == InCompiledBindingIndex;
+							}) != nullptr)
+						{
+							ensureAlwaysMsgf(false, TEXT("Recursive binding detected"));
+							//Todo add more infos. Callstack maybe? Log the chain?
+							UE::MVVM::FMessageLog Log(Self->GetUserWidget());
+							Log.Error(LOCTEXT("RecursionDetected", "A recursive binding was detected (ie. A->B->C->A->B->C) at runtime."));
+							return;
+						}
 
-		if (bBindingsInitialized && bPreviousEveryTickBinding != bHasEveryTickBinding)
-		{
-			if (bHasEveryTickBinding)
-			{
-				GEngine->GetEngineSubsystem<UMVVMBindingSubsystem>()->AddViewWithEveryTickBinding(this);
-			}
-			else
-			{
-				GEngine->GetEngineSubsystem<UMVVMBindingSubsystem>()->RemoveViewWithEveryTickBinding(this);
+						{
+							UE::MVVM::Private::RecursiveDetector.Emplace(this, SourceBinding.GetBindingKey());
+							ExecuteBindingImmediately(ClassBinding, SourceBinding.GetBindingKey());
+							UE::MVVM::Private::RecursiveDetector.Pop();
+						}
+					}
+					else if (ExecutionMode == EMVVMExecutionMode::Delayed)
+					{
+						GEngine->GetEngineSubsystem<UMVVMBindingSubsystem>()->AddDelayedBinding(this, SourceBinding.GetBindingKey());
+					}
+					else
+					{
+						ensureMsgf(false, TEXT("We should not have registered the binding."));
+					}
+				}
 			}
 		}
-
-#if UE_WITH_MVVM_DEBUGGING
-		UE::MVVM::FDebugging::BroadcastViewSourceValueChanged(this, ViewSource);
-#endif
 	}
-	return true;
 }
 
 
-namespace UE::MVVM::Private
-{
-	using FRecursiveDetctionElement = TTuple<const UObject*, int32>;
-	TArray<FRecursiveDetctionElement> RecursiveDetector;
-}
-
-void UMVVMView::HandledLibraryBindingValueChanged(UObject* InViewModelOrWidget, UE::FieldNotification::FFieldId InFieldId, int32 InCompiledBindingIndex) const
+void UMVVMView::ExecuteDelayedBinding(const FMVVMViewClass_BindingKey& DelayedBinding) const
 {
-	SCOPE_CYCLE_COUNTER(STAT_UMG_Viewmodel_ExecuteBinding_ValueChanged);
-
-	check(InViewModelOrWidget);
-	check(InFieldId.IsValid());
+	SCOPE_CYCLE_COUNTER(STAT_UMG_Viewmodel_ExecuteBinding_Delayed);
 
-	if (ensure(ClassExtension))
+	if (ensure(ClassExtension) && bBindingsInitialized)
 	{
-		checkf(ClassExtension->GetCompiledBindings().IsValidIndex(InCompiledBindingIndex), TEXT("The binding at index '%d' does not exist. The binding was probably not cleared on destroyed."), InCompiledBindingIndex);
-		const FMVVMViewClass_CompiledBinding& Binding = ClassExtension->GetCompiledBinding(InCompiledBindingIndex);
-
-		EMVVMExecutionMode ExecutionMode = Binding.GetExecuteMode();
-		if (ExecutionMode == EMVVMExecutionMode::Delayed)
-		{
-			GEngine->GetEngineSubsystem<UMVVMBindingSubsystem>()->AddDelayedBinding(this, InCompiledBindingIndex);
-		}
-		else if (ExecutionMode != EMVVMExecutionMode::Tick)
+		ensure(bSourcesInitialized);
+		if (ensure(ClassExtension->GetBindings().IsValidIndex(DelayedBinding.GetIndex())))
 		{
 			// Test for recursivity
 			const UMVVMView* Self = this;
-			if (UE::MVVM::Private::RecursiveDetector.FindByPredicate([Self, InCompiledBindingIndex](const UE::MVVM::Private::FRecursiveDetctionElement& Element)
+			if (UE::MVVM::Private::RecursiveDetector.FindByPredicate([Self, InCompiledBindingIndex = DelayedBinding](const UE::MVVM::Private::FRecursiveDetctionElement& Element)
 				{
 					return Element.Get<0>() == Self && Element.Get<1>() == InCompiledBindingIndex;
 				}) != nullptr)
 			{
 				ensureAlwaysMsgf(false, TEXT("Recursive binding detected"));
 				//Todo add more infos. Callstack maybe? Log the chain?
-				UE::MVVM::FMessageLog Log(GetUserWidget());
+				UE::MVVM::FMessageLog Log(Self->GetUserWidget());
 				Log.Error(LOCTEXT("RecursionDetected", "A recursive binding was detected (ie. A->B->C->A->B->C) at runtime."));
 				return;
 			}
 
 			{
-				UE::MVVM::Private::RecursiveDetector.Emplace(this, InCompiledBindingIndex);
-
-				ExecuteLibraryBinding(Binding, InCompiledBindingIndex);
-
+				UE::MVVM::Private::RecursiveDetector.Emplace(this, DelayedBinding);
+				const FMVVMViewClass_Binding& ClassBinding = ClassExtension->GetBinding(DelayedBinding);
+				ExecuteBindingImmediately(ClassBinding, DelayedBinding);
 				UE::MVVM::Private::RecursiveDetector.Pop();
 			}
 		}
-		else
-		{
-			ensureMsgf(false, TEXT("We should not have registered the binding since it will always be executed."));
-		}
 	}
 }
 
 
-void UMVVMView::ExecuteDelayedBinding(const FMVVMViewDelayedBinding& DelayedBinding) const
+void UMVVMView::ExecuteTickBindings() const
 {
-	SCOPE_CYCLE_COUNTER(STAT_UMG_Viewmodel_ExecuteBinding_Delayed);
-
-	if (ensure(ClassExtension) && bBindingsInitialized)
+	if (ensure(ClassExtension))
 	{
-		ensure(bSourcesInitialized);
-		if (ensure(ClassExtension->GetCompiledBindings().IsValidIndex(DelayedBinding.GetCompiledBindingIndex())))
+		ensure(NumberOfSourceWithTickBinding > 0 || bHasDefaultTickBinding);
+
+		bool bAtLeastOneBindingWasExecuted = false;
+		const TArrayView<const FMVVMViewClass_Binding> ClassBindings = ClassExtension->GetBindings();
+		for (int32 BindingIndex = 0; BindingIndex < ClassBindings.Num(); ++BindingIndex)
 		{
-			if (IsLibraryBindingEnabled(DelayedBinding.GetCompiledBindingIndex()))
+			const FMVVMViewClass_Binding& ClassBinding = ClassBindings[BindingIndex];
+			const EMVVMExecutionMode ExecutionMode = ClassBinding.GetExecuteMode();
+			if (ExecutionMode == EMVVMExecutionMode::Tick)
 			{
-				const FMVVMViewClass_CompiledBinding& Binding = ClassExtension->GetCompiledBinding(DelayedBinding.GetCompiledBindingIndex());
-
-				// Test for recursivity
-				int32 CompiledBindingIndex = DelayedBinding.GetCompiledBindingIndex();
-				const UMVVMView* Self = this;
-				if (UE::MVVM::Private::RecursiveDetector.FindByPredicate([Self, CompiledBindingIndex](const UE::MVVM::Private::FRecursiveDetctionElement& Element)
-					{
-						return Element.Get<0>() == Self && Element.Get<1>() == CompiledBindingIndex;
-					}) != nullptr)
-				{
-					ensureAlwaysMsgf(false, TEXT("Recursive binding detected"));
-					//Todo add more infos. Callstack maybe? Log the chain?
-					UE::MVVM::FMessageLog Log(GetUserWidget());
-					Log.Error(LOCTEXT("RecursionDetected", "A recursive binding was detected (ie. A->B->C->A->B->C) at runtime."));
-					return;
-				}
-
-				{
-					UE::MVVM::Private::RecursiveDetector.Emplace(this, CompiledBindingIndex);
-
-					ExecuteLibraryBinding(Binding, DelayedBinding.GetCompiledBindingIndex());
-
-					UE::MVVM::Private::RecursiveDetector.Pop();
-				}
+				ExecuteBindingImmediately(ClassBinding, FMVVMViewClass_BindingKey(BindingIndex));
+				bAtLeastOneBindingWasExecuted = true;
 			}
 		}
-	}
-}
-
 
-void UMVVMView::ExecuteEveryTickBindings() const
-{
-	SCOPE_CYCLE_COUNTER(STAT_UMG_Viewmodel_ExecuteBinding_Tick);
-
-	ensure(bHasEveryTickBinding);
-	ensure(bBindingsInitialized);
-	ensure(bSourcesInitialized);
-
-	if (ClassExtension && bBindingsInitialized)
-	{
-		const TArrayView<const FMVVMViewClass_CompiledBinding> CompiledBindings = ClassExtension->GetCompiledBindings();
-
-		for (int32 Index = 0; Index < CompiledBindings.Num(); ++Index)
+		if (!bAtLeastOneBindingWasExecuted)
 		{
-			const FMVVMViewClass_CompiledBinding& Binding = CompiledBindings[Index];
-			if (Binding.GetExecuteMode() == EMVVMExecutionMode::Tick && IsLibraryBindingEnabled(Index))
-			{
-				ExecuteLibraryBinding(Binding, Index);
-			}
+			ensureMsgf(false, TEXT("No tick binding was found but the flags was true."));
 		}
 	}
 }
 
 
-void UMVVMView::ExecuteLibraryBinding(const FMVVMViewClass_CompiledBinding& Binding, int32 BindingIndex) const
+TScriptInterface<INotifyFieldValueChanged> UMVVMView::GetViewModel(FName ViewModelName) const
 {
-	check(ClassExtension);
-	check(GetUserWidget());
-
-	bool bError = false;
-	if (!Binding.IsEvaluateSourceCreatorBinding())
+	int32 FoundIndex = ClassExtension->GetSources().IndexOfByPredicate([ViewModelName](const FMVVMViewClass_Source& Other)
 	{
-		FMVVMCompiledBindingLibrary::EConversionFunctionType FunctionType = Binding.IsConversionFunctionComplex() ? FMVVMCompiledBindingLibrary::EConversionFunctionType::Complex : FMVVMCompiledBindingLibrary::EConversionFunctionType::Simple;
-		TValueOrError<void, FMVVMCompiledBindingLibrary::EExecutionFailingReason> ExecutionResult = ClassExtension->GetBindingLibrary().Execute(GetUserWidget(), Binding.GetBinding(), FunctionType);
+		return Other.GetName() == ViewModelName && Other.IsViewModel();
+	});
+	return Sources.IsValidIndex(FoundIndex) ? Sources[FoundIndex].Source : nullptr;
+}
 
-#if UE_WITH_MVVM_DEBUGGING
-		if (ExecutionResult.HasError())
-		{
-			UE::MVVM::FDebugging::BroadcastLibraryBindingExecuted(this, Binding, ExecutionResult.GetError());
-		}
-#endif
 
-		if (ExecutionResult.HasError())
-		{
-			UE::MVVM::FMessageLog Log(GetUserWidget());
-			Log.Error(FText::Format(LOCTEXT("ExecuteBindingFailGeneric", "The binding '{0}' was not executed. {1}.")
-#if UE_WITH_MVVM_DEBUGGING
-				, FText::FromString(Binding.ToString(ClassExtension->GetBindingLibrary(), FMVVMViewClass_CompiledBinding::FToStringArgs::Short()))
-#else
-				, FText::AsNumber(BindingIndex)
-#endif
-				, FMVVMCompiledBindingLibrary::LexToText(ExecutionResult.GetError())
-			));
-			bError = true;
-		}
-	}
-	else
+bool UMVVMView::SetViewModel(FName ViewModelName, TScriptInterface<INotifyFieldValueChanged> NewValue)
+{
+	if (ViewModelName.IsNone())
 	{
-		if (!const_cast<UMVVMView*>(this)->EvaluateSourceCreator(Binding.GetEvaluateSourceCreatorBindingIndex()))
-		{
-			UE::MVVM::FMessageLog Log(GetUserWidget());
-			Log.Error(FText::Format(LOCTEXT("ExecuteBindingFailEvaluate", "The evaluate source creator binding '{0}' was not executed.")
-#if UE_WITH_MVVM_DEBUGGING
-				, FText::FromString(Binding.ToString(ClassExtension->GetBindingLibrary(), FMVVMViewClass_CompiledBinding::FToStringArgs::Short()))
-#else
-				, FText::AsNumber(BindingIndex)
-#endif
-			));
-			bError = true;
-		}
+		UE::MVVM::FMessageLog Log(GetUserWidget());
+		Log.Error(LOCTEXT("SetViewModelInvalidName", "The viewmodel name is empty."));
+		return false;
 	}
 
-	if (bLogBinding && !bError)
+	if (ClassExtension == nullptr)
 	{
 		UE::MVVM::FMessageLog Log(GetUserWidget());
-		Log.Info(FText::Format(LOCTEXT("ExecuteBindingGeneric", "Execute binding '{0}'.")
-#if UE_WITH_MVVM_DEBUGGING
-			, FText::FromString(Binding.ToString(ClassExtension->GetBindingLibrary(), FMVVMViewClass_CompiledBinding::FToStringArgs::All()))
-#else
-			, FText::AsNumber(BindingIndex)
-#endif
-		));
-
-#if UE_WITH_MVVM_DEBUGGING
-		UE::MVVM::FDebugging::BroadcastLibraryBindingExecuted(this, Binding);
-#endif
+		Log.Error(LOCTEXT("SetViewModelInvalidClass", "The view is not constructed."));
+		return false;
 	}
-}
 
+	const int32 ClassSourceIndex = ClassExtension->GetSources().IndexOfByPredicate([ViewModelName](const FMVVMViewClass_Source& Other)
+		{
+			return Other.GetName() == ViewModelName;
+		});
 
-//void UMVVMView::SetLibraryBindingEnabled(FGuid ViewModelId, FMVVMBindingName BindingName, bool bEnable)
-//{
-//	if (ensure(ClassExtension))
-//	{
-//		int32 LibraryBindingIndex = ClassExtension->IndexOfCompiledBinding(ViewModelId, BindingName);
-//		if (LibraryBindingIndex != INDEX_NONE)
-//		{
-//			const bool bIsEnabled = IsLibraryBindingEnabled(LibraryBindingIndex);
-//			if (bIsEnabled != bEnable)
-//			{
-//				const FMVVMViewClass_CompiledBinding& Item = ClassExtension->GetCompiledBindings(LibraryBindingIndex);
-//				if (bEnable)
-//				{
-//					EnableLibraryBinding(Item, LibraryBindingIndex);
-//				}
-//				else
-//				{
-//					DisableLibraryBinding(Item, LibraryBindingIndex);
-//				}
-//			}
-//		}
-//	}
-//}
-
-
-//bool UMVVMView::IsLibraryBindingEnabled(FGuid ViewModelId, FMVVMBindingName BindingName) const
-//{
-//	if (ensure(ClassExtension))
-//	{
-//		int32 LibraryBindingIndex = ClassExtension->IndexOfCompiledBinding(ViewModelId, BindingName);
-//		return IsLibraryBindingEnabled(LibraryBindingIndex);
-//	}
-//	return false;
-//}
-
-
-bool UMVVMView::IsLibraryBindingEnabled(int32 InBindindIndex) const
-{
-	return ensure(RegisteredLibraryBindings.IsValidIndex(InBindindIndex)) && RegisteredLibraryBindings[InBindindIndex].IsValid();
+	if (ClassSourceIndex == INDEX_NONE)
+	{
+		UE::MVVM::FMessageLog Log(GetUserWidget());
+		Log.Error(LOCTEXT("SetViewModelViewModelNameNotFound", "The viewmodel name could not be found."));
+		return false;
+	}
+	return SetSourceInternal(FMVVMViewClass_SourceKey(ClassSourceIndex), NewValue, false);
 }
 
 
-void UMVVMView::EnableLibraryBinding(const FMVVMViewClass_CompiledBinding& Binding, int32 BindingIndex, bool bSourceChanges, bool bForceDelayed)
+bool UMVVMView::SetViewModelByClass(TScriptInterface<INotifyFieldValueChanged> NewValue)
 {
-	check(!RegisteredLibraryBindings[BindingIndex].IsValid());
-
-	EMVVMExecutionMode ExecutionMode = Binding.GetExecuteMode();
-	bool bCanExecute = true;
-	bool bRegistered = false;
-	if (!Binding.IsOneTime() && ExecutionMode != EMVVMExecutionMode::Tick)
+	if (!NewValue.GetObject())
 	{
-		RegisteredLibraryBindings[BindingIndex] = RegisterLibraryBinding(Binding, BindingIndex);
-		bRegistered = RegisteredLibraryBindings[BindingIndex].IsValid();
-		bCanExecute = bRegistered;
+		UE::MVVM::FMessageLog Log(GetUserWidget());
+		Log.Error(LOCTEXT("SetViewModelInvalidObject", "The new viewmodel is null."));
+		return false;
 	}
 
-	if (bRegistered && bLogBinding)
+	if (ClassExtension == nullptr)
 	{
 		UE::MVVM::FMessageLog Log(GetUserWidget());
-		Log.Info(FText::Format(LOCTEXT("EnableLibraryBinding", "Enable binding '{0}'.")
-#if UE_WITH_MVVM_DEBUGGING
-			, FText::FromString(Binding.ToString(ClassExtension->GetBindingLibrary(), FMVVMViewClass_CompiledBinding::FToStringArgs::All()))
-#else
-			, FText::AsNumber(BindingIndex)
-#endif
-		));
+		Log.Error(LOCTEXT("SetViewModelInvalidClass", "The view is not initialized."));
+		return false;
 	}
 
-	if (bCanExecute)
+	int32 FoundSourceIndex = INDEX_NONE;
+	int32 SourceIndex = 0;
+	for (const FMVVMViewClass_Source& ClassSource : ClassExtension->GetSources())
 	{
-		// @TODO: Asil.Karatas - Check if bSourceChanges should be guarded with Binding.NeedsExecuteAtInitializationWhenSourceChanges()
-		if (Binding.NeedsExecutionAtInitialization() || bSourceChanges) 
+		if (NewValue.GetObject()->GetClass()->IsChildOf(ClassSource.GetSourceClass()))
 		{
-			if (!bForceDelayed)
+			if (FoundSourceIndex == INDEX_NONE)
 			{
-				ExecuteLibraryBinding(Binding, BindingIndex);
+				FoundSourceIndex = SourceIndex;
 			}
 			else
 			{
-				GEngine->GetEngineSubsystem<UMVVMBindingSubsystem>()->AddDelayedBinding(this, BindingIndex);
+				UE::MVVM::FMessageLog Log(GetUserWidget());
+				Log.Error(FText::Format(LOCTEXT("SetViewModelViewModelNotUnique", "More than one viewmodel match the type of the passed instance. Make sure there exists only one viewmodel of type {0} in the widget blueprint, or try SetViewModel by name."), FText::FromName(NewValue.GetObject()->GetClass()->GetFName())));
+				return false;
 			}
 		}
+		++SourceIndex;
+	}
+
+	if (FoundSourceIndex == INDEX_NONE)
+	{
+		UE::MVVM::FMessageLog Log(GetUserWidget());
+		Log.Error(LOCTEXT("SetViewModelViewModelNotFound", "A created viewmodel matching the class of the passed instance could not be found."));
+		return false;
 	}
 
-	bHasEveryTickBinding = bHasEveryTickBinding || ExecutionMode == EMVVMExecutionMode::Tick;
+	return SetSourceInternal(FMVVMViewClass_SourceKey(FoundSourceIndex), NewValue, false);
 }
 
 
-void UMVVMView::DisableLibraryBinding(const FMVVMViewClass_CompiledBinding& Binding, int32 BindingIndex)
+bool UMVVMView::SetSourceInternal(FMVVMViewClass_SourceKey ClassSourceKey, TScriptInterface<INotifyFieldValueChanged> NewValue, bool bForDynamicSource)
 {
-	check(IsLibraryBindingEnabled(BindingIndex));
+	SCOPE_CYCLE_COUNTER(STAT_UMG_Viewmodel_SetSource);
 
-	UnregisterLibraryBinding(Binding, RegisteredLibraryBindings[BindingIndex], BindingIndex);
-	RegisteredLibraryBindings[BindingIndex] = FDelegateHandle();
+	check(Sources.IsValidIndex(ClassSourceKey.GetIndex()));
+	FMVVMView_SourceKey ViewSourceKey = FMVVMView_SourceKey(ClassSourceKey.GetIndex());
 
-	if (bLogBinding)
+	const FMVVMViewClass_Source& ClassSource = ClassExtension->GetSource(ClassSourceKey);
+	FMVVMView_Source& ViewSource = Sources[ViewSourceKey.GetIndex()];
+
+	if (bForDynamicSource && !ClassSource.CanBeEvaluated())
 	{
 		UE::MVVM::FMessageLog Log(GetUserWidget());
-		Log.Info(FText::Format(LOCTEXT("DisableLibraryBinding", "Disable binding '{0}'.")
-#if UE_WITH_MVVM_DEBUGGING
-			, FText::FromString(Binding.ToString(ClassExtension->GetBindingLibrary(), FMVVMViewClass_CompiledBinding::FToStringArgs::All()))
-#else
-			, FText::AsNumber(BindingIndex)
-#endif
-			));
+		Log.Error(FText::Format(LOCTEXT("SetViewModelCannotBeEvalaute", "The new viewmodel {0} cannot be evaluated again."), FText::FromName(ClassSource.GetName())));
+		return false;
 	}
-}
-
-
-FDelegateHandle UMVVMView::RegisterLibraryBinding(const FMVVMViewClass_CompiledBinding& Binding, int32 BindingIndex)
-{
-	check(ClassExtension);
-
-	auto LogMessage = [this, &Binding, BindingIndex](const FText& Message)
+	else if (!bForDynamicSource && !ClassSource.CanBeSet())
 	{
 		UE::MVVM::FMessageLog Log(GetUserWidget());
-		Log.Error(FText::Format(LOCTEXT("RegisterBindingFailed_Format", "Widget '{0}' can't register binding '{1}'. {2}")
-			, FText::FromString(GetFullName())
-#if UE_WITH_MVVM_DEBUGGING
-			, FText::FromString(Binding.ToString(ClassExtension->GetBindingLibrary(), FMVVMViewClass_CompiledBinding::FToStringArgs::Short()))
-#else
-			, FText::AsNumber(BindingIndex)
-#endif
-			, Message
-			));
-	};
-
-	TValueOrError<UE::FieldNotification::FFieldId, void> FieldIdResult = ClassExtension->GetBindingLibrary().GetFieldId(Binding.GetSourceFieldId());
-	if (FieldIdResult.HasError())
+		Log.Error(FText::Format(LOCTEXT("SetViewModelCannotBeSet", "The new viewmodel {0} cannot be set."), FText::FromName(ClassSource.GetName())));
+		return false;
+	}
+			
+	if (NewValue.GetObject()
+		&& !NewValue.GetObject()->GetClass()->IsChildOf(ClassSource.GetSourceClass()))
 	{
-		LogMessage(LOCTEXT("RegisterBindingFailed_InvalidFieldId", "The FieldId is invalid."));
-
-#if UE_WITH_MVVM_DEBUGGING
-		UE::MVVM::FDebugging::BroadcastLibraryBindingRegistered(this, Binding, UE::MVVM::FDebugging::ERegisterLibraryBindingResult::Failed_InvalidFieldId);
-#endif
-		return FDelegateHandle();
+		UE::MVVM::FMessageLog Log(GetUserWidget());
+		Log.Error(LOCTEXT("SetViewModelInvalidValueType", "The new viewmodel is not of the expected type."));
+		UE_LOG(LogMVVM, Error, TEXT("The viewmodel name '%s' is invalid for the view '%s'"), *ClassSource.GetName().ToString(), *GetFullName());
+		return false;
 	}
 
-	UE::FieldNotification::FFieldId FieldId = FieldIdResult.StealValue();
-	if (!FieldId.IsValid())
-	{
-		LogMessage(LOCTEXT("RegisterBindingFailed_FieldIdNotFound", "The FieldId was not found on the source."));
+	UObject* PreviousValue = ViewSource.Source;
+	bool bPreviousSourceInitialized = ViewSource.bSourceInitialized;
+	bool bPreviousBindingsInitialized = ViewSource.bBindingsInitialized;
 
-#if UE_WITH_MVVM_DEBUGGING
-		UE::MVVM::FDebugging::BroadcastLibraryBindingRegistered(this, Binding, UE::MVVM::FDebugging::ERegisterLibraryBindingResult::Failed_FieldIdNotFound);
-#endif
-		return FDelegateHandle();
+	{
+		// Sanity check. Test if the bitfield matches the cached value.
+		const bool bIsValid = PreviousValue != nullptr;
+		const bool bIsBitfieldValid = (ClassSourceKey.GetBit() & ValidSources) != 0;
+		ensureMsgf(bIsValid == bIsBitfieldValid, TEXT("The source %s should be valid."), *ClassSource.GetName().ToString());
 	}
 
-	// The source may not have been created because the property path was wrong.
-	TScriptInterface<INotifyFieldValueChanged> Source = FindSource(Binding, BindingIndex, true);
-	if (Source.GetInterface() == nullptr)
+	if (PreviousValue != NewValue.GetObject())
 	{
-		if (!Binding.IsRegistrationOptional())
+		if (ViewSource.bBindingsInitialized)
 		{
-			LogMessage(LOCTEXT("RegisterBindingFailed_InvalidSourceObject", "The source object is invalid."));
+			UninitializeSourceBindings(ClassSourceKey, ClassSource, ViewSource);
+			GEngine->GetEngineSubsystem<UMVVMBindingSubsystem>()->RemoveDelayedBindings(this, ClassSourceKey);
 		}
+		UninitializeSource(ViewSourceKey);
 
-#if UE_WITH_MVVM_DEBUGGING
-		UE::MVVM::FDebugging::BroadcastLibraryBindingRegistered(this, Binding, UE::MVVM::FDebugging::ERegisterLibraryBindingResult::Failed_InvalidSource);
-#endif
-		return FDelegateHandle();
-	}
+		InitializeSourceInternal(NewValue.GetObject(), ClassSourceKey, ClassSource, ViewSource);
+		ViewSource.bSetManually = !bForDynamicSource;
+		ViewSource.bSourceInitialized = bPreviousSourceInitialized;
 
-	UE::FieldNotification::FFieldMulticastDelegate::FDelegate Delegate = UE::FieldNotification::FFieldMulticastDelegate::FDelegate::CreateUObject(this, &UMVVMView::HandledLibraryBindingValueChanged, BindingIndex);
-	FDelegateHandle ResultHandle = Source->AddFieldValueChangedDelegate(FieldId, MoveTemp(Delegate));
-	if (ResultHandle.IsValid())
-	{
-		if (FMVVMViewSource* RegisteredSource = FindViewSource(Binding.GetSourceName()))
+		// remove events, re-add the events.
+		if (bEventsInitialized)
 		{
-			++RegisteredSource->RegisteredCount;
+			ReinitializeEvents(ClassSourceKey, PreviousValue, NewValue.GetObject());
 		}
-		else
+
+		if (bPreviousBindingsInitialized)
 		{
-			// if it's a widget, it may not be in the Sources
-			FMVVMViewSource& NewSource = Sources.AddDefaulted_GetRef();
-			NewSource.Source = Source.GetObject();
-			NewSource.SourceName = Binding.GetSourceName();
-			NewSource.RegisteredCount = 1;
+			// If binding is A.B.Property, and A is set, we need to evaluate/update B
+			if (ClassSource.HasEvaluateBindings())
+			{
+				for (const FMVVMViewClass_EvaluateSource& ClassEvaluate : ClassExtension->GetEvaluateSources())
+				{
+					if (ClassEvaluate.GetParentSource() == ClassSourceKey)
+					{
+						EvaluateSource(ClassEvaluate.GetSource());
+					}
+				}
+			}
+
+			// initialize bindings and run time
+			InitializeSourceBindings(ViewSourceKey, true);
 		}
 
 #if UE_WITH_MVVM_DEBUGGING
-		UE::MVVM::FDebugging::BroadcastLibraryBindingRegistered(this, Binding, UE::MVVM::FDebugging::ERegisterLibraryBindingResult::Success);
+		UE::MVVM::FDebugging::BroadcastViewSourceValueChanged(this, ClassSourceKey, ViewSourceKey);
 #endif
 	}
-	else
-	{
-#if UE_WITH_MVVM_DEBUGGING
-		UE::MVVM::FDebugging::BroadcastLibraryBindingRegistered(this, Binding, UE::MVVM::FDebugging::ERegisterLibraryBindingResult::Failed_InvalidSourceField);
-#endif
-	}
-
-	return ResultHandle;
+	return true;
 }
 
 
-void UMVVMView::UnregisterLibraryBinding(const FMVVMViewClass_CompiledBinding& Binding, FDelegateHandle Handle, int32 BindingIndex)
+void UMVVMView::InitializeEvents()
 {
-	TScriptInterface<INotifyFieldValueChanged> Source = FindSource(Binding, BindingIndex, true);
-	if (Source.GetInterface() && ClassExtension)
+	if (bEventsInitialized)
 	{
-		TValueOrError<UE::FieldNotification::FFieldId, void> FieldIdResult = ClassExtension->GetBindingLibrary().GetFieldId(Binding.GetSourceFieldId());
-		if (ensureMsgf(FieldIdResult.HasValue() && FieldIdResult.GetValue().IsValid(), TEXT("If the binding was enabled then the FieldId should exist.")))
-		{
-			Source->RemoveFieldValueChangedDelegate(FieldIdResult.GetValue(), Handle);
-		}
-
-		FMVVMViewSource* RegisteredSource = FindViewSource(Binding.GetSourceName());
-		if (ensureMsgf(RegisteredSource, TEXT("If the binding was enabled, then the source should also be there.")))
-		{
-			ensureMsgf(RegisteredSource->RegisteredCount > 0, TEXT("The count should match the number of RegisterLibraryBinding and UnregisterLibraryBinding"));
-			--RegisteredSource->RegisteredCount;
-		}
+		return;
 	}
-}
-
 
-void UMVVMView::InitializeEvents()
-{
 	UUserWidget* UserWidget = GetUserWidget();
 	check(UserWidget);
 	check(ClassExtension);
 
-	const TArrayView<const FMVVMViewClass_CompiledEvent>& CompiledEvent = ClassExtension->GetCompiledEvents();
+	const TArrayView<const FMVVMViewClass_Event>& ClassEvents = ClassExtension->GetEvents();
 
-	ensure(RegisteredLibraryEvents.IsEmpty());
-	RegisteredLibraryEvents.Reset();
-	RegisteredLibraryEvents.Add(false, CompiledEvent.Num());
-
-	for (int32 Index = 0; Index < CompiledEvent.Num(); ++Index)
+	ensure(BoundEvents.IsEmpty());
+	BoundEvents.Reset(ClassEvents.Num());
+	for (int32 Index = 0; Index < ClassEvents.Num(); ++Index)
 	{
-		const FMVVMViewClass_CompiledEvent& Event = CompiledEvent[Index];
-		EnableLibraryEvent(Event, Index);
+		const FMVVMViewClass_Event& ClassEvent = ClassEvents[Index];
+		BindEvent(ClassEvent, FMVVMViewClass_EventKey(Index));
 	}
+
+	bEventsInitialized = true;
 }
 
 
 void UMVVMView::UninitializeEvents()
 {
+	if (!bEventsInitialized)
+	{
+		return;
+	}
+
 	check(ClassExtension);
+	
+	bEventsInitialized = false;
 
-	const TArrayView<const FMVVMViewClass_CompiledEvent>& CompiledEvent = ClassExtension->GetCompiledEvents();
-	for (int32 Index = 0; Index < CompiledEvent.Num(); ++Index)
+	for (int32 Index = BoundEvents.Num() - 1; Index >= 0; --Index)
 	{
-		const FMVVMViewClass_CompiledEvent& Event = CompiledEvent[Index];
-		if (RegisteredLibraryEvents[Index])
-		{
-			DisableLibraryEvent(Event, Index);
-		}
+		UnbindEvent(Index);
 	}
 }
 
 
-void UMVVMView::EnableLibraryEvent(const FMVVMViewClass_CompiledEvent& Event, int32 EventIndex)
+void UMVVMView::BindEvent(const FMVVMViewClass_Event& ClassEvent, FMVVMViewClass_EventKey KeyForLog)
 {
 	UUserWidget* UserWidget = GetUserWidget();
 	check(UserWidget);
 
-	auto LogMessage = [this, &Event, EventIndex](const FText& Message)
-	{
-		UE::MVVM::FMessageLog Log(GetUserWidget());
-		Log.Error(FText::Format(LOCTEXT("EnableLibraryEventFailed_Format", "Widget '{0}' can't register event '{1}'. {2}")
-			, FText::FromString(GetFullName())
+	auto LogMessage = [Self=this, UserWidget, &ClassEvent, KeyForLog](const FText& Message)
+		{
+			UE::MVVM::FMessageLog Log(UserWidget);
+			Log.Error(FText::Format(LOCTEXT("EnableLibraryEventFailed_Format", "Widget '{0}' can't register event '{1}'. {2}")
+				, FText::FromString(Self->GetFullName())
 #if UE_WITH_MVVM_DEBUGGING
-			, FText::FromString(Event.ToString(ClassExtension->GetBindingLibrary(), FMVVMViewClass_CompiledEvent::FToStringArgs()))
+				, FText::FromString(ClassEvent.ToString(Self->ClassExtension, FMVVMViewClass_Event::FToStringArgs()))
 #else
-			, FText::AsNumber(EventIndex)
+				, FText::AsNumber(KeyForLog.GetIndex())
 #endif
-			, Message
+				, Message
 			));
-	};
+		};
 
-	const UFunction* FunctionToBind = UserWidget->GetClass()->FindFunctionByName(Event.GetUserWidgetFunctionName());
+	const UFunction* FunctionToBind = UserWidget->GetClass()->FindFunctionByName(ClassEvent.GetUserWidgetFunctionName());
 	if (FunctionToBind == nullptr)
 	{
 		LogMessage(LOCTEXT("EnableLibraryEventFailed_FunctionNotFound", "Function not found"));
@@ -1007,17 +894,29 @@ void UMVVMView::EnableLibraryEvent(const FMVVMViewClass_CompiledEvent& Event, in
 	}
 
 	const FMVVMCompiledBindingLibrary& Library = ClassExtension->GetBindingLibrary();
-	TValueOrError<UE::MVVM::FFieldContext, void> FieldPathResult = Library.EvaluateFieldPath(UserWidget, Event.GetMulticastDelegatePath());
+	TValueOrError<UE::MVVM::FFieldContext, void> FieldPathResult = Library.EvaluateFieldPath(UserWidget, ClassEvent.GetMulticastDelegatePath());
 	if (FieldPathResult.HasError())
 	{
-		LogMessage(LOCTEXT("EnableLibraryEventFailed_CantEvaluateBindingPath", "The binding path couldn't be evaluated."));
+		bool bIsOptional = ClassEvent.GetSourceKey().IsValid() && (ClassExtension->GetOptionalSources() & ClassEvent.GetSourceKey().GetBit()) != 0;
+		if (!bIsOptional)
+		{
+			// Problem. The path can be long Viewmodel.ObjectA.ObjectB.Multicast.
+			//It should warn but will not if the Viewmodel is valid but ObjectA is invalid.
+			LogMessage(LOCTEXT("EnableLibraryEventFailed_CantEvaluateBindingPath", "The binding path couldn't be evaluated."));
+		}
 		return;
 	}
 
-	UE::MVVM::FFieldContext FieldContext = FieldPathResult.StealValue();
+	UE::MVVM::FFieldContext& FieldContext = FieldPathResult.GetValue();
 	if (FieldContext.GetObjectVariant().IsNull() || !FieldContext.GetObjectVariant().IsUObject())
 	{
-		LogMessage(LOCTEXT("EnableLibraryEventFailed_OwnerInvalid", "The owner is invalid."));
+		bool bIsOptional = ClassEvent.GetSourceKey().IsValid() && (ClassExtension->GetOptionalSources() & ClassEvent.GetSourceKey().GetBit()) != 0;
+		if (!bIsOptional)
+		{
+			// Problem. The path can be long Viewmodel.ObjectA.ObjectB.Multicast.
+			//It should warn but will not if the Viewmodel is valid but ObjectA is invalid.
+			LogMessage(LOCTEXT("EnableLibraryEventFailed_OwnerInvalid", "The owner is invalid."));
+		}
 		return;
 	}
 
@@ -1029,123 +928,71 @@ void UMVVMView::EnableLibraryEvent(const FMVVMViewClass_CompiledEvent& Event, in
 		return;
 	}
 
+	UObject* EventObject = FieldContext.GetObjectVariant().GetUObject();
+	FBoundEvent& Result = BoundEvents.AddDefaulted_GetRef();
+	Result.Object = EventObject;
+	Result.PropertyName = MulticastDelegateProp->GetFName();
+	Result.EventKey = KeyForLog;
+
 	FScriptDelegate Delegate;
-	Delegate.BindUFunction(UserWidget, Event.GetUserWidgetFunctionName());
-	MulticastDelegateProp->AddDelegate(MoveTemp(Delegate), FieldContext.GetObjectVariant().GetUObject());
-	RegisteredLibraryEvents[EventIndex] = true;
+	Delegate.BindUFunction(UserWidget, ClassEvent.GetUserWidgetFunctionName());
+	MulticastDelegateProp->AddDelegate(MoveTemp(Delegate), EventObject);
 }
 
 
-void UMVVMView::DisableLibraryEvent(const FMVVMViewClass_CompiledEvent& Event, int32 EventIndex)
+void UMVVMView::UnbindEvent(int32 Index)
 {
-	check(IsLibraryEventEnabled(EventIndex));
-
-	RegisteredLibraryEvents[EventIndex] = false;
-
-	UUserWidget* UserWidget = GetUserWidget();
-	check(UserWidget);
-	const FMVVMCompiledBindingLibrary& Library = ClassExtension->GetBindingLibrary();
+	const FBoundEvent& BoundEvent = BoundEvents[Index];
 
-	TValueOrError<UE::MVVM::FFieldContext, void> FieldPathResult = Library.EvaluateFieldPath(UserWidget, Event.GetMulticastDelegatePath());
-	if (FieldPathResult.HasError())
+	UObject* EventObject= BoundEvent.Object.Get();
+	if (EventObject == nullptr)
 	{
+		BoundEvents.RemoveAtSwap(Index);
 		return;
 	}
 
-	UE::MVVM::FFieldContext FieldContext = FieldPathResult.StealValue();
-	if (FieldContext.GetObjectVariant().IsNull())
+	FMulticastDelegateProperty* MulticastDelegateProp = CastField<FMulticastDelegateProperty>(EventObject->GetClass()->FindPropertyByName(BoundEvent.PropertyName));
+	if (!MulticastDelegateProp)
 	{
+		BoundEvents.RemoveAtSwap(Index);
 		return;
 	}
 
-	FMulticastDelegateProperty* MulticastDelegateProp = CastFieldChecked<FMulticastDelegateProperty>(FieldContext.GetFieldVariant().GetProperty());
-	FScriptDelegate Delegate;
-	Delegate.BindUFunction(UserWidget, Event.GetUserWidgetFunctionName());
-	MulticastDelegateProp->RemoveDelegate(Delegate, FieldContext.GetObjectVariant().GetUObject());
-}
+	UUserWidget* UserWidget = GetUserWidget();
+	check(UserWidget);
 
+	FScriptDelegate Delegate;
+	Delegate.BindUFunction(UserWidget, ClassExtension->GetEvent(BoundEvent.EventKey).GetUserWidgetFunctionName());
+	MulticastDelegateProp->RemoveDelegate(Delegate, EventObject);
 
-bool UMVVMView::IsLibraryEventEnabled(int32 EventIndex) const
-{
-	return ensure(RegisteredLibraryEvents.IsValidIndex(EventIndex)) && RegisteredLibraryEvents[EventIndex];
+	BoundEvents.RemoveAtSwap(Index);
 }
 
 
-TScriptInterface<INotifyFieldValueChanged> UMVVMView::FindSource(const FMVVMViewClass_CompiledBinding& Binding, int32 BindingIndex, bool bAllowNull) const
+void UMVVMView::ReinitializeEvents(FMVVMViewClass_SourceKey SourceKey, UObject* PreviousValue, UObject* NewValue)
 {
-	UUserWidget* UserWidget = GetUserWidget();
-	check(UserWidget);
-	check(ClassExtension);
-
-	if (Binding.IsSourceUserWidget())
-	{
-		return TScriptInterface<INotifyFieldValueChanged>(UserWidget);
-	}
-	else
+	if (BoundEvents.Num() > 0 && PreviousValue)
 	{
-		auto LogMessage = [this, &Binding, BindingIndex](const FText& Message)
+		FWeakObjectPtr WeakPreviousValue = PreviousValue;
+		for (int32 Index = BoundEvents.Num() - 1; Index >= 0; --Index)
 		{
-			UE::MVVM::FMessageLog Log(GetUserWidget());
-			Log.Error(FText::Format(LOCTEXT("FindSourceFailed_Format", "Widget '{0}' can't evaluate the source for binding '{1}'. {2}.")
-				, FText::FromString(GetFullName())
-#if UE_WITH_MVVM_DEBUGGING
-				, FText::FromString(Binding.ToString(ClassExtension->GetBindingLibrary(), FMVVMViewClass_CompiledBinding::FToStringArgs::Short()))
-#else
-				, FText::AsNumber(BindingIndex)
-#endif
-				, Message
-			));
-		};
-
-		UObject* Source = nullptr;
-
-		if (const FMVVMViewSource* RegisteredSource = FindViewSource(Binding.GetSourceName()))
-		{
-			Source = RegisteredSource->Source.Get();
-		}
-		else
-		{
-			const FObjectPropertyBase* SourceObjectProperty = CastField<FObjectPropertyBase>(UserWidget->GetClass()->FindPropertyByName(Binding.GetSourceName()));
-			if (SourceObjectProperty == nullptr)
-			{
-				LogMessage(LOCTEXT("FindSourceFailed_InvalidPropertyName", "The property name is invalid."));
-				return TScriptInterface<INotifyFieldValueChanged>();
-			}
-
-			Source = SourceObjectProperty->GetObjectPropertyValue_InContainer(UserWidget);
-		}
-
-		if (Source == nullptr)
-		{
-			if (!bAllowNull)
+			const FBoundEvent& BoundEvent = BoundEvents[Index];
+			if (BoundEvent.Object == WeakPreviousValue)
 			{
-				LogMessage(LOCTEXT("FindSourceFailed_InvalidSourceObject", "The source is an invalid object."));
+				UnbindEvent(Index);
 			}
-			return TScriptInterface<INotifyFieldValueChanged>();
 		}
+	}
 
-		if (!Source->Implements<UNotifyFieldValueChanged>())
+	const TArrayView<const FMVVMViewClass_Event>& ClassEvents = ClassExtension->GetEvents();
+	for (int32 Index = 0; Index < ClassEvents.Num(); ++Index)
+	{
+		const FMVVMViewClass_Event& ClassEvent = ClassEvents[Index];
+		if (ClassEvent.GetSourceKey() == SourceKey)
 		{
-			LogMessage(LOCTEXT("FindSourceFailed_SourceObjectNotImplements", "The source does not implements INotifyFieldValueChanged."));
-			return TScriptInterface<INotifyFieldValueChanged>();
+			BindEvent(ClassEvent, FMVVMViewClass_EventKey(Index));
 		}
-
-		return TScriptInterface<INotifyFieldValueChanged>(Source);
 	}
 }
 
-
-FMVVMViewSource* UMVVMView::FindViewSource(const FName SourceName)
-{
-	return Sources.FindByPredicate([SourceName](const FMVVMViewSource& Other){ return Other.SourceName == SourceName; });
-}
-
-
-const FMVVMViewSource* UMVVMView::FindViewSource(const FName SourceName) const
-{
-	return Sources.FindByPredicate([SourceName](const FMVVMViewSource& Other) { return Other.SourceName == SourceName; });
-}
-
-
 #undef LOCTEXT_NAMESPACE
-
diff --git a/Engine/Plugins/Runtime/ModelViewViewModel/Source/ModelViewViewModel/Private/View/MVVMViewClass.cpp b/Engine/Plugins/Runtime/ModelViewViewModel/Source/ModelViewViewModel/Private/View/MVVMViewClass.cpp
--- a/Engine/Plugins/Runtime/ModelViewViewModel/Source/ModelViewViewModel/Private/View/MVVMViewClass.cpp
+++ b/Engine/Plugins/Runtime/ModelViewViewModel/Source/ModelViewViewModel/Private/View/MVVMViewClass.cpp
@@ -29,16 +29,16 @@
 ///////////////////////////////////////////////////////////////////////
 // 
 ///////////////////////////////////////////////////////////////////////
-UObject* FMVVMViewClass_SourceCreator::CreateInstance(const UMVVMViewClass* InViewClass, UMVVMView* InView, UUserWidget* InUserWidget) const
+UObject* FMVVMViewClass_Source::GetOrCreateInstance(const UMVVMViewClass* InViewClass, UMVVMView* InView, UUserWidget* InUserWidget) const
 {
 	check(InViewClass);
 	check(InView);
 	check(InUserWidget);
 
 	UObject* Result = nullptr;
-	const bool bOptional = (Flags & (uint8)ESourceFlags::IsOptional) != 0;
+	const bool bOptional = IsOptional();
 
-	if ((Flags & (uint8)ESourceFlags::TypeCreateInstance) != 0)
+	if ((Flags & (uint16)EFlags::TypeCreateInstance) != 0)
 	{
 		if (ExpectedSourceType.Get() != nullptr)
 		{
@@ -50,7 +50,7 @@ UObject* FMVVMViewClass_SourceCreator::CreateInstance(const UMVVMViewClass* InVi
 			Log.Error(FText::Format(LOCTEXT("CreateInstanceCreateInstance", "The source '{0}' could not be created. The class is not loaded."), FText::FromName(PropertyName)));
 		}
 	}
-	else if ((Flags & (uint8)ESourceFlags::SelfReference) != 0)
+	else if (IsUserWidget())
 	{ 
 		Result = InUserWidget;
 	}
@@ -126,7 +126,7 @@ UObject* FMVVMViewClass_SourceCreator::CreateInstance(const UMVVMViewClass* InVi
 	return Result;
 }
 
-void FMVVMViewClass_SourceCreator::DestroyInstance(const UObject* ViewModel, const UMVVMView* View) const
+void FMVVMViewClass_Source::ReleaseInstance(const UObject* ViewModel, const UMVVMView* View) const
 {
 	if (Resolver)
 	{
@@ -135,7 +135,7 @@ void FMVVMViewClass_SourceCreator::DestroyInstance(const UObject* ViewModel, con
 }
 
 #if UE_WITH_MVVM_DEBUGGING
-FMVVMViewClass_SourceCreator::FToStringArgs FMVVMViewClass_SourceCreator::FToStringArgs::Short()
+FMVVMViewClass_Source::FToStringArgs FMVVMViewClass_Source::FToStringArgs::Short()
 {
 	FToStringArgs Result;
 	Result.bUseDisplayName = false;
@@ -144,13 +144,15 @@ FMVVMViewClass_SourceCreator::FToStringArgs FMVVMViewClass_SourceCreator::FToStr
 	return Result;
 }
 
-FMVVMViewClass_SourceCreator::FToStringArgs FMVVMViewClass_SourceCreator::FToStringArgs::All()
+FMVVMViewClass_Source::FToStringArgs FMVVMViewClass_Source::FToStringArgs::All()
 {
 	return FToStringArgs();
 }
 
-FString FMVVMViewClass_SourceCreator::ToString(const FMVVMCompiledBindingLibrary& BindingLibrary, FToStringArgs Args) const
+FString FMVVMViewClass_Source::ToString(const UMVVMViewClass* ViewClass, FToStringArgs Args) const
 {
+	check(ViewClass);
+
 	TStringBuilder<512> StringBuilder;
 	StringBuilder << TEXT("Type: ");
 #if WITH_EDITOR
@@ -165,13 +167,7 @@ FString FMVVMViewClass_SourceCreator::ToString(const FMVVMCompiledBindingLibrary
 	}
 	
 	StringBuilder << TEXT(", SourceName: ");
-	StringBuilder << GetSourceName();
-
-	if (!ParentSourceName.IsNone())
-	{
-		StringBuilder << TEXT(", ParentSourceName: ");
-		StringBuilder << GetParentSourceName();
-	}
+	StringBuilder << GetName();
 
 	if (Args.bAddFlags)
 	{
@@ -186,10 +182,20 @@ FString FMVVMViewClass_SourceCreator::ToString(const FMVVMCompiledBindingLibrary
 		};
 
 		StringBuilder << TEXT(", Flags: ");
-		if ((Flags & (uint8)ESourceFlags::IsOptional) != 0)
+		if (IsUserWidget())
 		{
 			AddPipe();
-			StringBuilder << TEXT("Optional");
+			StringBuilder << TEXT("UserWidget");
+		}
+		if (IsUserWidgetProperty())
+		{
+			AddPipe();
+			StringBuilder << TEXT("Property");
+		}
+		if (RequireSettingUserWidgetProperty())
+		{
+			AddPipe();
+			StringBuilder << TEXT("SetProperty");
 		}
 		if (CanBeSet())
 		{
@@ -201,17 +207,32 @@ FString FMVVMViewClass_SourceCreator::ToString(const FMVVMCompiledBindingLibrary
 			AddPipe();
 			StringBuilder << TEXT("CanBeEvaluated");
 		}
-		if (IsSourceAUserWidgetProperty())
+		if (HasEvaluateBindings())
+		{
+			AddPipe();
+			StringBuilder << TEXT("HasEvaluateBindings");
+		}
+		if (IsOptional())
+		{
+			AddPipe();
+			StringBuilder << TEXT("Optional");
+		}
+		if (HasTickBindings())
+		{
+			AddPipe();
+			StringBuilder << TEXT("HasTickBindings");
+		}
+		if (IsViewModel())
 		{
 			AddPipe();
-			StringBuilder << TEXT("IsUserWidgetProperty");
+			StringBuilder << TEXT("IsViewModel");
 		}
 	}
 
 	if (Args.bAddCreationMode)
 	{
 		StringBuilder << TEXT("\n    CreationType: ");
-		if ((Flags & (uint8)ESourceFlags::TypeCreateInstance) != 0)
+		if ((Flags & (uint8)EFlags::TypeCreateInstance) != 0)
 		{
 			StringBuilder << TEXT("CreateInsance");
 		}
@@ -237,7 +258,7 @@ FString FMVVMViewClass_SourceCreator::ToString(const FMVVMCompiledBindingLibrary
 		else if (FieldPath.IsValid())
 		{
 			StringBuilder << TEXT("Path=");
-			TValueOrError<FString, FString> FieldPathString = BindingLibrary.FieldPathToString(FieldPath, Args.bUseDisplayName);
+			TValueOrError<FString, FString> FieldPathString = ViewClass->GetBindingLibrary().FieldPathToString(FieldPath, Args.bUseDisplayName);
 			if (FieldPathString.HasValue())
 			{
 				StringBuilder << FieldPathString.GetValue();
@@ -258,21 +279,6 @@ FString FMVVMViewClass_SourceCreator::ToString(const FMVVMCompiledBindingLibrary
 }
 #endif
 
-#if WITH_EDITOR
-void FMVVMViewClass_SourceCreator::PostSerialize(const FArchive& Ar)
-{
-	if (Ar.IsLoading())
-	{
-		if (Flags == 0)
-		{
-			Flags |= bCreateInstance_DEPRECATED ? (uint8)ESourceFlags::TypeCreateInstance : 0;
-			Flags |= bIsUserWidgetProperty_DEPRECATED ? (uint8)ESourceFlags::IsUserWidgetProperty : 0;
-			Flags |= bOptional_DEPRECATED ? (uint8)ESourceFlags::IsOptional : 0;
-		}
-	}
-}
-#endif
-
 ///////////////////////////////////////////////////////////////////////
 // 
 ///////////////////////////////////////////////////////////////////////
@@ -286,7 +292,7 @@ namespace UE::MVVM::Private
 	);
 }
 
-EMVVMExecutionMode FMVVMViewClass_CompiledBinding::GetExecuteMode() const
+EMVVMExecutionMode FMVVMViewClass_Binding::GetExecuteMode() const
 {
 	struct FLocal
 	{
@@ -302,35 +308,25 @@ EMVVMExecutionMode FMVVMViewClass_CompiledBinding::GetExecuteMode() const
 	};
 	static FLocal Local;
 
-	EMVVMExecutionMode Result = (Flags & EBindingFlags::OverrideExecuteMode) == 0 ? Local.DefaultMode : ExecutionMode;
-	return Result == EMVVMExecutionMode::DelayedWhenSharedElseImmediate ? (Binding.IsShared() ? EMVVMExecutionMode::Delayed : EMVVMExecutionMode::Immediate) : Result;
+	EMVVMExecutionMode Result = (Flags & (uint8)EFlags::OverrideExecuteMode) == 0 ? Local.DefaultMode : ExecutionMode;
+	return Result == EMVVMExecutionMode::DelayedWhenSharedElseImmediate ? (IsShared() ? EMVVMExecutionMode::Delayed : EMVVMExecutionMode::Immediate) : Result;
 }
 
 #if UE_WITH_MVVM_DEBUGGING
-FMVVMViewClass_CompiledBinding::FToStringArgs FMVVMViewClass_CompiledBinding::FToStringArgs::Short()
+FMVVMViewClass_Binding::FToStringArgs FMVVMViewClass_Binding::FToStringArgs::Short()
 {
 	FToStringArgs Result;
 	Result.bUseDisplayName = false;
-	Result.bAddFieldPath = false;
 	Result.bAddFlags = false;
 	return Result;
 }
 
-FMVVMViewClass_CompiledBinding::FToStringArgs FMVVMViewClass_CompiledBinding::FToStringArgs::All()
+FMVVMViewClass_Binding::FToStringArgs FMVVMViewClass_Binding::FToStringArgs::All()
 {
 	return FToStringArgs();
 }
 
-FString FMVVMViewClass_CompiledBinding::ToString() const
-{
-#if WITH_EDITOR
-	return EditorId.ToString();
-#else
-	return FString();
-#endif
-}
-
-FString FMVVMViewClass_CompiledBinding::ToString(const FMVVMCompiledBindingLibrary& BindingLibrary, FToStringArgs Args) const
+FString FMVVMViewClass_Binding::ToString(const UMVVMViewClass* ViewClass, FToStringArgs Args) const
 {
 	TStringBuilder<1024> StringBuilder;
 
@@ -343,7 +339,7 @@ FString FMVVMViewClass_CompiledBinding::ToString(const FMVVMCompiledBindingLibra
 	}
 #endif
 
-	if (Args.bAddFieldPath)
+	if (Args.bAddBindingFields)
 	{
 		if (GetBinding().IsValid())
 		{
@@ -353,20 +349,20 @@ FString FMVVMViewClass_CompiledBinding::ToString(const FMVVMCompiledBindingLibra
 			FString DestinationString;
 			FString ConversionString;
 			bool bErrorStrings = false;
-			if (!IsConversionFunctionComplex())
+			if (GetBinding().GetSourceFieldPath().IsValid())
 			{
-				TValueOrError<FString, FString> SourceFieldPathString = BindingLibrary.FieldPathToString(GetBinding().GetSourceFieldPath(), Args.bUseDisplayName);
+				TValueOrError<FString, FString> SourceFieldPathString = ViewClass->GetBindingLibrary().FieldPathToString(GetBinding().GetSourceFieldPath(), Args.bUseDisplayName);
 				SourceString = SourceFieldPathString.HasValue() ? SourceFieldPathString.StealValue() : SourceFieldPathString.StealError();
 				bErrorStrings = bErrorStrings || SourceFieldPathString.HasError();
 			}
 			{
-				TValueOrError<FString, FString> DestinationFieldPathString = BindingLibrary.FieldPathToString(GetBinding().GetDestinationFieldPath(), Args.bUseDisplayName);
+				TValueOrError<FString, FString> DestinationFieldPathString = ViewClass->GetBindingLibrary().FieldPathToString(GetBinding().GetDestinationFieldPath(), Args.bUseDisplayName);
 				DestinationString = DestinationFieldPathString.HasValue() ? DestinationFieldPathString.StealValue() : DestinationFieldPathString.StealError();
 				bErrorStrings = bErrorStrings || DestinationFieldPathString.HasError();
 			}
 			if (GetBinding().GetConversionFunctionFieldPath().IsValid())
 			{
-				TValueOrError<FString, FString> ConversionFieldPathString = BindingLibrary.FieldPathToString(GetBinding().GetConversionFunctionFieldPath(), Args.bUseDisplayName);
+				TValueOrError<FString, FString> ConversionFieldPathString = ViewClass->GetBindingLibrary().FieldPathToString(GetBinding().GetConversionFunctionFieldPath(), Args.bUseDisplayName);
 				ConversionString = ConversionFieldPathString.HasValue() ? ConversionFieldPathString.StealValue() : ConversionFieldPathString.StealError();
 				bErrorStrings = bErrorStrings || ConversionFieldPathString.HasError();
 			}
@@ -385,47 +381,49 @@ FString FMVVMViewClass_CompiledBinding::ToString(const FMVVMCompiledBindingLibra
 			{
 				StringBuilder << TEXT('"');
 			}
-			StringBuilder << TEXT(" = ");
-			if (GetBinding().GetConversionFunctionFieldPath().IsValid())
+
+			if (GetBinding().IsRuntimeBinding())
 			{
-				if (Args.bUseDisplayName)
-				{
-					StringBuilder << TEXT('"');
-				}
-				StringBuilder << ConversionString;
-				if (Args.bUseDisplayName)
+				StringBuilder << TEXT(" = ");
+				if (GetBinding().GetConversionFunctionFieldPath().IsValid())
 				{
-					StringBuilder << TEXT('"');
+					if (Args.bUseDisplayName)
+					{
+						StringBuilder << TEXT('"');
+					}
+					StringBuilder << ConversionString;
+					if (Args.bUseDisplayName)
+					{
+						StringBuilder << TEXT('"');
+					}
+					StringBuilder << TEXT(" ( ");
 				}
-				StringBuilder << TEXT(" ( ");
-			}
 
-			if (!IsConversionFunctionComplex())
-			{
-				if (Args.bUseDisplayName)
+				if (!GetBinding().HasComplexConversionFunction())
 				{
-					StringBuilder << TEXT('"');
+					if (Args.bUseDisplayName)
+					{
+						StringBuilder << TEXT('"');
+					}
+					StringBuilder << SourceString;
+					if (Args.bUseDisplayName)
+					{
+						StringBuilder << TEXT('"');
+					}
 				}
-				StringBuilder << SourceString;
-				if (Args.bUseDisplayName)
+
+				if (GetBinding().GetConversionFunctionFieldPath().IsValid())
 				{
-					StringBuilder << TEXT('"');
+					StringBuilder << TEXT(" )");
 				}
 			}
-
-			if (GetBinding().GetConversionFunctionFieldPath().IsValid())
+			else
 			{
-				StringBuilder << TEXT(" )");
+				StringBuilder << TEXT("()");
 			}
 
 			StringBuilder << TEXT("\n    ");
 		}
-		else if (Args.bAddFieldPath && IsEvaluateSourceCreatorBinding())
-		{
-			StringBuilder << TEXT("Evaluate Source Creator: [");
-			StringBuilder << EvaluateSourceCreatorIndex;
-			StringBuilder << TEXT("]\n    ");
-		}
 		else
 		{
 			StringBuilder << TEXT("!!Invalid Binding Error!!");
@@ -433,22 +431,12 @@ FString FMVVMViewClass_CompiledBinding::ToString(const FMVVMCompiledBindingLibra
 		}
 	}
 
-	StringBuilder << TEXT("SourceName: ");
-	StringBuilder << GetSourceName();
-
-	if (Args.bAddFieldId)
-	{
-		if (GetSourceFieldId().IsValid())
-		{
-			TValueOrError<UE::FieldNotification::FFieldId, void> SourceFieldId = BindingLibrary.GetFieldId(GetSourceFieldId());
-			StringBuilder << TEXT(", FieldId: ");
-			StringBuilder << (SourceFieldId.HasValue() ? SourceFieldId.GetValue().GetName() : FName());
-		}
-	}
+	StringBuilder << TEXT("Sources: ");
+	StringBuilder << SourceBitField;
 
 	if (Args.bAddFlags)
 	{
-		if ((Flags & EBindingFlags::OverrideExecuteMode) != 0)
+		if ((Flags & (uint8)EFlags::OverrideExecuteMode) != 0)
 		{
 			StringBuilder << TEXT(", Mode: ");
 			StringBuilder << StaticEnum<EMVVMExecutionMode>()->GetNameByValue((int64)ExecutionMode);
@@ -458,48 +446,49 @@ FString FMVVMViewClass_CompiledBinding::ToString(const FMVVMCompiledBindingLibra
 
 		bool bAddPipe = false;
 		auto AddPipe = [&StringBuilder, &bAddPipe]()
-		{
-			if (bAddPipe)
 			{
-				StringBuilder << TEXT('|');
-			}
-			bAddPipe = true;
-		};
+				if (bAddPipe)
+				{
+					StringBuilder << TEXT('|');
+				}
+				bAddPipe = true;
+			};
 
-		if ((Flags & EBindingFlags::ExecuteAtInitialization) != 0)
+		if ((Flags & (uint8)EFlags::OneWay) != 0)
 		{
 			AddPipe();
-			StringBuilder << TEXT("ExecAtInit");
+			StringBuilder << TEXT("OneWay");
 		}
-		if ((Flags & EBindingFlags::OneTime) != 0)
+		if ((Flags & (uint8)EFlags::Shared) != 0)
 		{
 			AddPipe();
-			StringBuilder << TEXT("OneTime");
+			StringBuilder << TEXT("Shared");
 		}
-		if ((Flags & EBindingFlags::EnabledByDefault) != 0)
+		if ((Flags & (uint8)EFlags::OverrideExecuteMode) != 0)
 		{
 			AddPipe();
-			StringBuilder << TEXT("EnabledByDefault");
+			StringBuilder << TEXT("OverrideExecuteMode");
 		}
-		if ((Flags & EBindingFlags::ViewModelOptional) != 0)
+		if ((Flags & (uint8)EFlags::EnabledByDefault) != 0)
 		{
 			AddPipe();
-			StringBuilder << TEXT("Optional");
+			StringBuilder << TEXT("Enabled");
 		}
-		if (IsConversionFunctionComplex())
+
+		if (GetBinding().IsComplexBinding())
 		{
 			AddPipe();
 			StringBuilder << TEXT("Complex");
 		}
-		if ((Flags & EBindingFlags::OverrideExecuteMode) != 0)
+		else if (GetBinding().HasComplexConversionFunction())
 		{
 			AddPipe();
-			StringBuilder << TEXT("OverrideExecutionMode");
+			StringBuilder << TEXT("ComplexConversion");
 		}
-		if ((Flags & EBindingFlags::SourceObjectIsSelf) != 0)
+		else if (GetBinding().HasSimpleConversionFunction())
 		{
 			AddPipe();
-			StringBuilder << TEXT("Self");
+			StringBuilder << TEXT("SimpleConversion");
 		}
 	}
 
@@ -511,7 +500,63 @@ FString FMVVMViewClass_CompiledBinding::ToString(const FMVVMCompiledBindingLibra
 // 
 ///////////////////////////////////////////////////////////////////////
 #if UE_WITH_MVVM_DEBUGGING
-FString FMVVMViewClass_CompiledEvent::ToString(const FMVVMCompiledBindingLibrary& BindingLibrary, FToStringArgs Args) const
+FMVVMViewClass_EvaluateSource::FToStringArgs FMVVMViewClass_EvaluateSource::FToStringArgs::Short()
+{
+	FToStringArgs Result;
+	Result.bUseDisplayName = false;
+	return Result;
+}
+
+FMVVMViewClass_EvaluateSource::FToStringArgs FMVVMViewClass_EvaluateSource::FToStringArgs::All()
+{
+	return FToStringArgs();
+}
+
+FString FMVVMViewClass_EvaluateSource::ToString(const UMVVMViewClass* ViewClass, FToStringArgs Args) const
+{
+	TStringBuilder<1024> StringBuilder;
+
+	if (ParentSource.IsValid())
+	{
+		StringBuilder << TEXT("Parent: ");
+		StringBuilder << ViewClass->GetSource(ParentSource).GetName();
+	}
+
+	if (ParentFieldId.IsValid())
+	{
+		StringBuilder << TEXT("FieldId: ");
+		StringBuilder << ParentFieldId.GetFieldName();
+	}
+
+	if (ToEvaluate.IsValid())
+	{
+		StringBuilder << TEXT("Source: ");
+		StringBuilder << ViewClass->GetSource(ToEvaluate).GetName();
+	}
+
+	return StringBuilder.ToString();
+}
+#endif
+
+///////////////////////////////////////////////////////////////////////
+// 
+///////////////////////////////////////////////////////////////////////
+#if UE_WITH_MVVM_DEBUGGING
+FMVVMViewClass_Event::FToStringArgs FMVVMViewClass_Event::FToStringArgs::Short()
+{
+	FToStringArgs Result;
+	Result.bUseDisplayName = false;
+	return Result;
+}
+
+FMVVMViewClass_Event::FToStringArgs FMVVMViewClass_Event::FToStringArgs::All()
+{
+	FToStringArgs Result;
+	Result.bUseDisplayName = true;
+	return Result;
+}
+
+FString FMVVMViewClass_Event::ToString(const UMVVMViewClass* ViewClass, FToStringArgs Args) const
 {
 	TStringBuilder<1024> StringBuilder;
 
@@ -524,7 +569,7 @@ FString FMVVMViewClass_CompiledEvent::ToString(const FMVVMCompiledBindingLibrary
 			StringBuilder << TEXT('"');
 		}
 
-		TValueOrError<FString, FString> SourceFieldPathString = BindingLibrary.FieldPathToString(GetMulticastDelegatePath(), Args.bUseDisplayName);
+		TValueOrError<FString, FString> SourceFieldPathString = ViewClass->GetBindingLibrary().FieldPathToString(GetMulticastDelegatePath(), Args.bUseDisplayName);
 		StringBuilder << (SourceFieldPathString.HasValue() ? SourceFieldPathString.StealValue() : SourceFieldPathString.StealError());
 
 		if (Args.bUseDisplayName)
@@ -533,8 +578,11 @@ FString FMVVMViewClass_CompiledEvent::ToString(const FMVVMCompiledBindingLibrary
 		}
 	}
 
-	StringBuilder << TEXT("SourceName: ");
-	StringBuilder << GetSourceName();
+	if (SourceToReevaluate.IsValid())
+	{
+		StringBuilder << TEXT("SourceName: ");
+		StringBuilder << ViewClass->GetSource(SourceToReevaluate).GetName();
+	}
 
 	StringBuilder << TEXT(", Function: ");
 	StringBuilder << GetUserWidgetFunctionName();
@@ -566,6 +614,24 @@ void UMVVMViewClass::Construct(UUserWidget* UserWidget)
 	if (ViewCounter == 1)
 	{
 		BindingLibrary.Load();
+
+		// Fixup the FieldIds
+		for (const FMVVMViewClass_Source& Source : Sources)
+		{
+			const TArrayView<const FMVVMViewClass_FieldId> SourceFieldIds = Source.GetFieldIds();
+			if (SourceFieldIds.Num() > 0)
+			{
+				const UClass* SourceClass = CastChecked<UClass>(Source.GetSourceClass());
+				check(SourceClass->ImplementsInterface(UNotifyFieldValueChanged::StaticClass()));
+				TScriptInterface<INotifyFieldValueChanged> Interface = SourceClass->GetDefaultObject();
+
+				for (const FMVVMViewClass_FieldId& SourceFieldId : SourceFieldIds)
+				{
+					UE::FieldNotification::FFieldId FieldId = Interface->GetFieldNotificationDescriptor().GetField(SourceClass, SourceFieldId.GetName());
+					const_cast<FMVVMViewClass_FieldId&>(SourceFieldId).FieldIndex = FieldId.GetIndex();
+				}
+			}
+		}
 	}
 
 #if WITH_EDITOR
@@ -618,7 +684,25 @@ void UMVVMViewClass::HandleBlueprintCompiled()
 #endif
 
 #if UE_WITH_MVVM_DEBUGGING
-void UMVVMViewClass::Log(FMVVMViewClass_SourceCreator::FToStringArgs SourceArgs, FMVVMViewClass_CompiledBinding::FToStringArgs BindingArgs) const
+UMVVMViewClass::FToStringArgs UMVVMViewClass::FToStringArgs::Short()
+{
+	FToStringArgs Result;
+	Result.Source = FMVVMViewClass_Source::FToStringArgs::Short();
+	Result.Binding = FMVVMViewClass_Binding::FToStringArgs::Short();
+	Result.Event = FMVVMViewClass_Event::FToStringArgs::Short();
+	return Result;
+}
+
+UMVVMViewClass::FToStringArgs UMVVMViewClass::FToStringArgs::All()
+{
+	FToStringArgs Result;
+	Result.Source = FMVVMViewClass_Source::FToStringArgs::All();
+	Result.Binding = FMVVMViewClass_Binding::FToStringArgs::All();
+	Result.Event = FMVVMViewClass_Event::FToStringArgs::All();
+	return Result;
+}
+
+FString UMVVMViewClass::ToString(FToStringArgs Args) const
 {
 	bool bPreviousLoad = ViewCounter > 0;
 	if (!bPreviousLoad)
@@ -627,36 +711,33 @@ void UMVVMViewClass::Log(FMVVMViewClass_SourceCreator::FToStringArgs SourceArgs,
 	}
 
 	TStringBuilder<2048> Builder;
-	Builder << TEXT("Compiled Sources for: ");
-	Builder << GetOutermost()->GetFName();
-	for (const FMVVMViewClass_SourceCreator& Sources : GetViewModelCreators())
+	Builder << TEXT("Sources: ");
+	for (const FMVVMViewClass_Source& Source : GetSources())
 	{
 		Builder << TEXT("\n");
-		Builder << Sources.ToString(GetBindingLibrary(), SourceArgs);
+		Builder << Source.ToString(this, Args.Source);
 	}
-	UE_LOG(LogMVVM, Log, TEXT("%s"), Builder.ToString());
 
-	Builder.Reset();
-
-	Builder << TEXT("Compiled Bindings for: ");
-	Builder << GetOutermost()->GetFName();
-	for (const FMVVMViewClass_CompiledBinding& Binding : GetCompiledBindings())
+	Builder << TEXT("Bindings: ");
+	for (const FMVVMViewClass_Binding& Binding : GetBindings())
 	{
 		Builder << TEXT("\n");
-		Builder << Binding.ToString(GetBindingLibrary(), BindingArgs);
+		Builder << Binding.ToString(this, Args.Binding);
 	}
-	UE_LOG(LogMVVM, Log, TEXT("%s"), Builder.ToString());
 
-	Builder.Reset();
+	Builder << TEXT("Evaluate: ");
+	for (const FMVVMViewClass_EvaluateSource& EvaluateSource : GetEvaluateSources())
+	{
+		Builder << TEXT("\n");
+		Builder << EvaluateSource.ToString(this, Args.Evaluate);
+	}
 
-	Builder << TEXT("Compiled Events for: ");
-	Builder << GetOutermost()->GetFName();
-	for (const FMVVMViewClass_CompiledEvent& Event : GetCompiledEvents())
+	Builder << TEXT("Events: ");
+	for (const FMVVMViewClass_Event& Event : GetEvents())
 	{
 		Builder << TEXT("\n");
-		Builder << Event.ToString(GetBindingLibrary(), FMVVMViewClass_CompiledEvent::FToStringArgs());
+		Builder << Event.ToString(this, Args.Event);
 	}
-	UE_LOG(LogMVVM, Log, TEXT("%s"), Builder.ToString());
 
 	if (!bPreviousLoad)
 	{
@@ -664,6 +745,8 @@ void UMVVMViewClass::Log(FMVVMViewClass_SourceCreator::FToStringArgs SourceArgs,
 		const_cast<UMVVMViewClass*>(this)->BindingLibrary.Unload();
 #endif
 	}
+
+	return Builder.ToString();
 }
 #endif
 
diff --git a/Engine/Plugins/Runtime/ModelViewViewModel/Source/ModelViewViewModel/Public/Bindings/MVVMCompiledBindingLibrary.h b/Engine/Plugins/Runtime/ModelViewViewModel/Source/ModelViewViewModel/Public/Bindings/MVVMCompiledBindingLibrary.h
--- a/Engine/Plugins/Runtime/ModelViewViewModel/Source/ModelViewViewModel/Public/Bindings/MVVMCompiledBindingLibrary.h
+++ b/Engine/Plugins/Runtime/ModelViewViewModel/Source/ModelViewViewModel/Public/Bindings/MVVMCompiledBindingLibrary.h
@@ -37,11 +37,6 @@ public:
 		return ClassOrScriptStruct ? NumberOfFunctions : 0;
 	}
 
-	int32 GetFieldIdNum() const
-	{
-		return ClassOrScriptStruct ? NumberOfFieldIds : 0;
-	}
-
 	FName GetPropertyName(TArrayView<FName> Names, int32 Index) const
 	{
 		check(Index < NumberOfProperties && Index >= 0);
@@ -56,13 +51,6 @@ public:
 		return Names[LibraryStartIndex + NumberOfProperties + Index];
 	}
 
-	FName GetFieldIdName(TArrayView<FName> Names, int32 Index) const
-	{
-		check(Index < NumberOfFieldIds && Index >= 0);
-		check(Names.IsValidIndex(LibraryStartIndex + NumberOfProperties + NumberOfFunctions + Index));
-		return Names[LibraryStartIndex + NumberOfProperties + NumberOfFunctions + Index];
-	}
-
 	const UStruct* GetStruct() const
 	{
 		return ClassOrScriptStruct;
@@ -74,9 +62,6 @@ public:
 	/** Find the UFunction from the Class.*/
 	UFunction* GetFunction(FName FunctionName) const;
 
-	/** Find the FFieldId from the Class.*/
-	UE::FieldNotification::FFieldId GetFieldId(FName FieldName) const;
-
 private:
 	UPROPERTY()
 	TObjectPtr<const UStruct> ClassOrScriptStruct = nullptr;
@@ -89,9 +74,6 @@ private:
 
 	UPROPERTY()
 	int16 NumberOfFunctions = 0;
-	
-	UPROPERTY()
-	int16 NumberOfFieldIds = 0;
 };
 
 
@@ -125,33 +107,6 @@ private:
 };
 
 
-/** Contains the FieldId index. */
-USTRUCT()
-struct FMVVMVCompiledFieldId
-{
-	GENERATED_BODY()
-
-	friend FMVVMCompiledBindingLibrary;
-	friend UE::MVVM::FCompiledBindingLibraryCompiler;
-
-	bool IsValid() const
-	{
-		return FieldIdIndex >= 0;
-	}
-
-private:
-	using IndexType = int16;
-
-	UPROPERTY()
-	int16 FieldIdIndex = INDEX_NONE;
-
-#if WITH_EDITORONLY_DATA
-	UPROPERTY(meta = (IgnoreForMemberInitializationTest))
-	FGuid CompiledBindingLibraryId;
-#endif
-};
-
-
 /** */
 USTRUCT()
 struct FMVVMCompiledLoadedPropertyOrFunctionIndex
@@ -199,9 +154,9 @@ struct FMVVMVCompiledBinding
 public:
 	bool IsValid() const
 	{
-		const bool bSourceValid = SourceFieldPath.IsValid() || IsComplexFunction();
+		const bool bSourceValid = SourceFieldPath.IsValid() || HasComplexConversionFunction() || IsComplexBinding();
 		const bool bDestinationValid = DestinationFieldPath.IsValid();
-		const bool bFunctionValid = ConversionFunctionFieldPath.IsValid() || !HasConversionFunction();
+		const bool bFunctionValid = ConversionFunctionFieldPath.IsValid() || (EType)Type == EType::None || IsComplexBinding();
 		return bSourceValid && bDestinationValid && bFunctionValid;
 	}
 
@@ -220,21 +175,44 @@ public:
 		return ConversionFunctionFieldPath;
 	}
 
-	bool HasConversionFunction() const
+	/**
+	 * A conversion function is needed to run the binding.
+	 * The binding runs in native.
+	 * It is on the form: Destination = ConversionFunction(Source)
+	 */
+	bool HasSimpleConversionFunction() const
+	{
+		return (EType)Type == EType::HasConversionFunction;
+	}
+	
+	/**
+	 * A conversion function is needed to run the binding.
+	 * The conversion function takes more than one input.
+	 * The binding doesn't requires the SourceFieldPath to execute.
+	 * The binding runs in native.
+	 * It is on the form Destination = ComplexConversionFunction()
+	 */
+	bool HasComplexConversionFunction() const
 	{
-		return (Flags & (uint8)EFlags::HasConversionFunction) != 0;
+		return (EType)Type == EType::HasComplexConversionFunction;
 	}
 
-	bool IsComplexFunction() const
+	/**
+	 * The binding runs in native.
+	 * It is on the form: Destination = Source or Destination = ConversionFunction(Source) or Destination = ComplexConversionFunction()
+	 */
+	bool IsRuntimeBinding() const
 	{
-		uint8 AllFunctionFlags = (uint8)EFlags::IsConversionFunctionComplex | (uint8)EFlags::HasConversionFunction;
-		return (Flags & AllFunctionFlags) == AllFunctionFlags;
+		return (EType)Type != EType::BindingComplex;
 	}
 
-	/** @return true if multiple field can trigger this binding. */
-	bool IsShared() const
+	/**
+	 * The binding runs with the BP virtual machine.
+	 * It is on the form: Destination()
+	 */
+	bool IsComplexBinding() const
 	{
-		return (Flags & (uint8)EFlags::IsShared) != 0;
+		return (EType)Type == EType::BindingComplex;
 	}
 
 private:
@@ -249,16 +227,16 @@ private:
 	UPROPERTY()
 	FMVVMVCompiledFieldPath ConversionFunctionFieldPath;
 
-	enum class EFlags : uint8
+	enum class EType : uint8
 	{
 		None = 0,
-		HasConversionFunction = 1 << 0,
-		IsConversionFunctionComplex = 1 << 1,
-		IsShared = 1 << 2,
+		HasConversionFunction = 1,
+		HasComplexConversionFunction = 2,
+		BindingComplex = 3,
 	};
 
 	UPROPERTY()
-	uint8 Flags = 0;
+	uint8 Type = (uint8)EType::None;
 
 #if WITH_EDITORONLY_DATA
 	UPROPERTY(meta = (IgnoreForMemberInitializationTest))
@@ -326,9 +304,6 @@ public:
 	 */
 	TValueOrError<UE::MVVM::FFieldContext, void> EvaluateFieldPath(UObject* ExecutionSource, const FMVVMVCompiledFieldPath& FieldPath) const;
 
-	/** Get the loaded FieldId. */
-	TValueOrError<UE::FieldNotification::FFieldId, void> GetFieldId(const FMVVMVCompiledFieldId& FieldId) const;
-
 	/** Return a readable version of the FieldPath. */
 	TValueOrError<FString, FString> FieldPathToString(FMVVMVCompiledFieldPath FieldPath, bool bUseDisplayName) const;
 
@@ -354,7 +329,6 @@ private:
 
 	TArray<FProperty*> LoadedProperties;
 	TArray<FLoadedFunction> LoadedFunctions;
-	TArray<UE::FieldNotification::FFieldId> LoadedFieldIds;
 
 	UPROPERTY()
 	TArray<FMVVMCompiledLoadedPropertyOrFunctionIndex> FieldPaths;
diff --git a/Engine/Plugins/Runtime/ModelViewViewModel/Source/ModelViewViewModel/Public/Debugging/MVVMDebugging.h b/Engine/Plugins/Runtime/ModelViewViewModel/Source/ModelViewViewModel/Public/Debugging/MVVMDebugging.h
--- a/Engine/Plugins/Runtime/ModelViewViewModel/Source/ModelViewViewModel/Public/Debugging/MVVMDebugging.h
+++ b/Engine/Plugins/Runtime/ModelViewViewModel/Source/ModelViewViewModel/Public/Debugging/MVVMDebugging.h
@@ -12,6 +12,7 @@
 #if UE_WITH_MVVM_DEBUGGING
 
 #include "Bindings/MVVMCompiledBindingLibrary.h"
+#include "View/MVVMViewTypes.h"
 
 class UMVVMView;
 struct FMVVMViewClass_CompiledBinding;
@@ -68,48 +69,22 @@ public:
 public:
 	struct FViewSourceValueArgs
 	{
-		FViewSourceValueArgs() = delete;
-		FViewSourceValueArgs(const FMVVMViewSource& ViewSource);
-		const FMVVMViewSource& ViewSource;
+		FViewSourceValueArgs(const FMVVMViewClass_SourceKey& Class, const FMVVMView_SourceKey& View);
+		FMVVMViewClass_SourceKey ClassSource;
+		FMVVMView_SourceKey ViewSource;
 	};
 
 	DECLARE_EVENT_TwoParams(FDebugging, FViewSourceValueChanged, const FView&, const FViewSourceValueArgs&);
 	/** Broadcast when a source changes. */
 	static FViewSourceValueChanged OnViewSourceValueChanged;
-	static void BroadcastViewSourceValueChanged(const UMVVMView* View, const FMVVMViewSource& ViewSource);
-
-public:
-	enum class ERegisterLibraryBindingResult : uint8
-	{
-		Success,
-		Failed_InvalidFieldId,
-		Failed_FieldIdNotFound,
-		Failed_InvalidSource,
-		Failed_InvalidSourceField,
-		Failed_Unknown,
-	};
-
-	struct FLibraryBindingRegisteredArgs
-	{
-		FLibraryBindingRegisteredArgs() = delete;
-		FLibraryBindingRegisteredArgs(const FMVVMViewClass_CompiledBinding& Binding, ERegisterLibraryBindingResult Result);
-		const FMVVMViewClass_CompiledBinding& Binding;
-		ERegisterLibraryBindingResult Result;
-	};
-
-	DECLARE_EVENT_TwoParams(FDebugging, FLibraryBindingRegistered, const FView&, const FLibraryBindingRegisteredArgs&);
-	/** Broadcast when a binding is registered to the view/viewmodel or failed to registered to the view/viewmodel. */
-	static FLibraryBindingRegistered OnLibraryBindingRegistered;
-	static void BroadcastLibraryBindingRegistered(const UMVVMView* View, const FLibraryBindingRegisteredArgs& Args);
-	static void BroadcastLibraryBindingRegistered(const UMVVMView* View, const FMVVMViewClass_CompiledBinding& Binding, ERegisterLibraryBindingResult Result);
+	static void BroadcastViewSourceValueChanged(const UMVVMView* View, const FMVVMViewClass_SourceKey ClassSourceKey, const FMVVMView_SourceKey ViewSourceKey);
 
 public:
 	struct FLibraryBindingExecutedArgs
 	{
-		FLibraryBindingExecutedArgs() = delete;
-		FLibraryBindingExecutedArgs(const FMVVMViewClass_CompiledBinding& Binding);
-		FLibraryBindingExecutedArgs(const FMVVMViewClass_CompiledBinding& Binding, FMVVMCompiledBindingLibrary::EExecutionFailingReason Result);
-		const FMVVMViewClass_CompiledBinding& Binding;
+		FLibraryBindingExecutedArgs(FMVVMViewClass_BindingKey Binding);
+		FLibraryBindingExecutedArgs(FMVVMViewClass_BindingKey Binding, FMVVMCompiledBindingLibrary::EExecutionFailingReason Result);
+		FMVVMViewClass_BindingKey Binding;
 		TOptional<FMVVMCompiledBindingLibrary::EExecutionFailingReason> FailingReason;
 	};
 
@@ -118,8 +93,8 @@ public:
 	/** Broadcast when a registered field is modified and a binding need to execute. */
 	static FLibraryBindingExecuted OnLibraryBindingExecuted;
 	static void BroadcastLibraryBindingExecuted(const UMVVMView* View, const FLibraryBindingExecutedArgs& Args);
-	static void BroadcastLibraryBindingExecuted(const UMVVMView* View, const FMVVMViewClass_CompiledBinding& Binding);
-	static void BroadcastLibraryBindingExecuted(const UMVVMView* View, const FMVVMViewClass_CompiledBinding& Binding, FMVVMCompiledBindingLibrary::EExecutionFailingReason Result);
+	static void BroadcastLibraryBindingExecuted(const UMVVMView* View, FMVVMViewClass_BindingKey Binding);
+	static void BroadcastLibraryBindingExecuted(const UMVVMView* View, FMVVMViewClass_BindingKey Binding, FMVVMCompiledBindingLibrary::EExecutionFailingReason Result);
 };
 } // UE::MVVM
 
diff --git a/Engine/Plugins/Runtime/ModelViewViewModel/Source/ModelViewViewModel/Public/View/MVVMView.h b/Engine/Plugins/Runtime/ModelViewViewModel/Source/ModelViewViewModel/Public/View/MVVMView.h
--- a/Engine/Plugins/Runtime/ModelViewViewModel/Source/ModelViewViewModel/Public/View/MVVMView.h
+++ b/Engine/Plugins/Runtime/ModelViewViewModel/Source/ModelViewViewModel/Public/View/MVVMView.h
@@ -4,6 +4,7 @@
 
 #include "CoreTypes.h"
 #include "Extensions/UserWidgetExtension.h"
+#include "View/MVVMViewTypes.h"
 
 #include "MVVMView.generated.h"
 
@@ -12,46 +13,44 @@ namespace UE::FieldNotification { struct FFieldId; }
 template <typename InterfaceType> class TScriptInterface;
 
 class UMVVMViewClass;
-struct FMVVMViewClass_CompiledBinding;
-struct FMVVMViewClass_CompiledEvent;
-struct FMVVMViewDelayedBinding;
-class UMVVMViewModelBase;
-class UWidget;
-namespace UE::MVVM
-{
-	class FDebugging;
-}
+struct FMVVMViewClass_Binding;
+struct FMVVMViewClass_Event;
+struct FMVVMViewClass_Source;
 
-/** */
+/**
+ * Instance FMVVMViewClass_Source for the UUserWdiget
+ */
 USTRUCT()
-struct FMVVMViewSource
+struct FMVVMView_Source
 {
 	GENERATED_BODY()
 
-	UPROPERTY(VisibleAnywhere, Category = "Viewmodel")
-	TObjectPtr<UObject> Source; // TScriptInterface<INotifyFieldValueChanged>
+	/** The source object. The source implement the INotifyFieldValueChanged interface. */
+	UPROPERTY(VisibleAnywhere, Category = "View")
+	TObjectPtr<UObject> Source;
 
-	UPROPERTY(VisibleAnywhere, Category = "Viewmodel")
-	FName SourceName;
+	/** The key of this source in the ViewClass. */
+	UPROPERTY(VisibleAnywhere, Category = "View")
+	FMVVMViewClass_SourceKey ClassKey;
 
-	// Number of bindings connected to the source.
-	UPROPERTY(VisibleAnywhere, Category = "Viewmodel")
+	/** Number of bindings connected to the source. */
+	UPROPERTY(VisibleAnywhere, Category = "View")
 	int32 RegisteredCount = 0;
 
-	// The source is defined in the ClassExtension.
-	UPROPERTY(VisibleAnywhere, Category = "Viewmodel")
-	bool bCreatedSource = false;
+	/** The source is created. */
+	UPROPERTY(VisibleAnywhere, Category = "View")
+	bool bSourceInitialized = false;
 
-	// The source bindings are initialized
-	UPROPERTY(VisibleAnywhere, Category = "Viewmodel")
+	/** The source bindings are initialized. */
+	UPROPERTY(VisibleAnywhere, Category = "View")
 	bool bBindingsInitialized = false;
 
-	// The source was set manually via SetViewModel.
-	UPROPERTY(VisibleAnywhere, Category = "Viewmodel")
+	/** The source was set manually via SetViewModel. */
+	UPROPERTY(VisibleAnywhere, Category = "View")
 	bool bSetManually = false;
 
-	// The source was set to a UserWidget property.
-	UPROPERTY(VisibleAnywhere, Category = "Viewmodel")
+	/** The source was set to a UserWidget property. */
+	UPROPERTY(VisibleAnywhere, Category = "View")
 	bool bAssignedToUserWidgetProperty = false;
 };
 
@@ -63,8 +62,6 @@ class MODELVIEWVIEWMODEL_API UMVVMView : public UUserWidgetExtension
 {
 	GENERATED_BODY()
 
-	friend UE::MVVM::FDebugging;
-
 public:
 	void ConstructView(const UMVVMViewClass* ClassExtension);
 
@@ -79,16 +76,16 @@ public:
 	 * Initializing the sources will also initialize the bindings if the option bInitializeBindingsOnConstruct is enabled.
 	 * @note A sources can be a viewmodel or any other object used by a binding.
 	 */
-	UFUNCTION(BlueprintCallable, Category = "Viewmodel")
+	UFUNCTION(BlueprintCallable, Category = "View")
 	void InitializeSources();
 	/**
 	 * Uninitialize the sources if they are already initialized.
 	 * It will uninitialized the bindings.
 	 */
-	UFUNCTION(BlueprintCallable, Category = "Viewmodel")
+	UFUNCTION(BlueprintCallable, Category = "View")
 	void UninitializeSources();
 	/** The sources were initialized, manually or automatically. */
-	UFUNCTION(BlueprintPure, Category = "Viewmodel")
+	UFUNCTION(BlueprintPure, Category = "View")
 	bool AreSourcesInitialized() const
 	{
 		return bSourcesInitialized;
@@ -98,33 +95,55 @@ public:
 	 * Initialize the bindings if they are not already initialized.
 	 * Initializing the bindings will execute them.
 	 */
-	UFUNCTION(BlueprintCallable, Category = "Viewmodel")
+	UFUNCTION(BlueprintCallable, Category = "View")
 	void InitializeBindings();
 	/** Uninitialize the bindings if they are already initialized. */
-	UFUNCTION(BlueprintCallable, Category = "Viewmodel")
+	UFUNCTION(BlueprintCallable, Category = "View")
 	void UninitializeBindings();
 	/** The bindings were initialized, manually or automatically. */
-	UFUNCTION(BlueprintPure, Category = "Viewmodel")
+	UFUNCTION(BlueprintPure, Category = "View")
 	bool AreBindingsInitialized() const
 	{
 		return bBindingsInitialized;
 	}
+	
+	/** Initialize the events if they are not already initialized. */
+	UFUNCTION(BlueprintCallable, Category = "View")
+	void InitializeEvents();
+	/** Uninitialize the events if they are already initialized. */
+	UFUNCTION(BlueprintCallable, Category = "View")
+	void UninitializeEvents();
+	/** The events were initialized, manually or automatically. */
+	UFUNCTION(BlueprintPure, Category = "View")
+	bool AreEventsInitialized() const
+	{
+		return bEventsInitialized;
+	}
 
+	/** The shared information for each instance of the view. */
 	const UMVVMViewClass* GetViewClass() const
 	{
 		return ClassExtension;
 	}
 
-	const TArrayView<const FMVVMViewSource> GetSources() const
+	/** The list of the sources needed by the view. */
+	const TArrayView<const FMVVMView_Source> GetSources() const
 	{
 		return Sources;
 	}
 
-	void ExecuteDelayedBinding(const FMVVMViewDelayedBinding& DelayedBinding) const;
-	void ExecuteEveryTickBindings() const;
+	/** The source used by the view. */
+	const FMVVMView_Source& GetSource(FMVVMView_SourceKey Key) const
+	{
+		check(Sources.IsValidIndex(Key.GetIndex()));
+		return Sources[Key.GetIndex()];
+	}
+
+	void ExecuteDelayedBinding(const FMVVMViewClass_BindingKey& DelayedBinding) const;
+	void ExecuteTickBindings() const;
 
 	/** Find and return the viewmodel with the specified name. */
-	UFUNCTION(BlueprintCallable, Category = "Viewmodel")
+	UFUNCTION(BlueprintCallable, Category = "View")
 	TScriptInterface<INotifyFieldValueChanged> GetViewModel(FName ViewModelName) const;
 
 	/**
@@ -132,7 +151,7 @@ public:
 	 * The viewmodel needs to be settable and the type should match (child of the defined viewmodel).
 	 * If the view is initialized, all bindings that uses that viewmodel will be re-executed with the new viewmodel instance.
 	 */
-	UFUNCTION(BlueprintCallable, Category = "Viewmodel")
+	UFUNCTION(BlueprintCallable, Category = "View")
 	bool SetViewModel(FName ViewModelName, TScriptInterface<INotifyFieldValueChanged> ViewModel);
 
 	/**
@@ -144,62 +163,74 @@ public:
 	bool SetViewModelByClass(TScriptInterface<INotifyFieldValueChanged> NewValue);
 
 private:
-	bool EvaluateSourceCreator(int32 SourceIndex);
-	bool SetSourceInternal(FName ViewModelName, TScriptInterface<INotifyFieldValueChanged> ViewModel, bool bForDynamicSource);
-	void UninitializeInternal(bool bSources);
-
-	void HandledLibraryBindingValueChanged(UObject* InViewModel, UE::FieldNotification::FFieldId InFieldId, int32 InCompiledBindingIndex) const;
-
-	void ExecuteLibraryBinding(const FMVVMViewClass_CompiledBinding& Binding, int32 BindingIndex) const;
-
-	void EnableLibraryBinding(const FMVVMViewClass_CompiledBinding& Item, int32 BindingIndex, bool bSourceChanges, bool bForceDelayed);
-	void DisableLibraryBinding(const FMVVMViewClass_CompiledBinding& Item, int32 BindingIndex);
-	bool IsLibraryBindingEnabled(int32 InBindindIndex) const;
-
-	FDelegateHandle RegisterLibraryBinding(const FMVVMViewClass_CompiledBinding& Binding, int32 BindingIndex);
-	void UnregisterLibraryBinding(const FMVVMViewClass_CompiledBinding& Binding, FDelegateHandle Handle, int32 BindingIndex);
-
-	void InitializeEvents();
-	void UninitializeEvents();
-	void EnableLibraryEvent(const FMVVMViewClass_CompiledEvent& Item, int32 EventIndex);
-	void DisableLibraryEvent(const FMVVMViewClass_CompiledEvent& Item, int32 EventIndex);
-	bool IsLibraryEventEnabled(int32 EventIndex) const;
-
-	TScriptInterface<INotifyFieldValueChanged> FindSource(const FMVVMViewClass_CompiledBinding& Binding, int32 BindingIndex, bool bAllowNull) const;
-	FMVVMViewSource* FindViewSource(const FName SourceName);
-	const FMVVMViewSource* FindViewSource(const FName SourceName) const;
+	//~ Source
+	void InitializeSource(FMVVMView_SourceKey SourceKey);
+	void InitializeSourceInternal(UObject* NewSource, FMVVMViewClass_SourceKey SourceKey, const FMVVMViewClass_Source& ClassSource, FMVVMView_Source& ViewSource);
+	void UninitializeSource(FMVVMView_SourceKey SourceKey);
+	bool SetSourceInternal(FMVVMViewClass_SourceKey SourceKey, TScriptInterface<INotifyFieldValueChanged> ViewModel, bool bForDynamicSource);
+
+	//~ Bindings
+	void InitializeSourceBindings(FMVVMView_SourceKey SourceKey, bool bRunAllBindings);
+	void InitializeSourceBindingsCommon();
+	void UninitializeSourceBindings(FMVVMViewClass_SourceKey SourceKey, const FMVVMViewClass_Source& ClassSource, FMVVMView_Source& ViewSource);
+	void HandledLibraryBindingValueChanged(UObject* InSource, UE::FieldNotification::FFieldId InFieldId);
+	void ExecuteBindingImmediately(const FMVVMViewClass_Binding& ClassBinding, FMVVMViewClass_BindingKey KeyForLog) const;
+
+	//~ evaluate source
+	bool EvaluateSource(FMVVMViewClass_SourceKey SourceIndex);
+
+	//~ events
+	void BindEvent(const FMVVMViewClass_Event& ClassItem, FMVVMViewClass_EventKey KeyForLog);
+	void UnbindEvent(int32 BoundEventIndex);
+	void ReinitializeEvents(FMVVMViewClass_SourceKey SourceKey, UObject* PreviousValue, UObject* NewValue);
 
 private:
 	UPROPERTY(Transient)
 	TObjectPtr<const UMVVMViewClass> ClassExtension;
 
-	UPROPERTY(VisibleAnywhere, Transient, Category = "Viewmodel")
-	TArray<FMVVMViewSource> Sources;
+	UPROPERTY(VisibleAnywhere, Transient, Category = "View")
+	TArray<FMVVMView_Source> Sources;
 
-	/** The binding that are registered by the view to the sources. */
-	TArray<FDelegateHandle> RegisteredLibraryBindings;
-	
+	struct FBoundEvent
+	{
+		FWeakObjectPtr Object;
+		FName PropertyName;
+		FMVVMViewClass_EventKey EventKey;
+	};
 	/** The event that are registered by the view to the sources. */
-	TBitArray<> RegisteredLibraryEvents;
+	TArray<FBoundEvent> BoundEvents;
+
+	/** The view has at least one binding that need to be ticked every frame. */
+	UPROPERTY(VisibleAnywhere, Transient, Category = "View")
+	uint64 ValidSources = 0;
+
+	/** The number of source has at least one binding that need to be ticked every frame. */
+	UPROPERTY(VisibleAnywhere, Transient, Category = "View")
+	uint8 NumberOfSourceWithTickBinding = 0;
 
 	/** Should log when a binding is executed. */
-	UPROPERTY(EditAnywhere, Transient, Category = "Viewmodel")
+	UPROPERTY(EditAnywhere, Transient, Category = "View")
 	bool bLogBinding = false;
 
 	/** Is the Construct method called. */
-	UPROPERTY(VisibleAnywhere, Transient, Category = "Viewmodel")
+	UPROPERTY(VisibleAnywhere, Transient, Category = "View")
 	bool bConstructed = false;
 	
 	/** Is the Initialize method called. */
-	UPROPERTY(VisibleAnywhere, Transient, Category = "Viewmodel")
+	UPROPERTY(VisibleAnywhere, Transient, Category = "View")
 	bool bSourcesInitialized = false;
+
 	/** Is the Initialize method called. */
-	UPROPERTY(VisibleAnywhere, Transient, Category = "Viewmodel")
+	UPROPERTY(VisibleAnywhere, Transient, Category = "View")
 	bool bBindingsInitialized = false;
+	
+	/** Is the Initialize method called. */
+	UPROPERTY(VisibleAnywhere, Transient, Category = "View")
+	bool bEventsInitialized = false;
 
-	/** The view has at least one binding that need to be ticked every frame. */
-	UPROPERTY(VisibleAnywhere, Transient, Category = "Viewmodel")
-	bool bHasEveryTickBinding = false;
+	/** Is the view is registered to the binding subsystem for tick. */
+	UPROPERTY(VisibleAnywhere, Transient, Category = "View")
+	bool bHasDefaultTickBinding = false;
 };
 
 #if UE_ENABLE_INCLUDE_ORDER_DEPRECATED_IN_5_2
diff --git a/Engine/Plugins/Runtime/ModelViewViewModel/Source/ModelViewViewModel/Public/View/MVVMViewClass.h b/Engine/Plugins/Runtime/ModelViewViewModel/Source/ModelViewViewModel/Public/View/MVVMViewClass.h
--- a/Engine/Plugins/Runtime/ModelViewViewModel/Source/ModelViewViewModel/Public/View/MVVMViewClass.h
+++ b/Engine/Plugins/Runtime/ModelViewViewModel/Source/ModelViewViewModel/Public/View/MVVMViewClass.h
@@ -2,318 +2,405 @@
 
 #pragma once
 
-
 #include "Bindings/MVVMCompiledBindingLibrary.h"
 #include "Extensions/WidgetBlueprintGeneratedClassExtension.h"
 #include "Types/MVVMExecutionMode.h"
 #include "Types/MVVMViewModelContext.h"
+#include "View/MVVMViewTypes.h"
 
 #include "UObject/Package.h"
 #include "MVVMViewClass.generated.h"
 
 
-class UMVVMUserWidgetBinding;
 class UMVVMView;
 class UMVVMViewClass;
-class UMVVMViewModelBlueprintExtension;
 class UMVVMViewModelContextResolver;
 class UUserWidget;
 
-namespace UE::MVVM::Private
-{
-	struct FMVVMViewBlueprintCompiler;
-}
+namespace UE::MVVM::Private { struct FMVVMViewBlueprintCompiler; }
+
 
 /**
- * Shared data to find or create a ViewModel at runtime.
+ * A structure to identify the Binding and the associated FieldId
  */
 USTRUCT()
-struct FMVVMViewClass_SourceCreator
+struct FMVVMViewClass_SourceBinding
 {
 	GENERATED_BODY()
 
 	friend UE::MVVM::Private::FMVVMViewBlueprintCompiler;
 
 public:
-	UObject* CreateInstance(const UMVVMViewClass* ViewClass, UMVVMView* View, UUserWidget* UserWidget) const;
-
-	void DestroyInstance(const UObject* ViewModel, const UMVVMView* View) const;
+	/**
+	 * The id for the FieldId on the source.
+	 * Valid when it is OneWay or when we need to register to the FieldNotify system.
+	 */
+	FFieldNotificationId GetFieldId() const
+	{
+		return FieldId;
+	}
 
-	UClass* GetSourceClass() const
+	/** The key to identify a binding in the view class. */
+	FMVVMViewClass_BindingKey GetBindingKey() const
 	{
-		return ExpectedSourceType.Get();
+		return BindingKey;
 	}
 
-	bool IsSourceAUserWidgetProperty() const
+	/** @return true if this Binding should be executed at initialization. */
+	bool ExecuteAtInitialization() const
 	{
-		return (Flags & (uint8)ESourceFlags::IsUserWidgetProperty) != 0;
+		return (Flags & (uint8)EFlags::ExecuteAtInitialization) != 0;
 	}
-	
-	bool CanBeSet() const
+
+private:
+	UPROPERTY(VisibleAnywhere, Category = "View")
+	FFieldNotificationId FieldId;
+
+	UPROPERTY(VisibleAnywhere, Category = "View")
+	FMVVMViewClass_BindingKey BindingKey;
+
+	enum class EFlags : uint8
 	{
-		return (Flags & (uint8)ESourceFlags::CanBeSet) != 0;
+		None = 0,
+		ExecuteAtInitialization = 1 << 0,
+	};
+
+	UPROPERTY(VisibleAnywhere, Category = "View")
+	uint8 Flags = (uint8)EFlags::None;
+};
+
+
+/**
+ * A compiled and shared binding for ViewModel<->View
+ */
+USTRUCT()
+struct FMVVMViewClass_Binding
+{
+	GENERATED_BODY()
+
+	friend UE::MVVM::Private::FMVVMViewBlueprintCompiler;
+
+public:
+	/** @return the binding. From source to destination (if forward) or from destination to source (if backward). */
+	const FMVVMVCompiledBinding& GetBinding() const
+	{
+		return Binding;
 	}
 
-	bool CanBeEvaluated() const
+	/** @return true if multiple field or source can trigger this binding. */
+	bool IsShared() const
 	{
-		return (Flags & (uint8)ESourceFlags::CanBeEvaluated) != 0;
+		return (Flags & (uint8)EFlags::Shared) != 0;
 	}
-	
-	bool IsOptional() const
+
+	/** @return true if this Binding should be executed once (at initialization) and when FieldId is broadcasted. */
+	bool IsOneWay() const
 	{
-		return (Flags & (uint8)ESourceFlags::IsOptional) != 0;
+		return (Flags & (uint8)EFlags::OneWay) != 0;
 	}
 
-	FName GetSourceName() const
+	/**
+	 * A view binding may require more than one view sources to run the binding.
+	 * A binding will not execute if any view source is invalid.
+	 * It will not warn if the view source is make as optional.
+	 */
+	uint64 GetSources() const
 	{
-		return PropertyName;
+		return SourceBitField;
 	}
-	
-	FName GetParentSourceName() const
+
+	/** How/when the binding should be executed. */
+	MODELVIEWVIEWMODEL_API EMVVMExecutionMode GetExecuteMode() const;
+
+#if WITH_EDITORONLY_DATA
+	/** Get the id of the BlueprintViewBinding. */
+	FGuid GetEditorId() const
 	{
-		return ParentSourceName;
+		return EditorId;
 	}
+#endif
 
 #if UE_WITH_MVVM_DEBUGGING
 	struct FToStringArgs
 	{
 		bool bUseDisplayName = true;
-		bool bAddCreationMode = true;
 		bool bAddFlags = true;
+		bool bAddBindingId = true;
+		bool bAddBindingFields = true;
 
 		MODELVIEWVIEWMODEL_API static FToStringArgs Short();
 		MODELVIEWVIEWMODEL_API static FToStringArgs All();
 	};
 
 	/** @return a human readable version of the source that can be use for debugging purposes. */
-	MODELVIEWVIEWMODEL_API FString ToString(const FMVVMCompiledBindingLibrary& CompiledBindingLibrary, FToStringArgs Args) const;
-#endif
-
-#if WITH_EDITOR
-	MODELVIEWVIEWMODEL_API void PostSerialize(const FArchive& Ar);
+	MODELVIEWVIEWMODEL_API FString ToString(const UMVVMViewClass* ViewClass, FToStringArgs Args) const;
 #endif
 
 private:
-	/** Class type to create a source at runtime. */
-	UPROPERTY()
-	TSubclassOf<UObject> ExpectedSourceType;
-
-	/** The resolver to fetch the source at runtime. */
-	UPROPERTY(Instanced)
-	TObjectPtr<UMVVMViewModelContextResolver> Resolver = nullptr;
+	enum class EFlags : uint8
+	{
+		None = 0,
+		OneWay = 1 << 0,
+		Shared = 1 << 1,
+		OverrideExecuteMode = 1 << 2,
+		EnabledByDefault = 1 << 3,
+	};
 
-	/** Info to find the ViewModel instance at runtime. */
-	UPROPERTY()
-	FMVVMViewModelContext GlobalViewModelInstance;
+	UPROPERTY(VisibleAnywhere, Category = "View")
+	FMVVMVCompiledBinding Binding;
 
-	/**
-	 * A resolvable path to retrieve the source instance at runtime.
-	 * It can be a path "Property = Object.Function.Object".
-	 * It can be a UFunction's name of a FProperty's name.
-	 */
-	UPROPERTY()
-	FMVVMVCompiledFieldPath FieldPath;
+	//~ FMVVMVCompiledBinding ends with a uint8, adding the byte variable here to help with packing.
 
-	/** The source name and the property's name of the view (if the flag IsUserWidgetProperty is set). */
-	UPROPERTY()
-	FName PropertyName;
+	UPROPERTY(VisibleAnywhere, Category = "View")
+	uint8 Flags = (uint8)EFlags::None;
 
-	/** The name of the parent source if it's the dynamic source. */
-	UPROPERTY()
-	FName ParentSourceName;
+	UPROPERTY(VisibleAnywhere, Category = "View")
+	EMVVMExecutionMode ExecutionMode = EMVVMExecutionMode::Immediate;
 
-	enum class ESourceFlags : uint8
-	{
-		None = 0,
-		TypeCreateInstance = 1 << 0,
-		IsUserWidgetProperty = 1 << 1,
-		IsOptional = 1 << 2,
-		CanBeSet = 1 << 3,
-		CanBeEvaluated = 1 << 4,
-		SelfReference = 1 << 5,
-	};
+	UPROPERTY(VisibleAnywhere, Category = "View")
+	uint64 SourceBitField;
 
-	UPROPERTY()
-	uint8 Flags = (uint8)ESourceFlags::None;
-	
 #if WITH_EDITORONLY_DATA
-	UPROPERTY()
-	bool bCreateInstance_DEPRECATED = false;
-
-	UPROPERTY()
-	bool bIsUserWidgetProperty_DEPRECATED = true;
-
-	UPROPERTY()
-	bool bOptional_DEPRECATED = false;
+	UPROPERTY(meta = (IgnoreForMemberInitializationTest))
+	FGuid EditorId;
 #endif
 };
 
-#if WITH_EDITOR
-template<>
-struct TStructOpsTypeTraits<FMVVMViewClass_SourceCreator> : public TStructOpsTypeTraitsBase2<FMVVMViewClass_SourceCreator>
+
+/**
+ * A binding to evaluate the source when something change in its path.
+ */
+USTRUCT()
+struct FMVVMViewClass_EvaluateSource
 {
-	enum
+	GENERATED_BODY()
+
+	friend UE::MVVM::Private::FMVVMViewBlueprintCompiler;
+
+public:
+	/** The id for the FieldId on the source on the parent that will trigger the evaluation. */
+	FFieldNotificationId GetFieldId() const
+	{
+		return ParentFieldId;
+	}
+	
+	/** The source that owns the source. */
+	FMVVMViewClass_SourceKey GetParentSource() const
+	{
+		return ParentSource;
+	}
+
+	/** The source that needs to be evaluated. */
+	FMVVMViewClass_SourceKey GetSource() const
 	{
-		WithPostSerialize = true,
+		return ToEvaluate;
+	}
+
+#if UE_WITH_MVVM_DEBUGGING
+	struct FToStringArgs
+	{
+		bool bUseDisplayName = true;
+
+		MODELVIEWVIEWMODEL_API static FToStringArgs Short();
+		MODELVIEWVIEWMODEL_API static FToStringArgs All();
 	};
-};
+
+	/** @return a human readable version of the source that can be use for debugging purposes. */
+	MODELVIEWVIEWMODEL_API FString ToString(const UMVVMViewClass* ViewClass, FToStringArgs Args) const;
 #endif
 
+private:
+	UPROPERTY(VisibleAnywhere, Category = "View")
+	FFieldNotificationId ParentFieldId;
+
+	UPROPERTY(VisibleAnywhere, Category = "View")
+	FMVVMViewClass_SourceKey ParentSource;
+
+	UPROPERTY(VisibleAnywhere, Category = "View")
+	FMVVMViewClass_SourceKey ToEvaluate;
+};
+
 
 /**
- * A compiled and shared binding for ViewModel<->View
+ * Shared data to find or create a ViewModel at runtime.
  */
 USTRUCT()
-struct MODELVIEWVIEWMODEL_API FMVVMViewClass_CompiledBinding
+struct FMVVMViewClass_Source
 {
 	GENERATED_BODY()
 
 	friend UE::MVVM::Private::FMVVMViewBlueprintCompiler;
 
 public:
-	/** @return The id for the FieldId on the source (if forward) or the destination (if backward). */
-	FMVVMVCompiledFieldId GetSourceFieldId() const
+	/** Get or create the instance used to register the bindings. */
+	MODELVIEWVIEWMODEL_API UObject* GetOrCreateInstance(const UMVVMViewClass* ViewClass, UMVVMView* View, UUserWidget* UserWidget) const;
+
+	/** The source is not needed anymore. */
+	MODELVIEWVIEWMODEL_API void ReleaseInstance(const UObject* ViewModel, const UMVVMView* View) const;
+
+	/** The expected class of the source. */
+	UClass* GetSourceClass() const
 	{
-		return FieldId;
+		return ExpectedSourceType.Get();
 	}
 
-	/**
-	 * @return The unique name of the source object that contains the SourceFieldId.
-	 * It implements INotifyFieldValueChanged if it's not a One Time.
-	 * It can also be the name of a FProperty on the source object.
-	 */
-	FName GetSourceName() const
+	/** The source is the UserWidget. */
+	bool IsUserWidget() const
 	{
-		return SourcePropertyName;
+		return (Flags & (uint16)EFlags::SelfReference) != 0;
 	}
 
-	/** @return The SourceCreator index when the binding is of the EvaluateSourceCreator type. */
-	int32 GetEvaluateSourceCreatorBindingIndex() const
+	/** The source is a viewmodel added in the editor. */
+	bool IsViewModel() const
 	{
-		return EvaluateSourceCreatorIndex;
+		return (Flags & (uint16)EFlags::IsViewModel) != 0;
 	}
 
-	/** @return true if SourceFieldId is not from a property but the object is the UserWidget. */
-	bool IsSourceUserWidget() const
+	/**
+	 * The source is UserWidget's property.
+	 * @return false if the source is from a long path or the source is the UserWidget.
+	 */
+	bool IsUserWidgetProperty() const
 	{
-		return (Flags & EBindingFlags::SourceObjectIsSelf) != 0;
+		return (Flags & (uint16)EFlags::IsUserWidgetProperty) != 0;
 	}
-
-	/** @return the binding. From source to destination (if forward) or from destination to source (if backward). */
-	const FMVVMVCompiledBinding& GetBinding() const
+	
+	/** The UserWidget's property need to be set/reset when the Source is created/release. */
+	bool RequireSettingUserWidgetProperty() const
 	{
-		return Binding;
+		return (Flags & (uint16)EFlags::SetUserWidgetProperty) != 0;
 	}
 
-	/** @return true if this Binding should be executed at initialization. */
-	bool NeedsExecutionAtInitialization() const
+	/** Can be set at runtime from SetViewModel. */
+	bool CanBeSet() const
 	{
-		return (Flags & EBindingFlags::ExecuteAtInitialization) != 0;
+		return (Flags & (uint16)EFlags::CanBeSet) != 0;
 	}
 
-	/** @return true if this Binding should be executed once (at initialization) but should not be executed when the SourceFieldId value changes. */
-	bool IsOneTime() const
+	/** The source "GetOrCreateInstance" can be reevaluated. */
+	bool CanBeEvaluated() const
 	{
-		return (Flags & EBindingFlags::OneTime) != 0;
+		return (Flags & (uint16)EFlags::CanBeEvaluated) != 0;
 	}
-
-	/** @return true if the binding is enabled by default. */
-	bool IsEnabledByDefault() const
+	
+	/** The source has at least one evaluate binding. */
+	bool HasEvaluateBindings() const
 	{
-		return (Flags & EBindingFlags::EnabledByDefault) != 0;
+		return (Flags & (uint16)EFlags::HasEvaluatedBindings) != 0;
 	}
-
-	/** @return true if it's normal that the binding could not find it's source when registering it. */
-	bool IsRegistrationOptional() const
+	
+	/**
+	 * The source GetOrCreateInstance can fail.
+	 * The view will not warn if a binding can't execute because the source is invalid.
+	 */
+	bool IsOptional() const
 	{
-		return (Flags & EBindingFlags::ViewModelOptional) != 0;
+		return (Flags & (uint16)EFlags::IsOptional) != 0;
 	}
 	
-	/** @return true if the binding is not valid and we should only evaluate the source creator. */
-	bool IsEvaluateSourceCreatorBinding() const
+	/** The source has at least one binding that need to be tick every frame. */
+	bool HasTickBindings() const
+	{
+		return (Flags & (uint16)EFlags::HasTickBindings) != 0;
+	}
+
+	/** The name of the source. */
+	FName GetName() const
 	{
-		return EvaluateSourceCreatorIndex != INDEX_NONE;
+		return PropertyName;
+	}
+
+	/** The name of the UserWidget's property. */
+	FName GetUserWidgetPropertyName() const
+	{
+		return IsUserWidgetProperty() ? PropertyName : FName();
 	}
 
 	/**
-	 * @return true if the binding use a conversion function and that the conversion function is complex.
-	 * The conversion function is complex, there is no input. The inputs are calculated in the BP function.
+	 * FieldId owns by the source that we need to register to.
+	 * Only contains id for OneWay bindings or Evaluate bindings.
 	 */
-	bool IsConversionFunctionComplex() const
+	const TArrayView<const FMVVMViewClass_FieldId> GetFieldIds() const
 	{
-		return Binding.IsComplexFunction();
+		return FieldToRegisterTo;
+	}
+
+	/** The list of bindings owns by the source. A binding can be owns by more than one source. */
+	const TArrayView<const FMVVMViewClass_SourceBinding> GetBindings() const
+	{
+		return Bindings;
 	}
-	
-	/** How the binding should be executed. */
-	EMVVMExecutionMode GetExecuteMode() const;
 
 #if UE_WITH_MVVM_DEBUGGING
 	struct FToStringArgs
 	{
+		FMVVMViewClass_Binding::FToStringArgs Bindings;
 		bool bUseDisplayName = true;
-		bool bAddFieldPath = true;
-		bool bAddFieldId = true;
+		bool bAddCreationMode = true;
+		bool bAddFields = true;
+		bool bAddBindings = true;
 		bool bAddFlags = true;
-		bool bAddBindingId = true;
 
 		MODELVIEWVIEWMODEL_API static FToStringArgs Short();
 		MODELVIEWVIEWMODEL_API static FToStringArgs All();
 	};
 
-
-	UE_DEPRECATED(5.3, "ToString with no argument is deprecated.")
-	FString ToString() const;
-
-	/** @return a human readable version of the binding that can be use for debugging purposes. */
-	FString ToString(const FMVVMCompiledBindingLibrary& CompiledBindingLibrary, FToStringArgs Args) const;
-#endif
-
-#if WITH_EDITOR
-	FGuid GetEditorId() const
-	{
-		return EditorId;
-	}
+	/** @return a human readable version of the source that can be use for debugging purposes. */
+	MODELVIEWVIEWMODEL_API FString ToString(const UMVVMViewClass* ViewClass, FToStringArgs Args) const;
 #endif
 
 private:
-	UPROPERTY()
-	FMVVMVCompiledFieldId FieldId;
+	/** Class type to create a source at runtime. */
+	UPROPERTY(VisibleAnywhere, Category = "View")
+	TSubclassOf<UObject> ExpectedSourceType;
 
-	UPROPERTY(VisibleAnywhere, Category = "Viewmodel")
-	FName SourcePropertyName;
+	/** The resolver to fetch the source at runtime. */
+	UPROPERTY(VisibleAnywhere, Category = "View", Instanced)
+	TObjectPtr<UMVVMViewModelContextResolver> Resolver = nullptr;
 
-	UPROPERTY()
-	FMVVMVCompiledBinding Binding;
+	/** Info to find the ViewModel instance at runtime. */
+	UPROPERTY(VisibleAnywhere, Category = "View")
+	FMVVMViewModelContext GlobalViewModelInstance;
 
-	UPROPERTY(VisibleAnywhere, Category = "Viewmodel")
-	EMVVMExecutionMode ExecutionMode = EMVVMExecutionMode::Immediate;
+	/**
+	 * A resolvable path to retrieve the source instance at runtime.
+	 * It can be a path "Property = Object.Function.Object".
+	 * It can be a UFunction's name of a FProperty's name.
+	 */
+	UPROPERTY(VisibleAnywhere, Category = "View")
+	FMVVMVCompiledFieldPath FieldPath;
 
-	UPROPERTY(VisibleAnywhere, Category = "Viewmodel")
-	int8 EvaluateSourceCreatorIndex = INDEX_NONE;
+	UPROPERTY(VisibleAnywhere, Category = "View")
+	FName PropertyName;
 
-	enum EBindingFlags
+	/** All the fields that the view need to register to for this source. */
+	UPROPERTY(VisibleAnywhere, Category = "View")
+	TArray<FMVVMViewClass_FieldId> FieldToRegisterTo;
+
+	/**
+	 * All the bindings this source need to execute at initialization (OneTime)
+	 * And the bindings that needs to execute when the FieldId matches (OneWay).
+	 */
+	UPROPERTY(VisibleAnywhere, Category = "View")
+	TArray<FMVVMViewClass_SourceBinding> Bindings;
+
+	enum class EFlags : uint16
 	{
 		None = 0,
-		ExecuteAtInitialization = 1 << 0,
-		Unused01 = 1 << 1,
-		OneTime = 1 << 2,
-		EnabledByDefault = 1 << 3,
-		/** The source (viewmodel) can be nullptr and the binding could failed and should not log a warning. */
-		ViewModelOptional = 1 << 4,
-		Unused02 = 1 << 5,
-		/** In development, (when the Blueprint maybe not be compiled with the latest data), the ExecutionMode may not reflect the default project setting value. */
-		OverrideExecuteMode = 1 << 6,
-		/** When the source object is the object itself. */
-		SourceObjectIsSelf = 1 << 7,
+		TypeCreateInstance = 1 << 0,
+		IsUserWidgetProperty = 1 << 1,
+		SetUserWidgetProperty = 1 << 2,
+		IsOptional = 1 << 3,
+		CanBeSet = 1 << 4,
+		CanBeEvaluated = 1 << 5,
+		HasEvaluatedBindings = 1 << 6,
+		SelfReference = 1 << 7,
+		HasTickBindings = 1 << 8,
+		IsViewModel = 1 << 9,
 	};
 
-	UPROPERTY(VisibleAnywhere, Category = "Viewmodel")
-	uint8 Flags = EBindingFlags::None;
-
-#if WITH_EDITORONLY_DATA
-	UPROPERTY(meta = (IgnoreForMemberInitializationTest))
-	FGuid EditorId;
-#endif
+	UPROPERTY(VisibleAnywhere, Category = "View")
+	uint16 Flags = (uint16)EFlags::None;
 };
 
 
@@ -321,36 +408,47 @@ private:
  * A compiled and shared delegate bindings
  */
 USTRUCT()
-struct MODELVIEWVIEWMODEL_API FMVVMViewClass_CompiledEvent
+struct MODELVIEWVIEWMODEL_API FMVVMViewClass_Event
 {
 	GENERATED_BODY()
 
 	friend UE::MVVM::Private::FMVVMViewBlueprintCompiler;
 
 public:
-	/** @return The unique name of the source object that owns the multicast delegate. */
-	FName GetSourceName() const
-	{
-		return SourceName;
-	}
-
+	/**
+	 * The path to get access to the multicast from the UserWidget.
+	 * Include the name of the Widget/ViewModel
+	 */
 	const FMVVMVCompiledFieldPath& GetMulticastDelegatePath() const
 	{
 		return FieldPath;
 	}
 
+	/** The name of the UFunction on the UserWidget. */
 	const FName GetUserWidgetFunctionName() const
 	{
-		return FunctionName;
+		return UserWidgetFunctionName;
+	}
+
+	/**
+	 * The source, if the multicast parent is a valid source.
+	 * This is used when the source value changes at runtime and we want to bound the event again.
+	 */
+	FMVVMViewClass_SourceKey GetSourceKey() const
+	{
+		return SourceToReevaluate;
 	}
 
 #if UE_WITH_MVVM_DEBUGGING
 	struct FToStringArgs
 	{
 		bool bUseDisplayName = true;
+		
+		MODELVIEWVIEWMODEL_API static FToStringArgs Short();
+		MODELVIEWVIEWMODEL_API static FToStringArgs All();
 	};
 	/** @return a human readable version of the binding that can be use for debugging purposes. */
-	FString ToString(const FMVVMCompiledBindingLibrary& CompiledBindingLibrary, FToStringArgs Args) const;
+	FString ToString(const UMVVMViewClass* ViewClass, FToStringArgs Args) const;
 #endif
 
 private:
@@ -358,10 +456,10 @@ private:
 	FMVVMVCompiledFieldPath FieldPath;
 
 	UPROPERTY()
-	FName FunctionName;
+	FName UserWidgetFunctionName;
 
 	UPROPERTY()
-	FName SourceName;
+	FMVVMViewClass_SourceKey SourceToReevaluate;
 };
 
 
@@ -386,95 +484,136 @@ public:
 	//~ End UWidgetBlueprintGeneratedClassExtension
 
 public:
-	/** Should it automatically execute the bindings when the view is constructed or they will be executed manually. */
-	bool InitializeSourcesOnConstruct() const
+	/** Should it automatically create the binding sources when the view is constructed. */
+	UE_NODISCARD bool DoesInitializeSourcesOnConstruct() const
 	{
 		return bInitializeSourcesOnConstruct;
 	}
-	/** Should it automatically execute the bindings when the view is constructed or they will be executed manually. */
-	bool InitializeBindingsOnConstruct() const
+
+	/** Should it automatically register and execute the bindings when the view is constructed. */
+	UE_NODISCARD bool DoesInitializeBindingsOnConstruct() const
 	{
 		return bInitializeBindingsOnConstruct;
 	}
 	
-	/** Get the list of the needed ViewModel. */
-	const TArrayView<const FMVVMViewClass_SourceCreator> GetViewModelCreators() const
+	/** Should it automatically register the events when the view is constructed. */
+	UE_NODISCARD bool DoesInitializeEventsOnConstruct() const
 	{
-		return MakeArrayView(SourceCreators);
+		return bInitializeEventsOnConstruct;
 	}
 
 	/** Get the container of all the bindings. */
-	const FMVVMCompiledBindingLibrary& GetBindingLibrary() const
+	UE_NODISCARD const FMVVMCompiledBindingLibrary& GetBindingLibrary() const
 	{
 		return BindingLibrary;
 	}
 
-	/**  */
-	const TArrayView<const FMVVMViewClass_CompiledBinding> GetCompiledBindings() const
+	/** The field of all the sources that are optional. */
+	uint64 GetOptionalSources() const
 	{
-		return MakeArrayView(CompiledBindings);
+		return OptionalSources;
 	}
-
-	/** */
-	const FMVVMViewClass_CompiledBinding& GetCompiledBinding(int32 Index) const
+	
+	/** Get the list of all the needed viewmodel or widgets. */
+	UE_NODISCARD const TArrayView<const FMVVMViewClass_Source> GetSources() const
 	{
-		return CompiledBindings[Index];
+		return MakeArrayView(Sources);
 	}
 	
-	/**  */
-	const TArrayView<const FMVVMViewClass_CompiledEvent> GetCompiledEvents() const
+	/** The shared source used by the view. */
+	UE_NODISCARD const FMVVMViewClass_Source& GetSource(FMVVMViewClass_SourceKey Key) const
 	{
-		return MakeArrayView(CompiledEvents);
+		check(Sources.IsValidIndex(Key.GetIndex()));
+		return Sources[Key.GetIndex()];
 	}
 
-#if UE_WITH_MVVM_DEBUGGING
-	void Log(FMVVMViewClass_SourceCreator::FToStringArgs SourceArgs, FMVVMViewClass_CompiledBinding::FToStringArgs BindingArgs) const;
-#endif
+	/** The list of bindings. A binding can be used by more than one source. */
+	UE_NODISCARD const TArrayView<const FMVVMViewClass_Binding> GetBindings() const
+	{
+		return Bindings;
+	}
 
-private:
-	/**  */
-	TArrayView<FMVVMViewClass_CompiledBinding> GetCompiledBindings()
+	/** The binding (can be shared by more than one source). */
+	UE_NODISCARD const FMVVMViewClass_Binding& GetBinding(FMVVMViewClass_BindingKey Key) const
+	{
+		check(Bindings.IsValidIndex(Key.GetIndex()));
+		return Bindings[Key.GetIndex()];
+	}
+
+	/** The list of evaluate bindings. */
+	UE_NODISCARD const TArrayView<const FMVVMViewClass_EvaluateSource> GetEvaluateSources() const
+	{
+		return EvaluateSources;
+	}
+
+	/** The evaluate binding. */
+	UE_NODISCARD const FMVVMViewClass_EvaluateSource& GetEvaluateSource(FMVVMViewClass_EvaluateBindingKey Key) const
 	{
-		return MakeArrayView(CompiledBindings);
+		check(EvaluateSources.IsValidIndex(Key.GetIndex()));
+		return EvaluateSources[Key.GetIndex()];
 	}
 	
-	/**  */
-	TArrayView<FMVVMViewClass_CompiledEvent> GetCompiledEvents()
+	/** The list of events. */
+	UE_NODISCARD const TArrayView<const FMVVMViewClass_Event> GetEvents() const
+	{
+		return Events;
+	}
+
+	/** The event. */
+	UE_NODISCARD const FMVVMViewClass_Event& GetEvent(FMVVMViewClass_EventKey Key) const
 	{
-		return MakeArrayView(CompiledEvents);
+		check(Events.IsValidIndex(Key.GetIndex()));
+		return Events[Key.GetIndex()];
 	}
 
+#if UE_WITH_MVVM_DEBUGGING
+	struct FToStringArgs
+	{
+		FMVVMViewClass_Source::FToStringArgs Source;
+		FMVVMViewClass_Binding::FToStringArgs Binding;
+		FMVVMViewClass_EvaluateSource::FToStringArgs Evaluate;
+		FMVVMViewClass_Event::FToStringArgs Event;
+
+		MODELVIEWVIEWMODEL_API static FToStringArgs Short();
+		MODELVIEWVIEWMODEL_API static FToStringArgs All();
+	};
+	UE_NODISCARD FString ToString(FToStringArgs SourceArgs) const;
+#endif
+
+private:
 #if WITH_EDITOR
 	void HandleBlueprintCompiled();
 #endif
 
 private:
-	/** Data to retrieve/create the sources (could be viewmodel, widget, ...). */
-	UPROPERTY()
-	TArray<FMVVMViewClass_SourceCreator> SourceCreators;
-
-	/** */
-	UPROPERTY()
-	TArray<FMVVMViewClass_CompiledBinding> CompiledBindings;
-	
-	/** */
-	UPROPERTY()
-	TArray<FMVVMViewClass_CompiledEvent> CompiledEvents;
+	UPROPERTY(VisibleAnywhere, Category = "View")
+	TArray<FMVVMViewClass_Source> Sources;
+	UPROPERTY(VisibleAnywhere, Category = "View")
+	TArray<FMVVMViewClass_Binding> Bindings;
+	UPROPERTY(VisibleAnywhere, Category = "View")
+	TArray<FMVVMViewClass_EvaluateSource> EvaluateSources;
+	UPROPERTY(VisibleAnywhere, Category = "View")
+	TArray<FMVVMViewClass_Event> Events;
 
 	/** All the bindings shared between all the View instance. */
 	UPROPERTY()
 	FMVVMCompiledBindingLibrary BindingLibrary;
 
-	/** */
+	/** Which view source are optional. */
+	UPROPERTY()
+	uint64 OptionalSources;
+
 	int32 ViewCounter = 0;
 
-	/** */
 	UPROPERTY()
 	bool bInitializeSourcesOnConstruct = true;
-	/** */
+
 	UPROPERTY()
 	bool bInitializeBindingsOnConstruct = true;
 
+	UPROPERTY()
+	bool bInitializeEventsOnConstruct = true;
+
 #if WITH_EDITORONLY_DATA
 	FDelegateHandle BluerpintCompiledHandle;
 #endif
diff --git a/Engine/Plugins/Runtime/ModelViewViewModel/Source/ModelViewViewModel/Public/View/MVVMViewDeprecatedTypes.h b/Engine/Plugins/Runtime/ModelViewViewModel/Source/ModelViewViewModel/Public/View/MVVMViewDeprecatedTypes.h
new file mode 100644
--- /dev/null
+++ b/Engine/Plugins/Runtime/ModelViewViewModel/Source/ModelViewViewModel/Public/View/MVVMViewDeprecatedTypes.h
@@ -0,0 +1,354 @@
+// Copyright Epic Games, Inc. All Rights Reserved.
+
+#pragma once
+
+#include "CoreTypes.h"
+#include "Bindings/MVVMCompiledBindingLibrary.h"
+#include "Types/MVVMExecutionMode.h"
+#include "Types/MVVMViewModelContext.h"
+
+#include "MVVMViewDeprecatedTypes.generated.h"
+
+class UMVVMViewClass;
+class UMVVMView;
+class UMVVMViewModelContextResolver;
+class UObject;
+class UUserWidget;
+
+
+USTRUCT()
+struct
+UE_DEPRECATED(5.4, "FMVVMVCompiledFieldId is no longer used.")
+FMVVMVCompiledFieldId
+{
+	GENERATED_BODY()
+
+	bool IsValid() const
+	{
+		return FieldIdIndex >= 0;
+	}
+
+private:
+	using IndexType = int16;
+
+	UPROPERTY()
+	int16 FieldIdIndex = INDEX_NONE;
+
+#if WITH_EDITORONLY_DATA
+	UPROPERTY(meta = (IgnoreForMemberInitializationTest))
+	FGuid CompiledBindingLibraryId;
+#endif
+};
+
+USTRUCT()
+struct
+UE_DEPRECATED(5.4, "FMVVMViewSource is no longer used, please use FMVVMView_Source instead.")
+FMVVMViewSource
+{
+	GENERATED_BODY()
+
+	UPROPERTY(VisibleAnywhere, Category = "Viewmodel")
+	TObjectPtr<UObject> Source; // TScriptInterface<INotifyFieldValueChanged>
+
+	UPROPERTY(VisibleAnywhere, Category = "Viewmodel")
+	FName SourceName;
+
+	// Number of bindings connected to the source.
+	UPROPERTY(VisibleAnywhere, Category = "Viewmodel")
+	int32 RegisteredCount = 0;
+
+	// The source is defined in the ClassExtension.
+	UPROPERTY(VisibleAnywhere, Category = "Viewmodel")
+	bool bCreatedSource = false;
+
+	// The source bindings are initialized
+	UPROPERTY(VisibleAnywhere, Category = "Viewmodel")
+	bool bBindingsInitialized = false;
+
+	// The source was set manually via SetViewModel.
+	UPROPERTY(VisibleAnywhere, Category = "Viewmodel")
+	bool bSetManually = false;
+
+	// The source was set to a UserWidget property.
+	UPROPERTY(VisibleAnywhere, Category = "Viewmodel")
+	bool bAssignedToUserWidgetProperty = false;
+};
+
+USTRUCT()
+struct
+UE_DEPRECATED(5.4, "FMVVMViewClass_SourceCreator is no longer used, please use FMVVMViewClass_Source instead.")
+FMVVMViewClass_SourceCreator
+{
+	GENERATED_BODY()
+
+public:
+	UObject* CreateInstance(const UMVVMViewClass* ViewClass, UMVVMView* View, UUserWidget* UserWidget) const
+	{
+		return nullptr;
+	}
+
+	void DestroyInstance(const UObject* ViewModel, const UMVVMView* View) const
+	{
+	}
+
+	UClass* GetSourceClass() const
+	{
+		return ExpectedSourceType.Get();
+	}
+
+	bool IsSourceAUserWidgetProperty() const
+	{
+		return (Flags & (uint8)ESourceFlags::IsUserWidgetProperty) != 0;
+	}
+	
+	bool CanBeSet() const
+	{
+		return (Flags & (uint8)ESourceFlags::CanBeSet) != 0;
+	}
+
+	bool CanBeEvaluated() const
+	{
+		return (Flags & (uint8)ESourceFlags::CanBeEvaluated) != 0;
+	}
+	
+	bool IsOptional() const
+	{
+		return (Flags & (uint8)ESourceFlags::IsOptional) != 0;
+	}
+
+	FName GetSourceName() const
+	{
+		return PropertyName;
+	}
+	
+	FName GetParentSourceName() const
+	{
+		return ParentSourceName;
+	}
+
+#if UE_WITH_MVVM_DEBUGGING
+	struct FToStringArgs
+	{
+		bool bUseDisplayName = true;
+		bool bAddCreationMode = true;
+		bool bAddFlags = true;
+
+		static FToStringArgs Short() { return FToStringArgs(); }
+		static FToStringArgs All() { return FToStringArgs(); }
+	};
+	FString ToString(const FMVVMCompiledBindingLibrary& CompiledBindingLibrary, FToStringArgs Args) const{ return FString(); }
+#endif
+
+private:
+	UPROPERTY()
+	TSubclassOf<UObject> ExpectedSourceType;
+	UPROPERTY(Instanced)
+	TObjectPtr<UMVVMViewModelContextResolver> Resolver = nullptr;
+	UPROPERTY()
+	FMVVMViewModelContext GlobalViewModelInstance;
+	UPROPERTY()
+	FMVVMVCompiledFieldPath FieldPath;
+	UPROPERTY()
+	FName PropertyName;
+	UPROPERTY()
+	FName ParentSourceName;
+
+	enum class ESourceFlags : uint8
+	{
+		None = 0,
+		TypeCreateInstance = 1 << 0,
+		IsUserWidgetProperty = 1 << 1,
+		IsOptional = 1 << 2,
+		CanBeSet = 1 << 3,
+		CanBeEvaluated = 1 << 4,
+		SelfReference = 1 << 5,
+	};
+
+	UPROPERTY()
+	uint8 Flags = (uint8)ESourceFlags::None;
+};
+
+PRAGMA_DISABLE_DEPRECATION_WARNINGS
+USTRUCT()
+struct
+UE_DEPRECATED(5.4, "FMVVMViewClass_CompiledBinding is no longer used, please use FMVVMViewClass_Binding and FMVVMViewClass_SourceBinding instead.")
+FMVVMViewClass_CompiledBinding
+{
+	GENERATED_BODY()
+
+public:
+	FMVVMVCompiledFieldId GetSourceFieldId() const
+	{
+		return FieldId;
+	}
+
+	FName GetSourceName() const
+	{
+		return SourcePropertyName;
+	}
+
+	int32 GetEvaluateSourceCreatorBindingIndex() const
+	{
+		return EvaluateSourceCreatorIndex;
+	}
+
+	bool IsSourceUserWidget() const
+	{
+		return (Flags & EBindingFlags::SourceObjectIsSelf) != 0;
+	}
+
+	const FMVVMVCompiledBinding& GetBinding() const
+	{
+		return Binding;
+	}
+
+	bool NeedsExecutionAtInitialization() const
+	{
+		return (Flags & EBindingFlags::ExecuteAtInitialization) != 0;
+	}
+
+	bool IsOneTime() const
+	{
+		return (Flags & EBindingFlags::OneTime) != 0;
+	}
+
+	bool IsEnabledByDefault() const
+	{
+		return (Flags & EBindingFlags::EnabledByDefault) != 0;
+	}
+
+	bool IsRegistrationOptional() const
+	{
+		return (Flags & EBindingFlags::ViewModelOptional) != 0;
+	}
+	
+	bool IsEvaluateSourceCreatorBinding() const
+	{
+		return EvaluateSourceCreatorIndex != INDEX_NONE;
+	}
+
+	bool IsConversionFunctionComplex() const
+	{
+		return Binding.HasComplexConversionFunction();
+	}
+	
+	EMVVMExecutionMode GetExecuteMode() const
+	{
+		return (Flags & EBindingFlags::OverrideExecuteMode) != 0 ? ExecutionMode : EMVVMExecutionMode::Immediate;
+	}
+
+#if UE_WITH_MVVM_DEBUGGING
+	struct FToStringArgs
+	{
+		bool bUseDisplayName = true;
+		bool bAddFieldPath = true;
+		bool bAddFieldId = true;
+		bool bAddFlags = true;
+		bool bAddBindingId = true;
+
+		static FToStringArgs Short()
+		{
+			return FToStringArgs();
+		}
+		static FToStringArgs All()
+		{
+			return FToStringArgs();
+		}
+	};
+
+
+	FString ToString() const
+	{
+		return FString();
+	}
+	FString ToString(const FMVVMCompiledBindingLibrary& CompiledBindingLibrary, FToStringArgs Args) const
+	{
+		return FString();
+	}
+#endif
+
+#if WITH_EDITOR
+	FGuid GetEditorId() const
+	{
+		return EditorId;
+	}
+#endif
+
+private:
+	UPROPERTY()
+	FMVVMVCompiledFieldId FieldId;
+	UPROPERTY(VisibleAnywhere, Category = "Viewmodel")
+	FName SourcePropertyName;
+	UPROPERTY()
+	FMVVMVCompiledBinding Binding;
+	UPROPERTY(VisibleAnywhere, Category = "Viewmodel")
+	EMVVMExecutionMode ExecutionMode = EMVVMExecutionMode::Immediate;
+	UPROPERTY(VisibleAnywhere, Category = "Viewmodel")
+	int8 EvaluateSourceCreatorIndex = INDEX_NONE;
+
+	enum EBindingFlags
+	{
+		None = 0,
+		ExecuteAtInitialization = 1 << 0,
+		Unused01 = 1 << 1,
+		OneTime = 1 << 2,
+		EnabledByDefault = 1 << 3,
+		ViewModelOptional = 1 << 4,
+		Unused02 = 1 << 5,
+		OverrideExecuteMode = 1 << 6,
+		SourceObjectIsSelf = 1 << 7,
+	};
+
+	UPROPERTY(VisibleAnywhere, Category = "Viewmodel")
+	uint8 Flags = EBindingFlags::None;
+
+#if WITH_EDITORONLY_DATA
+	UPROPERTY(meta = (IgnoreForMemberInitializationTest))
+	FGuid EditorId;
+#endif
+};
+PRAGMA_ENABLE_DEPRECATION_WARNINGS
+
+
+USTRUCT()
+struct
+UE_DEPRECATED(5.4, "MVVMViewClass_CompiledEvent is no longer used, please use MVVMViewClass_Event instead.")
+FMVVMViewClass_CompiledEvent
+{
+	GENERATED_BODY()
+
+public:
+	FName GetSourceName() const
+	{
+		return SourceName;
+	}
+
+	const FMVVMVCompiledFieldPath& GetMulticastDelegatePath() const
+	{
+		return FieldPath;
+	}
+
+	const FName GetUserWidgetFunctionName() const
+	{
+		return FunctionName;
+	}
+
+#if UE_WITH_MVVM_DEBUGGING
+	struct FToStringArgs
+	{
+		bool bUseDisplayName = true;
+	};
+	FString ToString(const FMVVMCompiledBindingLibrary& CompiledBindingLibrary, FToStringArgs Args) const
+	{
+		return FString();
+	}
+#endif
+
+private:
+	UPROPERTY()
+	FMVVMVCompiledFieldPath FieldPath;
+	UPROPERTY()
+	FName FunctionName;
+	UPROPERTY()
+	FName SourceName;
+};
diff --git a/Engine/Plugins/Runtime/ModelViewViewModel/Source/ModelViewViewModel/Public/View/MVVMViewTypes.h b/Engine/Plugins/Runtime/ModelViewViewModel/Source/ModelViewViewModel/Public/View/MVVMViewTypes.h
new file mode 100644
--- /dev/null
+++ b/Engine/Plugins/Runtime/ModelViewViewModel/Source/ModelViewViewModel/Public/View/MVVMViewTypes.h
@@ -0,0 +1,235 @@
+// Copyright Epic Games, Inc. All Rights Reserved.
+
+#pragma once
+
+#include "FieldNotificationId.h"
+
+#include "MVVMViewTypes.generated.h"
+
+namespace UE::MVVM::Private { struct FMVVMViewBlueprintCompiler; }
+class UMVVMViewClass;
+
+
+
+/**
+ * FieldName and Bit needed to create a UE::FieldNotification::FFieldId.
+ */
+USTRUCT()
+struct FMVVMViewClass_FieldId
+{
+	GENERATED_BODY()
+
+	friend UE::MVVM::Private::FMVVMViewBlueprintCompiler;
+	friend UMVVMViewClass;
+
+public:
+	FMVVMViewClass_FieldId() = default;
+	FMVVMViewClass_FieldId(UE::FieldNotification::FFieldId InFieldId)
+	{
+		FieldName = InFieldId.GetName();
+		FieldIndex = InFieldId.GetIndex();
+	}
+
+	bool IsValid() const
+	{
+		return FieldIndex != INDEX_NONE && !FieldName.IsNone();
+	}
+
+	UE::FieldNotification::FFieldId GetFieldId() const
+	{
+		return UE::FieldNotification::FFieldId(FieldName, FieldIndex);
+	}
+
+	FName GetName() const
+	{
+		return FieldName;
+	}
+
+	bool operator== (const FMVVMViewClass_FieldId& Other) const
+	{
+		return FieldName == Other.FieldName && FieldIndex == Other.FieldIndex;
+	}
+
+private:
+	UPROPERTY(VisibleAnywhere, Category = "View")
+	FName FieldName;
+
+	int32 FieldIndex = INDEX_NONE;
+};
+
+/**
+ * Key that identify the FMVVMViewClass_Source.
+ */
+USTRUCT()
+struct FMVVMViewClass_SourceKey
+{
+	GENERATED_BODY()
+
+public:
+	FMVVMViewClass_SourceKey() = default;
+	FMVVMViewClass_SourceKey(int32 InIndex)
+		: Index(InIndex)
+	{}
+
+	bool IsValid() const
+	{
+		return Index != INDEX_NONE;
+	}
+
+	uint64 GetBit() const
+	{
+		// keep the first bit free for now as it may be useful in the future.
+		return (uint64)1 << Index;
+	}
+
+	int32 GetIndex() const
+	{
+		return Index;
+	}
+
+	bool operator== (const FMVVMViewClass_SourceKey& Other) const
+	{
+		return Index == Other.Index;
+	}
+
+private:
+	UPROPERTY(VisibleAnywhere, Category = "View")
+	int32 Index = INDEX_NONE;
+};
+
+/**
+ * Key that identify the FMVVMViewClass_Binding.
+ */
+USTRUCT()
+struct FMVVMViewClass_BindingKey
+{
+	GENERATED_BODY()
+
+public:
+	FMVVMViewClass_BindingKey() = default;
+	FMVVMViewClass_BindingKey(int32 InIndex)
+		: Index(InIndex)
+	{}
+	bool IsValid() const
+	{
+		return Index != INDEX_NONE;
+	}
+
+	int32 GetIndex() const
+	{
+		return Index;
+	}
+
+	bool operator== (const FMVVMViewClass_BindingKey& Other) const
+	{
+		return Index == Other.Index;
+	}
+
+private:
+	UPROPERTY(VisibleAnywhere, Category = "View")
+	int32 Index = INDEX_NONE;
+};
+
+/**
+ * Key that identify the FMVVMViewClass_EvaluateBinding.
+ */
+USTRUCT()
+struct FMVVMViewClass_EvaluateBindingKey
+{
+	GENERATED_BODY()
+
+	friend UE::MVVM::Private::FMVVMViewBlueprintCompiler;
+
+public:
+	FMVVMViewClass_EvaluateBindingKey() = default;
+	FMVVMViewClass_EvaluateBindingKey(int32 InIndex)
+		: Index(InIndex)
+	{}
+	bool IsValid() const
+	{
+		return Index != INDEX_NONE;
+	}
+
+	int32 GetIndex() const
+	{
+		return Index;
+	}
+
+	bool operator== (const FMVVMViewClass_EvaluateBindingKey& Other) const
+	{
+		return Index == Other.Index;
+	}
+
+private:
+	UPROPERTY(VisibleAnywhere, Category = "View")
+	int32 Index = INDEX_NONE;
+};
+
+/**
+ * Key that identify the FMVVMViewClass_Event.
+ */
+USTRUCT()
+struct FMVVMViewClass_EventKey
+{
+	GENERATED_BODY()
+
+	friend UE::MVVM::Private::FMVVMViewBlueprintCompiler;
+
+public:
+	FMVVMViewClass_EventKey() = default;
+	FMVVMViewClass_EventKey(int32 InIndex)
+		: Index(InIndex)
+	{}
+	bool IsValid() const
+	{
+		return Index != INDEX_NONE;
+	}
+
+	int32 GetIndex() const
+	{
+		return Index;
+	}
+
+	bool operator== (const FMVVMViewClass_EventKey& Other) const
+	{
+		return Index == Other.Index;
+	}
+
+private:
+	UPROPERTY(VisibleAnywhere, Category = "View")
+	int32 Index = INDEX_NONE;
+};
+
+/**
+ * Key that identify the FMVVMView_Source.
+ */
+USTRUCT()
+struct FMVVMView_SourceKey
+{
+	GENERATED_BODY()
+
+public:
+	FMVVMView_SourceKey() = default;
+	FMVVMView_SourceKey(int32 InIndex)
+		: Index(InIndex)
+	{}
+
+	bool IsValid() const
+	{
+		return Index != INDEX_NONE;
+	}
+
+	int32 GetIndex() const
+	{
+		return Index;
+	}
+
+	bool operator== (const FMVVMView_SourceKey& Other) const
+	{
+		return Index == Other.Index;
+	}
+
+private:
+	UPROPERTY(VisibleAnywhere, Category = "View")
+	int32 Index = INDEX_NONE;
+};
diff --git a/Engine/Plugins/Runtime/ModelViewViewModel/Source/ModelViewViewModelBlueprint/Private/Bindings/MVVMCompiledBindingLibraryCompiler.cpp b/Engine/Plugins/Runtime/ModelViewViewModel/Source/ModelViewViewModelBlueprint/Private/Bindings/MVVMCompiledBindingLibraryCompiler.cpp
--- a/Engine/Plugins/Runtime/ModelViewViewModel/Source/ModelViewViewModelBlueprint/Private/Bindings/MVVMCompiledBindingLibraryCompiler.cpp
+++ b/Engine/Plugins/Runtime/ModelViewViewModel/Source/ModelViewViewModelBlueprint/Private/Bindings/MVVMCompiledBindingLibraryCompiler.cpp
@@ -25,10 +25,7 @@ struct FRawFieldId
 {
 	const UClass* NotifyFieldValueChangedClass;
 	UE::FieldNotification::FFieldId FieldId;
-
-	int32 LoadedFieldIdIndex = INDEX_NONE;
 	FCompiledBindingLibraryCompiler::FFieldIdHandle IdHandle;
-	FMVVMVCompiledFieldId CompiledFieldId;
 };
 
 
@@ -717,29 +714,6 @@ TValueOrError<FCompiledBindingLibraryCompiler::FCompileResult, FText> FCompiledB
 			ClassInfo.RawFieldIndex.Add(Index);
 		}
 	}
-	{
-		for (int32 Index = 0; Index < Impl->FieldIds.Num(); ++Index)
-		{
-			const Private::FRawFieldId& RawFieldId = Impl->FieldIds[Index];
-			check(RawFieldId.FieldId.IsValid());
-			check(RawFieldId.NotifyFieldValueChangedClass);
-
-			FCompiledClassInfo& ClassInfo = MapOfFieldInClass.FindOrAdd(Private::GetSavedGeneratedStruct(RawFieldId.NotifyFieldValueChangedClass));
-
-			// Test if the Field is not there more than one
-			{
-				UE::FieldNotification::FFieldId FieldIdToTest = RawFieldId.FieldId;
-				const TArray<Private::FRawFieldId>& ListOfFieldIds = Impl->FieldIds;
-				const bool bContains = ClassInfo.RawFieldIdIndex.ContainsByPredicate([FieldIdToTest, &ListOfFieldIds](int32 OtherIndex)
-					{
-						return ListOfFieldIds[OtherIndex].FieldId == FieldIdToTest;
-					});
-				check(!bContains);
-			}
-
-			ClassInfo.RawFieldIdIndex.Add(Index);
-		}
-	}
 
 	// Todo optimize that list to group common type. ie UWidget::ToolTip == UProgressBar::ToolTip. We can merge UWidget in UProgressBar.
 	//Algo: for each class entry
@@ -752,7 +726,6 @@ TValueOrError<FCompiledBindingLibraryCompiler::FCompileResult, FText> FCompiledB
 	// Create FMVVMCompiledBindingLibrary::CompiledFields and FMVVMCompiledBindingLibrary::CompiledFieldNames
 	int32 TotalNumberOfProperties = 0;
 	int32 TotalNumberOfFunctions = 0;
-	int32 TotalNumberOfFieldIds = 0;
 	for (TPair<const UStruct*, FCompiledClassInfo>& StructCompiledFields : MapOfFieldInClass)
 	{
 		FMVVMVCompiledFields CompiledFields;
@@ -787,17 +760,6 @@ TValueOrError<FCompiledBindingLibraryCompiler::FCompileResult, FText> FCompiledB
 			}
 		}
 
-		for (const int32 FieldIdIndex : StructCompiledFields.Value.RawFieldIdIndex)
-		{
-			check(Impl->FieldIds.IsValidIndex(FieldIdIndex));
-			Private::FRawFieldId& RawFieldId = Impl->FieldIds[FieldIdIndex];
-			
-			Result.Library.LoadedFieldIds.Add(RawFieldId.FieldId);
-			FieldIdNames.Add(RawFieldId.FieldId.GetName());
-			RawFieldId.LoadedFieldIdIndex = TotalNumberOfFieldIds;
-			++TotalNumberOfFieldIds;
-		}
-
 		if (PropertyNames.Num() > std::numeric_limits<FMVVMVCompiledBinding::IndexType>::max())
 		{
 			return MakeError(FText::Format(LOCTEXT("TooManyPropertiesBound", "There are too many properties bound to struct '{0}'"), StructCompiledFields.Key->GetDisplayNameText()));
@@ -810,12 +772,6 @@ TValueOrError<FCompiledBindingLibraryCompiler::FCompileResult, FText> FCompiledB
 		}
 		CompiledFields.NumberOfFunctions = static_cast<int16>(FunctionNames.Num());
 
-		if (FieldIdNames.Num() > std::numeric_limits<FMVVMVCompiledBinding::IndexType>::max())
-		{
-			return MakeError(FText::Format(LOCTEXT("TooManyFieldIdsBound", "There are too many field IDs bound to struct '{0}'"), StructCompiledFields.Key->GetDisplayNameText()));
-		}
-		CompiledFields.NumberOfFieldIds = static_cast<int16>(FieldIdNames.Num());
-
 		int32 LibraryStartIndex = Result.Library.CompiledFieldNames.Num();
 		if (LibraryStartIndex > std::numeric_limits<FMVVMVCompiledBinding::IndexType>::max())
 		{
@@ -827,8 +783,6 @@ TValueOrError<FCompiledBindingLibraryCompiler::FCompileResult, FText> FCompiledB
 		PropertyNames.Reset();
 		Result.Library.CompiledFieldNames.Append(FunctionNames);
 		FunctionNames.Reset();
-		Result.Library.CompiledFieldNames.Append(FieldIdNames);
-		FieldIdNames.Reset();
 		if (Result.Library.CompiledFieldNames.Num() > std::numeric_limits<FMVVMVCompiledBinding::IndexType>::max())
 		{
 			return MakeError(LOCTEXT("TooManyPropertiesBoundInLibrary", "There are too many properties bound in the library."));
@@ -836,10 +790,9 @@ TValueOrError<FCompiledBindingLibraryCompiler::FCompileResult, FText> FCompiledB
 
 		Result.Library.CompiledFields.Add(CompiledFields);
 
-		check(Result.Library.LoadedProperties.Num() + Result.Library.LoadedFunctions.Num() + Result.Library.LoadedFieldIds.Num()  == Result.Library.CompiledFieldNames.Num());
+		check(Result.Library.LoadedProperties.Num() + Result.Library.LoadedFunctions.Num() == Result.Library.CompiledFieldNames.Num());
 		check(Result.Library.LoadedProperties.Num() == TotalNumberOfProperties);
 		check(Result.Library.LoadedFunctions.Num() == TotalNumberOfFunctions);
-		check(Result.Library.LoadedFieldIds.Num() == TotalNumberOfFieldIds);
 	}
 
 	// Create FMVVMCompiledBindingLibrary::FieldPaths
@@ -880,10 +833,7 @@ TValueOrError<FCompiledBindingLibraryCompiler::FCompileResult, FText> FCompiledB
 	// Create FieldId
 	for (Private::FRawFieldId& FieldId: Impl->FieldIds)
 	{
-		FieldId.CompiledFieldId.CompiledBindingLibraryId = Result.Library.CompiledBindingLibraryId;
-		FieldId.CompiledFieldId.FieldIdIndex = FieldId.LoadedFieldIdIndex;
-
-		Result.FieldIds.Add(FieldId.IdHandle, FieldId.CompiledFieldId);
+		Result.FieldIds.Add(FieldId.IdHandle, FieldId.FieldId);
 	}
 
 	auto GetCompiledFieldPath = [this](const FFieldPathHandle Handle)
@@ -913,17 +863,24 @@ TValueOrError<FCompiledBindingLibraryCompiler::FCompileResult, FText> FCompiledB
 
 		Binding.CompiledBinding.ConversionFunctionFieldPath = GetCompiledFieldPath(Binding.ConversionFunctionPathHandle);
 
-		Binding.CompiledBinding.Flags = 0;
-		Binding.CompiledBinding.Flags |= (Binding.ConversionFunctionPathHandle.IsValid()) ? (uint8)FMVVMVCompiledBinding::EFlags::HasConversionFunction : 0;
-		Binding.CompiledBinding.Flags |= (Binding.bIsConversionFunctionComplex) ? (uint8)FMVVMVCompiledBinding::EFlags::IsConversionFunctionComplex : 0;
-		Binding.CompiledBinding.Flags |= Binding.BindingCount > 1 ? (uint8)FMVVMVCompiledBinding::EFlags::IsShared : 0;
+		if (Binding.bIsConversionFunctionComplex)
+		{
+			Binding.CompiledBinding.Type = (uint8)FMVVMVCompiledBinding::EType::HasComplexConversionFunction;
+		}
+		else if (Binding.ConversionFunctionPathHandle.IsValid())
+		{
+			Binding.CompiledBinding.Type = (uint8)FMVVMVCompiledBinding::EType::HasConversionFunction;
+		}
+		else
+		{
+			Binding.CompiledBinding.Type = (uint8)FMVVMVCompiledBinding::EType::None;
+		}
 
 		Result.Bindings.Add(Binding.BindingHandle, Binding.CompiledBinding);
 	}
 
 	Result.Library.LoadedProperties.Reset();
 	Result.Library.LoadedFunctions.Reset();
-	Result.Library.LoadedFieldIds.Reset();
 
 	Impl->bCompiled = true;
 	return MakeValue(MoveTemp(Result));
diff --git a/Engine/Plugins/Runtime/ModelViewViewModel/Source/ModelViewViewModelBlueprint/Private/MVVMViewBlueprintCompiler.cpp b/Engine/Plugins/Runtime/ModelViewViewModel/Source/ModelViewViewModelBlueprint/Private/MVVMViewBlueprintCompiler.cpp
--- a/Engine/Plugins/Runtime/ModelViewViewModel/Source/ModelViewViewModelBlueprint/Private/MVVMViewBlueprintCompiler.cpp
+++ b/Engine/Plugins/Runtime/ModelViewViewModel/Source/ModelViewViewModelBlueprint/Private/MVVMViewBlueprintCompiler.cpp
@@ -1312,13 +1312,28 @@ bool FMVVMViewBlueprintCompiler::Compile(UWidgetBlueprintGeneratedClass* Class,
 		WidgetBlueprintCompilerContext.MessageLog.Error(*FText::Format(LOCTEXT("BindingCompilationFailed", "The binding compilation failed. {1}"), CompileResult.GetError()).ToString());
 		return false;
 	}
-	CompileViewModelCreatorContexts(CompileResult.GetValue(), Class, ViewExtension);
+	CompileSources(CompileResult.GetValue(), Class, ViewExtension);
 	CompileBindings(CompileResult.GetValue(), Class, ViewExtension);
+	CompileEvaluateSources(CompileResult.GetValue(), Class, ViewExtension);
 	CompileEvents(CompileResult.GetValue(), Class, ViewExtension);
+	SortSourceFields(CompileResult.GetValue(), Class, ViewExtension);
 
 	{
 		ViewExtension->bInitializeSourcesOnConstruct = BlueprintView->GetSettings()->bInitializeSourcesOnConstruct;
 		ViewExtension->bInitializeBindingsOnConstruct = ViewExtension->bInitializeSourcesOnConstruct ? BlueprintView->GetSettings()->bInitializeBindingsOnConstruct : false;
+		ViewExtension->bInitializeEventsOnConstruct = BlueprintView->GetSettings()->bInitializeEventsOnConstruct;
+	}
+	{
+		ViewExtension->OptionalSources = 0;
+		for (int32 Index = 0; Index < ViewExtension->Sources.Num(); ++Index)
+		{
+			const FMVVMViewClass_Source& Source = ViewExtension->Sources[Index];
+			if (Source.IsOptional())
+			{
+				FMVVMViewClass_SourceKey ClassSourceKey = FMVVMViewClass_SourceKey(Index);
+				ViewExtension->OptionalSources |= ClassSourceKey.GetBit();
+			}
+		}
 	}
 
 	bool bResult = AreStepsValid();
@@ -1329,11 +1344,12 @@ bool FMVVMViewBlueprintCompiler::Compile(UWidgetBlueprintGeneratedClass* Class,
 #if UE_WITH_MVVM_DEBUGGING
 		if (CVarLogViewCompiledResult->GetBool())
 		{
-			FMVVMViewClass_SourceCreator::FToStringArgs CreatorsToStringArgs = FMVVMViewClass_SourceCreator::FToStringArgs::All();
-			CreatorsToStringArgs.bUseDisplayName = false;
-			FMVVMViewClass_CompiledBinding::FToStringArgs BindingToStringArgs = FMVVMViewClass_CompiledBinding::FToStringArgs::All();
-			BindingToStringArgs.bUseDisplayName = false;
-			ViewExtension->Log(CreatorsToStringArgs, BindingToStringArgs);
+			UMVVMViewClass::FToStringArgs ToStringArgs = UMVVMViewClass::FToStringArgs::All();
+			ToStringArgs.Source.bUseDisplayName = false;
+			ToStringArgs.Binding.bUseDisplayName = false;
+			ToStringArgs.Evaluate.bUseDisplayName = false;
+			ToStringArgs.Event.bUseDisplayName = false;
+			UE_LOG(LogMVVM, Log, TEXT("%s"), *ViewExtension->ToString(ToStringArgs));
 		}
 #endif
 	}
@@ -1789,13 +1805,40 @@ void FMVVMViewBlueprintCompiler::PreCompileViewModelCreatorContexts(UWidgetBluep
 }
 
 
-void FMVVMViewBlueprintCompiler::CompileViewModelCreatorContexts(const FCompiledBindingLibraryCompiler::FCompileResult& CompileResult, UWidgetBlueprintGeneratedClass* Class, UMVVMViewClass* ViewExtension)
+void FMVVMViewBlueprintCompiler::CompileSources(const FCompiledBindingLibraryCompiler::FCompileResult& CompileResult, UWidgetBlueprintGeneratedClass* Class, UMVVMViewClass* ViewExtension)
 {
-	TArray<FMVVMViewClass_SourceCreator> UnsortedSourceCreators;
+	struct FSortData
+	{
+		FSortData() = default;
+		FName SourceName;
+		FName ParentSourceName;
+		int32 SortIndex = -1;
+
+		void CalculateSortIndex(TMap<FName, FSortData>& Map)
+		{
+			if (SortIndex < 0)
+			{
+				if (ParentSourceName.IsNone())
+				{
+					SortIndex = 0;
+				}
+				else
+				{
+					Map[ParentSourceName].CalculateSortIndex(Map);
+					// calculate the Depth recursively
+					SortIndex = Map[ParentSourceName].SortIndex + 1;
+				}
+			}
+		}
+	};
+
+	TMap<FName, FSortData> SortDatas;
+	TArray<FMVVMViewClass_Source> UnsortedSourceCreators;
+
 	for (FCompilerViewModelCreatorContext& SourceCreatorContext : ViewModelCreatorContexts)
 	{
 		const FMVVMBlueprintViewModelContext& ViewModelContext = SourceCreatorContext.ViewModelContext;
-		FMVVMViewClass_SourceCreator CompiledSourceCreator;
+		FMVVMViewClass_Source CompiledSourceCreator;
 
 		ensure(ViewModelContext.GetViewModelClass() && ViewModelContext.GetViewModelClass()->ImplementsInterface(UNotifyFieldValueChanged::StaticClass()));
 		CompiledSourceCreator.ExpectedSourceType = ViewModelContext.GetViewModelClass();
@@ -1876,11 +1919,12 @@ void FMVVMViewBlueprintCompiler::CompileViewModelCreatorContexts(const FCompiled
 			continue;
 		}
 
+		FName ParentSourceName;
 		if (SourceCreatorContext.DynamicContext)
 		{
 			if (SourceCreatorContext.DynamicContext->ParentSource)
 			{
-				CompiledSourceCreator.ParentSourceName = SourceCreatorContext.DynamicContext->ParentSource->Name;
+				ParentSourceName = SourceCreatorContext.DynamicContext->ParentSource->Name;
 			}
 			else
 			{
@@ -1894,24 +1938,30 @@ void FMVVMViewBlueprintCompiler::CompileViewModelCreatorContexts(const FCompiled
 		}
 
 		CompiledSourceCreator.Flags = 0;
-		CompiledSourceCreator.Flags |= bCreateInstance ? (uint8)FMVVMViewClass_SourceCreator::ESourceFlags::TypeCreateInstance : 0;
-		CompiledSourceCreator.Flags |= bIsUserWidgetProperty ? (uint8)FMVVMViewClass_SourceCreator::ESourceFlags::IsUserWidgetProperty : 0;
-		CompiledSourceCreator.Flags |= bIsOptional ? (uint8)FMVVMViewClass_SourceCreator::ESourceFlags::IsOptional : 0;
-		CompiledSourceCreator.Flags |= bCanBeSet ? (uint8)FMVVMViewClass_SourceCreator::ESourceFlags::CanBeSet : 0;
-		CompiledSourceCreator.Flags |= bCanBeEvaluated ? (uint8)FMVVMViewClass_SourceCreator::ESourceFlags::CanBeEvaluated : 0;
-
+		CompiledSourceCreator.Flags |= bCreateInstance ? (uint16)FMVVMViewClass_Source::EFlags::TypeCreateInstance : 0;
+		CompiledSourceCreator.Flags |= bIsUserWidgetProperty ? (uint16)FMVVMViewClass_Source::EFlags::IsUserWidgetProperty : 0;
+		CompiledSourceCreator.Flags |= bIsUserWidgetProperty ? (uint16)FMVVMViewClass_Source::EFlags::SetUserWidgetProperty : 0;
+		CompiledSourceCreator.Flags |= bIsOptional ? (uint16)FMVVMViewClass_Source::EFlags::IsOptional : 0;
+		CompiledSourceCreator.Flags |= bCanBeSet ? (uint16)FMVVMViewClass_Source::EFlags::CanBeSet : 0;
+		CompiledSourceCreator.Flags |= bCanBeEvaluated ? (uint16)FMVVMViewClass_Source::EFlags::CanBeEvaluated : 0;
+		CompiledSourceCreator.Flags |= (uint16)FMVVMViewClass_Source::EFlags::IsViewModel;
+
+		FSortData SortData;
+		SortData.SourceName = CompiledSourceCreator.GetName();
+		SortData.ParentSourceName = ParentSourceName;
+		SortDatas.Add(CompiledSourceCreator.GetName(), SortData);
 		UnsortedSourceCreators.Add(MoveTemp(CompiledSourceCreator));
 	}
 
 	// The other sources needed by the view
 	for (FCompilerWidgetCreatorContext& WidgetCreator : WidgetCreatorContexts)
 	{
-		FMVVMViewClass_SourceCreator CompiledSourceCreator;
+		FMVVMViewClass_Source CompiledSourceCreator;
 		ensure(WidgetCreator.Source->AuthoritativeClass&& WidgetCreator.Source->AuthoritativeClass->ImplementsInterface(UNotifyFieldValueChanged::StaticClass()));
 		CompiledSourceCreator.ExpectedSourceType = const_cast<UClass*>(WidgetCreator.Source->AuthoritativeClass);
 		CompiledSourceCreator.PropertyName = WidgetCreator.Source->Name;
 		CompiledSourceCreator.Flags = 0;
-		CompiledSourceCreator.Flags |= WidgetCreator.bSelfReference ? (uint8)FMVVMViewClass_SourceCreator::ESourceFlags::SelfReference : 0;
+		CompiledSourceCreator.Flags |= WidgetCreator.bSelfReference ? (uint16)FMVVMViewClass_Source::EFlags::SelfReference : 0;
 
 		if (!WidgetCreator.bSelfReference)
 		{
@@ -1927,64 +1977,57 @@ void FMVVMViewBlueprintCompiler::CompileViewModelCreatorContexts(const FCompiled
 			CompiledSourceCreator.FieldPath = *CompiledFieldPath;
 		}
 
+		FSortData SortData;
+		SortDatas.Add(CompiledSourceCreator.GetName(), SortData);
 		UnsortedSourceCreators.Add(MoveTemp(CompiledSourceCreator));
 	}
 
-	// sort the source creators by priority and then by name
-	if (UnsortedSourceCreators.Num() > 1)
+	// Sanity check
 	{
-		struct FSortData
+		// Test if all NeededBindingSources in inside the UnsortedSourceCreators
+		for (const TSharedRef<FCompilerBindingSource>& BindingSource : NeededBindingSources)
 		{
-			FSortData() = default;
-			FName SourceName;
-			FName ParentSourceName;
-			int32 SortIndex = -1;
-
-			void CalculateSortIndex(TMap<FName, FSortData>& Map)
+			const FMVVMViewClass_Source* FoundClassSource = UnsortedSourceCreators.FindByPredicate([ToFindName = BindingSource->Name](const FMVVMViewClass_Source& Other){ return Other.GetName() == ToFindName; });
+			if (FoundClassSource == nullptr)
 			{
-				if (SortIndex < 0)
-				{
-					if (ParentSourceName.IsNone())
-					{
-						SortIndex = 0;
-					}
-					else
-					{
-						Map[ParentSourceName].CalculateSortIndex(Map);
-						// calculate the Depth recursively
-						SortIndex = Map[ParentSourceName].SortIndex + 1;
-					}
-				}
+				AddMessage(FText::Format(LOCTEXT("CompileSources_MissingSources", "Internal error. The source {0} was not compiled.."), FText::FromName(BindingSource->Name))
+					, EMessageType::Warning
+				);
 			}
-		};
-		TMap<FName, FSortData> SortDatas;
-		SortDatas.Reserve(UnsortedSourceCreators.Num());
-		for (int32 Index = 0; Index < UnsortedSourceCreators.Num(); ++Index)
-		{
-			FSortData SortData;
-			SortData.SourceName = UnsortedSourceCreators[Index].GetSourceName();
-			SortData.ParentSourceName =UnsortedSourceCreators[Index].GetParentSourceName();
-			SortDatas.Emplace(SortData.SourceName, SortData);
 		}
+
+		ensure(SortDatas.Num() == UnsortedSourceCreators.Num());
+	}
+
+	// sort the source creators by priority and then by name
+	if (UnsortedSourceCreators.Num() > 1)
+	{
 		for (auto& SortDataPair : SortDatas)
 		{
 			SortDataPair.Value.CalculateSortIndex(SortDatas);
 		}
 
-		UnsortedSourceCreators.Sort([&SortDatas](const FMVVMViewClass_SourceCreator& A, const FMVVMViewClass_SourceCreator& B)
+		UnsortedSourceCreators.Sort([&SortDatas](const FMVVMViewClass_Source& A, const FMVVMViewClass_Source& B)
 			{
-				int32 ASortIndex = SortDatas[A.GetSourceName()].SortIndex;
-				int32 BSortIndex = SortDatas[B.GetSourceName()].SortIndex;
+				int32 ASortIndex = SortDatas[A.GetName()].SortIndex;
+				int32 BSortIndex = SortDatas[B.GetName()].SortIndex;
 				if (ASortIndex == BSortIndex)
 				{
-					return A.GetSourceName().LexicalLess(B.GetSourceName());
+					return A.GetName().LexicalLess(B.GetName());
 				}
 				return ASortIndex < BSortIndex;
 			});
 	}
 
-	// Add the sorted viewmode array to the ViewExtension
-	ViewExtension->SourceCreators = MoveTemp(UnsortedSourceCreators);
+	// Add the sorted viewmodel array to the ViewExtension
+	ViewExtension->Sources = MoveTemp(UnsortedSourceCreators);
+
+	if (ViewExtension->Sources.Num() > 64)
+	{
+		// The view use a uint64 bitfield to filter which viewmodel/source is valid/initialized.
+		//You can use the MVVM.LogViewCompiledResult command to display the compile result and see the name of the sources.
+		AddMessage(LOCTEXT("TooManySources", "There is too many sources for the view. Try spliting your widget into other widgets."), EMessageType::Error);
+	}
 }
 
 
@@ -2291,176 +2334,197 @@ void FMVVMViewBlueprintCompiler::CompileBindings(const FCompiledBindingLibraryCo
 		return;
 	}
 
-	struct FSortData
+	// Sort the array to have a predictable list
 	{
-		int32 ValidBindingIndex = INDEX_NONE;
-		int32 SourceIndex = INDEX_NONE;
-		FMVVMViewClass_CompiledBinding CompiledBinding;
-	};
-	TArray<FSortData> TempBindingToSort;
-	TempBindingToSort.Reserve(ValidBindings.Num()*2);
+		UMVVMBlueprintView* LocalBlueprintView = BlueprintView.Get();
+		ValidBindings.StableSort([LocalBlueprintView](const TSharedRef<FCompilerBinding>& A, const TSharedRef<FCompilerBinding>& B)
+			{
+				const FMVVMBlueprintViewBinding& BindingA = *(LocalBlueprintView->GetBindingAt(A->Key.ViewBindingIndex));
+				const FMVVMBlueprintViewBinding& BindingB = *(LocalBlueprintView->GetBindingAt(B->Key.ViewBindingIndex));
+				return BindingA.BindingId < BindingB.BindingId;
+			});
+	}
 
+	ensure(ViewExtension->Bindings.Num() == 0);
+	ViewExtension->Bindings.Reset(ValidBindings.Num());
 	for (int32 ValidBindingIndex = 0; ValidBindingIndex < ValidBindings.Num(); ++ValidBindingIndex)
 	{
 		const TSharedRef<FCompilerBinding>& ValidBinding = ValidBindings[ValidBindingIndex];
 		const FMVVMBlueprintViewBinding& Binding = *(BlueprintView->GetBindingAt(ValidBinding->Key.ViewBindingIndex));
 
+		const FMVVMVCompiledBinding* CompiledBinding = CompileResult.Bindings.Find(ValidBinding->BindingHandle);
+		if (CompiledBinding == nullptr)
+		{
+			AddMessageForBinding(Binding, LOCTEXT("CompiledBindingNotGenerated", "Could not generate compiled binding."), EMessageType::Error, FName());
+			bIsCompileStepValid = false;
+			continue;
+		}
+
+		FMVVMViewClass_BindingKey BindingKey = FMVVMViewClass_BindingKey(ViewExtension->Bindings.AddDefaulted());
+		FMVVMViewClass_Binding& NewBinding = ViewExtension->Bindings[BindingKey.GetIndex()];
+		NewBinding.Binding = CompiledBinding ? *CompiledBinding : FMVVMVCompiledBinding();
+		NewBinding.ExecutionMode = Binding.bOverrideExecutionMode ? Binding.OverrideExecutionMode : (EMVVMExecutionMode)CVarDefaultExecutionMode->GetInt();
+		NewBinding.SourceBitField = 0;
+		NewBinding.EditorId = Binding.BindingId;
+
+		NewBinding.Flags = 0;
+		NewBinding.Flags |= (!ValidBinding->bIsOneTimeBinding) ? (uint8)FMVVMViewClass_Binding::EFlags::OneWay : 0;
+		NewBinding.Flags |= (Binding.bOverrideExecutionMode) ? (uint8)FMVVMViewClass_Binding::EFlags::OverrideExecuteMode : 0;
+		NewBinding.Flags |= (ValidBinding->ReadPaths.Num() > 1) ? (uint8)FMVVMViewClass_Binding::EFlags::Shared : 0;
+		NewBinding.Flags |= (Binding.bEnabled) ? (uint8)FMVVMViewClass_Binding::EFlags::EnabledByDefault : 0;
+
+		TArray<FMVVMViewClass_SourceKey, TInlineAllocator<16>>  SharedBindings;
 		for (TSharedPtr<FGeneratedReadFieldPathContext> ReadPath : ValidBinding->ReadPaths)
 		{
 			if (ReadPath->OptionalSource == nullptr)
 			{
-				AddMessageForBinding(Binding, LOCTEXT("InvalidSourceInternal", "Internal error. The binding doesn't have a valid source."), EMessageType::Error, FName());
+				AddMessageForBinding(Binding, LOCTEXT("InvalidSourceInternal", "Internal error. The binding has an invalid source."), EMessageType::Error, FName());
 				bIsCompileStepValid = false;
 				continue;
 			}
 
-			int32 ViewExtensionSourceCreatorsIndex = ViewExtension->SourceCreators.IndexOfByPredicate([LookFor = ReadPath->OptionalSource->Name](const FMVVMViewClass_SourceCreator& Other)
+			int32 ViewExtensionSourceCreatorsIndex = ViewExtension->Sources.IndexOfByPredicate([LookFor = ReadPath->OptionalSource->Name](const FMVVMViewClass_Source& Other)
 				{
-					return Other.GetSourceName() == LookFor;
+					return Other.GetName() == LookFor;
 				});
-			if (!ViewExtension->SourceCreators.IsValidIndex(ViewExtensionSourceCreatorsIndex))
+			if (!ViewExtension->Sources.IsValidIndex(ViewExtensionSourceCreatorsIndex))
 			{
 				AddMessageForBinding(Binding, LOCTEXT("CompiledSourceCreatorNotGenerated", "Internal error. The source creator was not generated."), EMessageType::Error, FName());
 				bIsCompileStepValid = false;
 				continue;
 			}
 
-			const FMVVMVCompiledFieldId* CompiledFieldId = (ReadPath->NotificationField && !ValidBinding->bIsOneTimeBinding) ? CompileResult.FieldIds.Find(ReadPath->NotificationField->LibraryCompilerHandle) : nullptr;
+			// Add the needed source.
+			FMVVMViewClass_SourceKey FieldClassSourceKey = FMVVMViewClass_SourceKey(ViewExtensionSourceCreatorsIndex);
+			NewBinding.SourceBitField |= FieldClassSourceKey.GetBit();
 
-			const FMVVMVCompiledBinding* CompiledBinding = CompileResult.Bindings.Find(ValidBinding->BindingHandle);
-			if (CompiledBinding == nullptr)
+			// FieldId
+			const bool bHasField = (ReadPath->NotificationField && !ValidBinding->bIsOneTimeBinding);
+			const UE::FieldNotification::FFieldId* CompiledFieldId = bHasField ? CompileResult.FieldIds.Find(ReadPath->NotificationField->LibraryCompilerHandle) : nullptr;
+			if (CompiledFieldId == nullptr && bHasField)
 			{
-				AddMessageForBinding(Binding, LOCTEXT("CompiledBindingNotGenerated", "Could not generate compiled binding."), EMessageType::Error, FName());
+				AddMessageForBinding(Binding, LOCTEXT("CompiledFieldNotGenerated", "Internal error. The FieldId was not generated."), EMessageType::Error, FName());
 				bIsCompileStepValid = false;
 				continue;
 			}
 
 			bool bExecuteAtInitialization = ValidBinding->Key.bIsForwardBinding;
-			bool bIsOptional = ViewExtension->SourceCreators[ViewExtensionSourceCreatorsIndex].IsOptional();
 
-			FMVVMViewClass_CompiledBinding NewBinding;
-			NewBinding.FieldId = CompiledFieldId ? *CompiledFieldId : FMVVMVCompiledFieldId();
-			NewBinding.SourcePropertyName = ReadPath->OptionalSource->Name;
-			NewBinding.Binding = CompiledBinding ? *CompiledBinding : FMVVMVCompiledBinding();
-			NewBinding.ExecutionMode = Binding.bOverrideExecutionMode ? Binding.OverrideExecutionMode : (EMVVMExecutionMode)CVarDefaultExecutionMode->GetInt();
-			NewBinding.EvaluateSourceCreatorIndex = INDEX_NONE;;
-			NewBinding.EditorId = Binding.BindingId;
+			FMVVMViewClass_Source& ClassSource = ViewExtension->Sources[ViewExtensionSourceCreatorsIndex];
+			FMVVMViewClass_SourceBinding& NewSourceBinding = ClassSource.Bindings.AddDefaulted_GetRef();
+			NewSourceBinding.BindingKey = BindingKey;
+
+			NewSourceBinding.Flags = 0;
+			NewSourceBinding.Flags |= (bExecuteAtInitialization) ? (uint8)FMVVMViewClass_SourceBinding::EFlags::ExecuteAtInitialization : 0;
 
-			NewBinding.Flags = 0;
-			NewBinding.Flags |= (Binding.bEnabled) ? FMVVMViewClass_CompiledBinding::EBindingFlags::EnabledByDefault : 0;
-			NewBinding.Flags |= (bExecuteAtInitialization) ? FMVVMViewClass_CompiledBinding::EBindingFlags::ExecuteAtInitialization : 0;
-			NewBinding.Flags |= (ValidBinding->bIsOneTimeBinding) ? FMVVMViewClass_CompiledBinding::EBindingFlags::OneTime : 0;
-			NewBinding.Flags |= (bIsOptional) ? FMVVMViewClass_CompiledBinding::EBindingFlags::ViewModelOptional : 0;
-			NewBinding.Flags |= (Binding.bOverrideExecutionMode) ? FMVVMViewClass_CompiledBinding::EBindingFlags::OverrideExecuteMode : 0;
-			NewBinding.Flags |= (ReadPath->OptionalSource->Type == FCompilerBindingSource::EType::Self) ? FMVVMViewClass_CompiledBinding::EBindingFlags::SourceObjectIsSelf : 0;
+			if (CompiledFieldId)
+			{
+				NewSourceBinding.FieldId = FFieldNotificationId(CompiledFieldId->GetName());
+				ClassSource.FieldToRegisterTo.AddUnique(FMVVMViewClass_FieldId(*CompiledFieldId));
 
-			FSortData& SortData = TempBindingToSort.AddDefaulted_GetRef();
-			SortData.ValidBindingIndex = ValidBindingIndex;
-			SortData.SourceIndex = ViewExtensionSourceCreatorsIndex;
-			SortData.CompiledBinding = MoveTemp(NewBinding);
+				// Count the number of shared instance of that binding. (they can come from the same Source)
+				//This flag is used at runtime to know if we should delay the binding execution.
+				//We only delay when the field changes. It should only be shared if it has field.
+				SharedBindings.Add(FieldClassSourceKey);
+			}
+		}
+
+		NewBinding.Flags |= (SharedBindings.Num() > 0) ? (uint8)FMVVMViewClass_Binding::EFlags::Shared : 0;
+
+		// Only the last binding in the complex conversion.
+		if (SharedBindings.Num() > 0)
+		{
+			SharedBindings.Sort([](const FMVVMViewClass_SourceKey& A, const FMVVMViewClass_SourceKey&B)
+				{
+					return A.GetIndex() < B.GetIndex();
+				});
+
+			//Only the last one has ExecuteAtInitialization. Remove the flag on all the others.
+			for (int32 Index = 0; Index < SharedBindings.Num() - 2; ++Index)
+			{
+				FMVVMViewClass_Source& ClassSource = ViewExtension->Sources[SharedBindings[Index].GetIndex()];
+				for (FMVVMViewClass_SourceBinding& SourceBinding : ClassSource.Bindings)
+				{
+					if (SourceBinding.GetBindingKey() == BindingKey)
+					{
+						SourceBinding.Flags &= ~(uint8)FMVVMViewClass_SourceBinding::EFlags::ExecuteAtInitialization;
+					}
+				}
+			}
+			// keep only one ExecuteAtInitialization on the last source
+			{
+				FMVVMViewClass_Source& ClassSource = ViewExtension->Sources[SharedBindings.Last().GetIndex()];
+				bool bFound = false;
+				for (int32 Index = ClassSource.Bindings.Num() - 1; Index >= 0; --Index)
+				{
+					FMVVMViewClass_SourceBinding& SourceBinding = ClassSource.Bindings[Index];
+					if (SourceBinding.GetBindingKey() == BindingKey)
+					{
+						if (bFound)
+						{
+							SourceBinding.Flags &= ~(uint8)FMVVMViewClass_SourceBinding::EFlags::ExecuteAtInitialization;
+						}
+						bFound = true;
+					}
+				}
+			}
 		}
 	}
+}
+
 
+void FMVVMViewBlueprintCompiler::CompileEvaluateSources(const FCompiledBindingLibraryCompiler::FCompileResult& CompileResult, UWidgetBlueprintGeneratedClass* Class, UMVVMViewClass* ViewExtension)
+{
 	for (const TSharedRef<FCompilerSourceViewModelDynamicCreatorContext>& ViewModelDynamic : SourceViewModelDynamicCreatorContexts)
 	{
-		int32 ViewExtensionSourceCreatorsIndex = ViewExtension->SourceCreators.IndexOfByPredicate([LookFor = ViewModelDynamic->Source->Name](const FMVVMViewClass_SourceCreator& Other)
+		int32 ViewExtensionSourceCreatorsIndex = ViewExtension->Sources.IndexOfByPredicate([LookFor = ViewModelDynamic->Source->Name](const FMVVMViewClass_Source& Other)
 			{
-				return Other.GetSourceName() == LookFor;
+				return Other.GetName() == LookFor;
 			});
-		if (!ViewExtension->SourceCreators.IsValidIndex(ViewExtensionSourceCreatorsIndex))
+		if (!ViewExtension->Sources.IsValidIndex(ViewExtensionSourceCreatorsIndex))
 		{
 			AddMessage(LOCTEXT("CompiledSourceCreatorNotGenerated", "Internal error. The source creator was not generated."), EMessageType::Error);
 			bIsCompileStepValid = false;
 			continue;
 		}
+		int32 ViewExtensionParentCreatorsIndex = ViewExtension->Sources.IndexOfByPredicate([LookFor = ViewModelDynamic->ParentSource->Name](const FMVVMViewClass_Source& Other)
+			{
+				return Other.GetName() == LookFor;
+			});
+		if (!ViewExtension->Sources.IsValidIndex(ViewExtensionParentCreatorsIndex))
+		{
+			AddMessage(LOCTEXT("CompiledParentSourceCreatorNotGenerated", "Internal error. The parent source creator was not generated."), EMessageType::Error);
+			bIsCompileStepValid = false;
+			continue;
+		}
 
-		const FMVVMVCompiledFieldId* CompiledFieldId = CompileResult.FieldIds.Find(ViewModelDynamic->NotificationIdLibraryCompilerHandle);
+		const UE::FieldNotification::FFieldId* CompiledFieldId = CompileResult.FieldIds.Find(ViewModelDynamic->NotificationIdLibraryCompilerHandle);
 		if (CompiledFieldId == nullptr)
 		{
 			AddMessage(LOCTEXT("CompiledFieldNotifyNotGenerated", "Internal error. The field notify was not generated."), EMessageType::Error);
 			bIsCompileStepValid = false;
 			continue;
 		}
-		
-		bool bIsOptional = ViewModelDynamic->Source->bIsOptional;
 
-		FMVVMViewClass_CompiledBinding NewBinding;
-		NewBinding.FieldId = CompiledFieldId ? *CompiledFieldId : FMVVMVCompiledFieldId();
-		NewBinding.SourcePropertyName = ViewModelDynamic->ParentSource->Name;
-		NewBinding.Binding = FMVVMVCompiledBinding();
-		NewBinding.ExecutionMode = EMVVMExecutionMode::Immediate;
-		NewBinding.EvaluateSourceCreatorIndex = ViewExtensionSourceCreatorsIndex;
-		NewBinding.EditorId = FGuid();
+		FMVVMViewClass_EvaluateSource& NewBinding = ViewExtension->EvaluateSources.AddDefaulted_GetRef();
+		NewBinding.ParentFieldId = FFieldNotificationId(CompiledFieldId->GetName());
+		NewBinding.ParentSource = FMVVMViewClass_SourceKey(ViewExtensionParentCreatorsIndex);
+		NewBinding.ToEvaluate = FMVVMViewClass_SourceKey(ViewExtensionSourceCreatorsIndex);
 
-		NewBinding.Flags = 0;
-		NewBinding.Flags |= FMVVMViewClass_CompiledBinding::EBindingFlags::EnabledByDefault;
-		NewBinding.Flags |= (bIsOptional) ? FMVVMViewClass_CompiledBinding::EBindingFlags::ViewModelOptional : 0;
-		NewBinding.Flags |= FMVVMViewClass_CompiledBinding::EBindingFlags::OverrideExecuteMode; // The mode needs to be Immediate.
-		NewBinding.Flags |= (ViewModelDynamic->ParentSource->Type == FCompilerBindingSource::EType::Self) ? FMVVMViewClass_CompiledBinding::EBindingFlags::SourceObjectIsSelf : 0;
+		FMVVMViewClass_Source& ClassSource = ViewExtension->Sources[ViewExtensionParentCreatorsIndex];
+		ClassSource.FieldToRegisterTo.AddUnique(FMVVMViewClass_FieldId(*CompiledFieldId));
 
-		FSortData& SortData = TempBindingToSort.AddDefaulted_GetRef();
-		SortData.ValidBindingIndex = INDEX_NONE;
-		SortData.SourceIndex = ViewExtensionSourceCreatorsIndex;
-		SortData.CompiledBinding = MoveTemp(NewBinding);
+		ClassSource.Flags |= (uint16)FMVVMViewClass_Source::EFlags::HasEvaluatedBindings;
 	}
 
-
-	// sort the array
-	if (ViewExtension->SourceCreators.Num() > 1)
-	{
-		// sort by source, then regular binding, then "dynamic viewmodel" at the end
-		TempBindingToSort.StableSort([](const FSortData& A, const FSortData& B)
-			{
-				if (A.SourceIndex == B.SourceIndex)
-				{
-					if (A.CompiledBinding.GetEvaluateSourceCreatorBindingIndex() == B.CompiledBinding.GetEvaluateSourceCreatorBindingIndex())
-					{
-						return A.CompiledBinding.GetSourceName().LexicalLess(B.CompiledBinding.GetSourceName());
-					}
-
-					if (A.CompiledBinding.GetEvaluateSourceCreatorBindingIndex() == INDEX_NONE)
-					{
-						return false;
-					}
-					if (B.CompiledBinding.GetEvaluateSourceCreatorBindingIndex() == INDEX_NONE)
-					{
-						return true;
-					}
-
-					return A.CompiledBinding.GetEvaluateSourceCreatorBindingIndex() < B.CompiledBinding.GetEvaluateSourceCreatorBindingIndex();
-				}
-
-				return A.SourceIndex < B.SourceIndex;
-			});
-	}
-
-	// Go backward to be able to mark only the last binding in the complex conversion.
-	TSet<int32> ComplexFunctionInitialized;
-	for (int32 Index = TempBindingToSort.Num() - 1; Index >= 0; --Index)
-	{
-		FSortData& Data = TempBindingToSort[Index];
-
-		// Only handle bindings with complex conversions.
-		if (Data.CompiledBinding.IsConversionFunctionComplex())
+	ViewExtension->EvaluateSources.StableSort([](const FMVVMViewClass_EvaluateSource& A, const FMVVMViewClass_EvaluateSource& B)
 		{
-			// The last binding keep the init flag. Remove it for the others
-			if (ComplexFunctionInitialized.Contains(Data.ValidBindingIndex))
-			{
-				Data.CompiledBinding.Flags &= ~FMVVMViewClass_CompiledBinding::EBindingFlags::ExecuteAtInitialization;
-			}
-			else
+			if (A.GetParentSource() == B.GetParentSource())
 			{
-				ComplexFunctionInitialized.Add(Data.ValidBindingIndex);
+				return A.GetFieldId().GetFieldName().Compare(B.GetFieldId().GetFieldName()) < 0;
 			}
-		}
-	}
-
-	ViewExtension->CompiledBindings.Reset(TempBindingToSort.Num());
-	for (FSortData& Data : TempBindingToSort)
-	{
-		ViewExtension->CompiledBindings.Emplace(MoveTemp(Data.CompiledBinding));
-	}
-	TempBindingToSort.Reset();
+			return A.GetParentSource().GetIndex() < B.GetParentSource().GetIndex();
+		});
 }
 
 
@@ -2584,10 +2648,41 @@ void FMVVMViewBlueprintCompiler::CompileEvents(const FCompiledBindingLibraryComp
 			continue;
 		}
 
-		FMVVMViewClass_CompiledEvent& NewBinding = ViewExtension->CompiledEvents.AddDefaulted_GetRef();
+		int32 FoundSourceIndex = INDEX_NONE;
+		if (!ValidEvent->SourceName.IsNone())
+		{
+			FoundSourceIndex = ViewExtension->Sources.IndexOfByPredicate([ToFind = ValidEvent->SourceName](const FMVVMViewClass_Source& Other)
+				{
+					return Other.GetName() == ToFind;
+				});
+		}
+
+		FMVVMViewClass_Event& NewBinding = ViewExtension->Events.AddDefaulted_GetRef();
 		NewBinding.FieldPath = *CompiledFieldPath;
-		NewBinding.FunctionName = ValidEvent->GeneratedGraphName;
-		NewBinding.SourceName = ValidEvent->SourceName;
+		NewBinding.UserWidgetFunctionName = ValidEvent->GeneratedGraphName;
+		NewBinding.SourceToReevaluate = FoundSourceIndex != INDEX_NONE ? FMVVMViewClass_SourceKey(FoundSourceIndex) : FMVVMViewClass_SourceKey();
+	}
+}
+
+
+void FMVVMViewBlueprintCompiler::SortSourceFields(const FCompiledBindingLibraryCompiler::FCompileResult& CompileResult, UWidgetBlueprintGeneratedClass* Class, UMVVMViewClass* ViewExtension)
+{
+	for (FMVVMViewClass_Source& Source : ViewExtension->Sources)
+	{
+		for (const FMVVMViewClass_FieldId& Field : Source.FieldToRegisterTo)
+		{
+			if (Field.GetFieldId().GetIndex() == INDEX_NONE)
+			{
+				AddMessage(LOCTEXT("SortSourceFieldsInvalidId", "Internal error. The id is invalid."), EMessageType::Error);
+				bIsCompileStepValid = false;
+				continue;
+			}
+		}
+
+		Source.FieldToRegisterTo.StableSort([](const FMVVMViewClass_FieldId& A, const FMVVMViewClass_FieldId& B)
+			{
+				return A.GetFieldId().GetIndex() < B.GetFieldId().GetIndex();
+			});
 	}
 }
 
diff --git a/Engine/Plugins/Runtime/ModelViewViewModel/Source/ModelViewViewModelBlueprint/Private/MVVMViewBlueprintCompiler.h b/Engine/Plugins/Runtime/ModelViewViewModel/Source/ModelViewViewModelBlueprint/Private/MVVMViewBlueprintCompiler.h
--- a/Engine/Plugins/Runtime/ModelViewViewModel/Source/ModelViewViewModelBlueprint/Private/MVVMViewBlueprintCompiler.h
+++ b/Engine/Plugins/Runtime/ModelViewViewModel/Source/ModelViewViewModelBlueprint/Private/MVVMViewBlueprintCompiler.h
@@ -99,17 +99,19 @@ private:
 	void PreCompileEvents(UWidgetBlueprintGeneratedClass* Class);
 
 	// Compile
-	void CompileViewModelCreatorContexts(const FCompiledBindingLibraryCompiler::FCompileResult& CompileResult, UWidgetBlueprintGeneratedClass* Class, UMVVMViewClass* ViewExtension);
+	void CompileSources(const FCompiledBindingLibraryCompiler::FCompileResult& CompileResult, UWidgetBlueprintGeneratedClass* Class, UMVVMViewClass* ViewExtension);
 	void CompileBindings(const FCompiledBindingLibraryCompiler::FCompileResult& CompileResult, UWidgetBlueprintGeneratedClass* Class, UMVVMViewClass* ViewExtension);
+	void CompileEvaluateSources(const FCompiledBindingLibraryCompiler::FCompileResult& CompileResult, UWidgetBlueprintGeneratedClass* Class, UMVVMViewClass* ViewExtension);
 	void CompileEvents(const FCompiledBindingLibraryCompiler::FCompileResult& CompileResult, UWidgetBlueprintGeneratedClass* Class, UMVVMViewClass* ViewExtension);
+	void SortSourceFields(const FCompiledBindingLibraryCompiler::FCompileResult& CompileResult, UWidgetBlueprintGeneratedClass* Class, UMVVMViewClass* ViewExtension);
 
 private:
 	/**
-	 * List of all the sources needed by the view to register the bindings.
+	 * List of all the sources needed by the view to register/execute the bindings.
 	 * They could be viewmodel, widget or any properties on the UserWidget.
 	 * They could also be a viewmodel in a long path.
 	 * It may not have an associated property (dynamicviewmodel).
-	 * It may only have OneTime binding
+	 * It may only have OneTime binding.
 	 */
 	struct FCompilerBindingSource
 	{
@@ -323,7 +325,7 @@ private:
 		TSharedPtr<FGeneratedWriteFieldPathContext> DelegateFieldPath;
 
 		FName GeneratedGraphName;
-		FName SourceName;
+		FName SourceName; // may not be in the NeededBindingSources
 		FCompiledBindingLibraryCompiler::FFieldPathHandle DelegateFieldPathHandle;
 	};
 	TArray<TSharedRef<FCompilerEvent>> ValidEvents;
diff --git a/Engine/Plugins/Runtime/ModelViewViewModel/Source/ModelViewViewModelBlueprint/Private/MVVMWidgetBlueprintExtension_View.cpp b/Engine/Plugins/Runtime/ModelViewViewModel/Source/ModelViewViewModelBlueprint/Private/MVVMWidgetBlueprintExtension_View.cpp
--- a/Engine/Plugins/Runtime/ModelViewViewModel/Source/ModelViewViewModelBlueprint/Private/MVVMWidgetBlueprintExtension_View.cpp
+++ b/Engine/Plugins/Runtime/ModelViewViewModel/Source/ModelViewViewModelBlueprint/Private/MVVMWidgetBlueprintExtension_View.cpp
@@ -104,6 +104,10 @@ void UMVVMWidgetBlueprintExtension_View::HandleBeginCompilation(FWidgetBlueprint
 	if (BlueprintView)
 	{
 		BlueprintView->ResetBindingMessages();
+		for (UMVVMBlueprintViewEvent* ViewEvent : BlueprintView->GetEvents())
+		{
+			ViewEvent->ResetCompilationMessages();
+		}
 		CurrentCompilerContext = MakePimpl<UE::MVVM::Private::FMVVMViewBlueprintCompiler>(InCreationContext, GetBlueprintView());
 	}
 }
@@ -174,7 +178,7 @@ void UMVVMWidgetBlueprintExtension_View::HandleFinishCompilingClass(UWidgetBluep
 			check(ViewExtension);
 
 			// Does it have any bindings
-			if (const_cast<const UMVVMViewClass*>(ViewExtension)->GetCompiledBindings().Num() > 0)
+			if (const_cast<const UMVVMViewClass*>(ViewExtension)->GetBindings().Num() > 0)
 			{
 				// Test if parent also has a view
 				if (Class->GetExtension<UMVVMViewClass>(true))
diff --git a/Engine/Plugins/Runtime/ModelViewViewModel/Source/ModelViewViewModelBlueprint/Public/Bindings/MVVMCompiledBindingLibraryCompiler.h b/Engine/Plugins/Runtime/ModelViewViewModel/Source/ModelViewViewModelBlueprint/Public/Bindings/MVVMCompiledBindingLibraryCompiler.h
--- a/Engine/Plugins/Runtime/ModelViewViewModel/Source/ModelViewViewModelBlueprint/Public/Bindings/MVVMCompiledBindingLibraryCompiler.h
+++ b/Engine/Plugins/Runtime/ModelViewViewModel/Source/ModelViewViewModelBlueprint/Public/Bindings/MVVMCompiledBindingLibraryCompiler.h
@@ -178,7 +178,7 @@ public:
 		FMVVMCompiledBindingLibrary Library;
 		TMap<FFieldPathHandle, FMVVMVCompiledFieldPath> FieldPaths;
 		TMap<FBindingHandle, FMVVMVCompiledBinding> Bindings;
-		TMap<FFieldIdHandle, FMVVMVCompiledFieldId> FieldIds;
+		TMap<FFieldIdHandle, UE::FieldNotification::FFieldId> FieldIds;
 	};
 
 	/** */
diff --git a/Engine/Plugins/Runtime/ModelViewViewModel/Source/ModelViewViewModelBlueprint/Public/MVVMBlueprintView.h b/Engine/Plugins/Runtime/ModelViewViewModel/Source/ModelViewViewModelBlueprint/Public/MVVMBlueprintView.h
--- a/Engine/Plugins/Runtime/ModelViewViewModel/Source/ModelViewViewModelBlueprint/Public/MVVMBlueprintView.h
+++ b/Engine/Plugins/Runtime/ModelViewViewModel/Source/ModelViewViewModelBlueprint/Public/MVVMBlueprintView.h
@@ -43,8 +43,9 @@ public:
 	 * If false, the user will have to initialize the sources manually.
 	 * It prevents the sources evaluating until you are ready.
 	 */
-	UPROPERTY(EditAnywhere, Category = "Viewmodel")
+	UPROPERTY(EditAnywhere, Category = "View")
 	bool bInitializeSourcesOnConstruct = true;
+
 	/**
 	 * Auto initialize the view bindings when the Widget is constructed.
 	 * If false, the user will have to initialize the bindings manually.
@@ -53,8 +54,15 @@ public:
 	 * @note Sources needs to be initialized before initializing the bindings.
 	 * @note When Sources is manually initialized, the bindings will also be initialized if this is true.
 	 */
-	UPROPERTY(EditAnywhere, Category = "Viewmodel", meta=(EditCondition="bInitializeSourcesOnConstruct"))
+	UPROPERTY(EditAnywhere, Category = "View", meta=(EditCondition="bInitializeSourcesOnConstruct"))
 	bool bInitializeBindingsOnConstruct = true;
+
+	/**
+	 * Auto initialize the view events when the Widget is constructed.
+	 * If false, the user will have to initialize the event manually.
+	 */
+	UPROPERTY(EditAnywhere, Category = "View")
+	bool bInitializeEventsOnConstruct = true;
 };
 
 /**
diff --git a/Engine/Plugins/Runtime/ModelViewViewModel/Source/ModelViewViewModelDebugger/Private/MVVMDebugSnapshot.cpp b/Engine/Plugins/Runtime/ModelViewViewModel/Source/ModelViewViewModelDebugger/Private/MVVMDebugSnapshot.cpp
--- a/Engine/Plugins/Runtime/ModelViewViewModel/Source/ModelViewViewModelDebugger/Private/MVVMDebugSnapshot.cpp
+++ b/Engine/Plugins/Runtime/ModelViewViewModel/Source/ModelViewViewModelDebugger/Private/MVVMDebugSnapshot.cpp
@@ -20,17 +20,18 @@ namespace UE::MVVM
 {
 namespace Private
 {
-FMVVMViewSourceDebugEntry CreateSourceInstanceEntry(const FMVVMViewSource& ViewSource)
+
+FMVVMViewSourceDebugEntry CreateSourceInstanceEntry(const UMVVMView* View, const FMVVMView_Source& ViewSource)
 {
 	FMVVMViewSourceDebugEntry Result;
-	Result.SourceInstanceName = ViewSource.SourceName;
+	Result.SourceInstanceName = View->GetViewClass()->GetSource(ViewSource.ClassKey).GetName();
 	Result.SourceAsset = ViewSource.Source ? FAssetData(ViewSource.Source->GetClass()) : FAssetData();
 	Result.ViewModelDebugId = FGuid::NewGuid();
 	Result.LiveSource = ViewSource.Source;
 	return Result;
 }
 
-FMVVMViewBindingDebugEntry CreateViewBindingDebugEntry(const UMVVMViewClass* ViewClass, const FMVVMViewClass_CompiledBinding& Binding, int32 Index)
+FMVVMViewBindingDebugEntry CreateViewBindingDebugEntry(const UMVVMViewClass* ViewClass, const FMVVMViewClass_Binding& Binding, int32 Index)
 {
 	FMVVMViewBindingDebugEntry Result;
 #if WITH_EDITOR
@@ -38,11 +39,6 @@ FMVVMViewBindingDebugEntry CreateViewBindingDebugEntry(const UMVVMViewClass* Vie
 #endif
 	if (ViewClass->GetBindingLibrary().IsLoaded())
 	{
-		if (Binding.GetSourceFieldId().IsValid())
-		{
-			TValueOrError<UE::FieldNotification::FFieldId, void> FieldIdValue = ViewClass->GetBindingLibrary().GetFieldId(Binding.GetSourceFieldId());
-			Result.FieldId = FieldIdValue.HasValue() ? FieldIdValue.GetValue().GetName() : FName();
-		}
 		if (Binding.GetBinding().GetSourceFieldPath().IsValid())
 		{
 			TValueOrError<FString, FString> SourceFieldPathValue = ViewClass->GetBindingLibrary().FieldPathToString(Binding.GetBinding().GetSourceFieldPath(), false);
@@ -63,16 +59,6 @@ FMVVMViewBindingDebugEntry CreateViewBindingDebugEntry(const UMVVMViewClass* Vie
 	return Result;
 }
 
-FMVVMViewSourceDebugEntry CreateViewModelEntry(const FMVVMViewSource& ViewSource)
-{
-	FMVVMViewSourceDebugEntry Result;
-	Result.SourceInstanceName = ViewSource.SourceName;
-	Result.SourceAsset = ViewSource.Source ? FAssetData(ViewSource.Source->GetClass()) : FAssetData();
-	Result.ViewModelDebugId = FGuid::NewGuid();
-	Result.LiveSource = ViewSource.Source;
-	return Result;
-}
-
 FMVVMViewModelFieldBoundDebugEntry CreateViewModelFieldBoundDebugEntry(const UE::FieldNotification::FFieldMulticastDelegate::FDelegateView& DelegateView)
 {
 	FMVVMViewModelFieldBoundDebugEntry Result;
@@ -125,9 +111,10 @@ TSharedPtr<FDebugSnapshot> FDebugSnapshot::CreateSnapshot()
 			DebugEntry->WorldName = World ? World->GetFName() : FName();
 			DebugEntry->UserWidgetAsset = FAssetData(UserWidget->GetClass());
 		}
-		for (const FMVVMViewSource& ViewSource : View->GetSources())
+
+		for (const FMVVMView_Source& ViewSource : View->GetSources())
 		{
-			DebugEntry->Sources.Add(UE::MVVM::Private::CreateSourceInstanceEntry(ViewSource));
+			DebugEntry->Sources.Add(Private::CreateSourceInstanceEntry(View, ViewSource));
 		}
 		DebugEntry->ViewClassDebugId = Snapshot->FindOrAddViewClassEntry(View->GetViewClass())->ViewClassDebugId;
 		DebugEntry->ViewInstanceDebugId = FGuid::NewGuid();
@@ -151,7 +138,7 @@ TSharedPtr<FDebugSnapshot> FDebugSnapshot::CreateSnapshot()
 		DebugEntry->ViewModelAsset = FAssetData(ViewModel->GetClass());
 		for (const UE::FieldNotification::FFieldMulticastDelegate::FDelegateView& DelegateView : ViewModel->GetNotificationDelegateView())
 		{
-			DebugEntry->FieldBound.Add(UE::MVVM::Private::CreateViewModelFieldBoundDebugEntry(DelegateView));
+			DebugEntry->FieldBound.Add(Private::CreateViewModelFieldBoundDebugEntry(DelegateView));
 		}
 		//DebugEntry.PropertyBag;
 		DebugEntry->ViewModelDebugId = FGuid::NewGuid();
@@ -178,7 +165,7 @@ TSharedRef<FMVVMViewClassDebugEntry> FDebugSnapshot::FindOrAddViewClassEntry(con
 	NewEntry->LiveViewClass = ViewClass;
 
 	int32 Index = 0;
-	for (const FMVVMViewClass_CompiledBinding& Binding : ViewClass->GetCompiledBindings())
+	for (const FMVVMViewClass_Binding& Binding : ViewClass->GetBindings())
 	{
 		NewEntry->Bindings.Add(Private::CreateViewBindingDebugEntry(ViewClass, Binding, Index));
 		++Index;
diff --git a/Engine/Plugins/Runtime/ModelViewViewModel/Source/ModelViewViewModelDebugger/Public/MVVMDebugViewClass.h b/Engine/Plugins/Runtime/ModelViewViewModel/Source/ModelViewViewModelDebugger/Public/MVVMDebugViewClass.h
--- a/Engine/Plugins/Runtime/ModelViewViewModel/Source/ModelViewViewModelDebugger/Public/MVVMDebugViewClass.h
+++ b/Engine/Plugins/Runtime/ModelViewViewModel/Source/ModelViewViewModelDebugger/Public/MVVMDebugViewClass.h
@@ -18,9 +18,6 @@ struct FMVVMViewBindingDebugEntry
 
 	UPROPERTY(VisibleAnywhere, Category = "Viewmodel")
 	FGuid BlueprintViewBindingId;
-
-	UPROPERTY(VisibleAnywhere, Category = "Viewmodel")
-	FName FieldId;
 	
 	UPROPERTY(VisibleAnywhere, Category = "Viewmodel")
 	FString SourceFieldPath;
diff --git a/Engine/Plugins/Runtime/ModelViewViewModel/Source/ModelViewViewModelEditor/Private/Widgets/SMVVMPreviewSourcePanel.cpp b/Engine/Plugins/Runtime/ModelViewViewModel/Source/ModelViewViewModelEditor/Private/Widgets/SMVVMPreviewSourcePanel.cpp
--- a/Engine/Plugins/Runtime/ModelViewViewModel/Source/ModelViewViewModelEditor/Private/Widgets/SMVVMPreviewSourcePanel.cpp
+++ b/Engine/Plugins/Runtime/ModelViewViewModel/Source/ModelViewViewModelEditor/Private/Widgets/SMVVMPreviewSourcePanel.cpp
@@ -4,6 +4,7 @@
 
 #include "MVVMSubsystem.h"
 #include "View/MVVMView.h"
+#include "View/MVVMViewClass.h"
 
 #include "Preview/PreviewMode.h"
 #include "WidgetBlueprintEditor.h"
@@ -102,9 +103,10 @@ void SPreviewSourcePanel::HandlePreviewWidgetChanged()
 			if (UMVVMView* View = UMVVMSubsystem::GetViewFromUserWidget(NewWidget))
 			{
 				WeakView = View;
-				for (const FMVVMViewSource& Source : View->GetSources())
+				for (const FMVVMView_Source& ViewSource : View->GetSources())
 				{
-					SourceList.Emplace(MakeShared<Private::SPreviewSourceEntry>(Source.Source, Source.SourceName));
+					FName SourceName = View->GetViewClass()->GetSource(ViewSource.ClassKey).GetName();
+					SourceList.Emplace(MakeShared<Private::SPreviewSourceEntry>(ViewSource.Source, SourceName));
 				}
 			}
 		}
