MVVM: Change the description when there is no viewmodel or bindings or events. Change the condition for when adding the view to the Generated UserWidget.

diff --git a/Engine/Plugins/Runtime/ModelViewViewModel/Source/ModelViewViewModelBlueprint/Private/MVVMWidgetBlueprintExtension_View.cpp b/Engine/Plugins/Runtime/ModelViewViewModel/Source/ModelViewViewModelBlueprint/Private/MVVMWidgetBlueprintExtension_View.cpp
--- a/Engine/Plugins/Runtime/ModelViewViewModel/Source/ModelViewViewModelBlueprint/Private/MVVMWidgetBlueprintExtension_View.cpp
+++ b/Engine/Plugins/Runtime/ModelViewViewModel/Source/ModelViewViewModelBlueprint/Private/MVVMWidgetBlueprintExtension_View.cpp
@@ -178,7 +178,7 @@ void UMVVMWidgetBlueprintExtension_View::HandleFinishCompilingClass(UWidgetBluep
 			check(ViewExtension);
 
 			// Does it have any bindings
-			if (const_cast<const UMVVMViewClass*>(ViewExtension)->GetBindings().Num() > 0)
+			if (ViewExtension->GetBindings().Num() > 0 || ViewExtension->GetEvents().Num() > 0)
 			{
 				// Test if parent also has a view
 				if (Class->GetExtension<UMVVMViewClass>(true))
diff --git a/Engine/Plugins/Runtime/ModelViewViewModel/Source/ModelViewViewModelEditor/Private/Widgets/SMVVMViewBindingPanel.cpp b/Engine/Plugins/Runtime/ModelViewViewModel/Source/ModelViewViewModelEditor/Private/Widgets/SMVVMViewBindingPanel.cpp
--- a/Engine/Plugins/Runtime/ModelViewViewModel/Source/ModelViewViewModelEditor/Private/Widgets/SMVVMViewBindingPanel.cpp
+++ b/Engine/Plugins/Runtime/ModelViewViewModel/Source/ModelViewViewModelEditor/Private/Widgets/SMVVMViewBindingPanel.cpp
@@ -655,7 +655,7 @@ TSharedRef<SWidget> SBindingsPanel::GenerateEditViewWidget()
 				.AutoWidth()
 				[
 					SNew(STextBlock)
-					.Text(LOCTEXT("MissingViewModel", "This editor requires a viewmodel that widgets can bind to, would you like to create a viewmodel now?"))
+					.Text(LOCTEXT("MissingViewModel", "This editor requires a viewmodel that widgets can bind to, would you like to add a viewmodel now?"))
 				]
 				+ SHorizontalBox::Slot()
 				.Padding(0, 0, 20, 0)
@@ -666,7 +666,7 @@ TSharedRef<SWidget> SBindingsPanel::GenerateEditViewWidget()
 					.OnClicked(this, &SBindingsPanel::HandleCreateViewModelClicked)
 					[
 						SNew(STextBlock)
-						.Text(LOCTEXT("CreateViewModel", "Create Viewmodel"))
+						.Text(LOCTEXT("CreateViewModel", "Add Viewmodel"))
 					]
 				]
 			]
diff --git a/Engine/Plugins/Runtime/ModelViewViewModel/Source/ModelViewViewModelEditor/Private/Widgets/SMVVMViewModelPanel.cpp b/Engine/Plugins/Runtime/ModelViewViewModel/Source/ModelViewViewModelEditor/Private/Widgets/SMVVMViewModelPanel.cpp
--- a/Engine/Plugins/Runtime/ModelViewViewModel/Source/ModelViewViewModelEditor/Private/Widgets/SMVVMViewModelPanel.cpp
+++ b/Engine/Plugins/Runtime/ModelViewViewModel/Source/ModelViewViewModelEditor/Private/Widgets/SMVVMViewModelPanel.cpp
@@ -886,7 +886,7 @@ EVisibility SMVVMViewModelPanel::GetWarningPanelVisibility() const
 	{
 		if (UMVVMBlueprintView* WidgetBlueprint = WeakBlueprintView.Get())
 		{
-			if (WidgetBlueprint->GetNumBindings() == 0 && WidgetBlueprint->GetViewModels().Num() != 0)
+			if (WidgetBlueprint->GetNumBindings() == 0 && WidgetBlueprint->GetEvents().Num() == 0 && WidgetBlueprint->GetViewModels().Num() != 0)
 			{
 				return EVisibility::Visible;
 			}
@@ -909,7 +909,7 @@ FText SMVVMViewModelPanel::GetWarningMessage() const
 		{
 			if (WidgetBlueprint->GetNumBindings() == 0 && WidgetBlueprint->GetViewModels().Num() != 0)
 			{
-				return LOCTEXT("NoBindingWarningDescription", "No view will be created for this widget because there are no bindings. Your viewmodels won't be initialized.");
+				return LOCTEXT("NoBindingWarningDescription", "No view will be created for this widget. There are no bindings and there are no events. Your viewmodels won't be initialized.");
 			}
 			if (!WidgetBlueprint->GetSettings()->bInitializeSourcesOnConstruct || !WidgetBlueprint->GetSettings()->bInitializeBindingsOnConstruct)
 			{
