MVVM: Fix localization issue.

diff --git a/Engine/Plugins/Runtime/ModelViewViewModel/Source/ModelViewViewModel/Private/View/MVVMView.cpp b/Engine/Plugins/Runtime/ModelViewViewModel/Source/ModelViewViewModel/Private/View/MVVMView.cpp
--- a/Engine/Plugins/Runtime/ModelViewViewModel/Source/ModelViewViewModel/Private/View/MVVMView.cpp
+++ b/Engine/Plugins/Runtime/ModelViewViewModel/Source/ModelViewViewModel/Private/View/MVVMView.cpp
@@ -701,7 +701,7 @@ bool UMVVMView::SetViewModelByClass(TScriptInterface<INotifyFieldValueChanged> N
 	if (ClassExtension == nullptr)
 	{
 		UE::MVVM::FMessageLog Log(GetUserWidget());
-		Log.Error(LOCTEXT("SetViewModelInvalidClass", "The view is not initialized."));
+		Log.Error(LOCTEXT("SetViewModelInvalidClassExtension", "The view is not initialized."));
 		return false;
 	}
 
