MVVM: Move rows code from ViewBindingListView.cpp to different cpp files. No behavior modifications.

diff --git a/Engine/Plugins/Runtime/ModelViewViewModel/Source/ModelViewViewModelEditor/Private/Types/MVVMBindingEntry.cpp b/Engine/Plugins/Runtime/ModelViewViewModel/Source/ModelViewViewModelEditor/Private/Types/MVVMBindingEntry.cpp
new file mode 100644
--- /dev/null
+++ b/Engine/Plugins/Runtime/ModelViewViewModel/Source/ModelViewViewModelEditor/Private/Types/MVVMBindingEntry.cpp
@@ -0,0 +1,149 @@
+// Copyright Epic Games, Inc. All Rights Reserved.
+
+#include "Types/MVVMBindingEntry.h"
+#include "MVVMBlueprintView.h"
+#include "MVVMBlueprintViewEvent.h"
+
+
+#define LOCTEXT_NAMESPACE "BindingEntry"
+
+namespace UE::MVVM
+{
+
+FMVVMBlueprintViewBinding* FBindingEntry::GetBinding(UMVVMBlueprintView* View) const
+{
+	return View->GetBinding(BindingId);
+}
+
+const FMVVMBlueprintViewBinding* FBindingEntry::GetBinding(const UMVVMBlueprintView* View) const
+{
+	return View->GetBinding(BindingId);
+}
+
+void FBindingEntry::SetGroup(FName WidgetName)
+{
+	check(RowType == ERowType::None);
+	RowType = ERowType::Group;
+	Name = WidgetName;
+	bGroupIsWidget = true;
+}
+
+void FBindingEntry::SetGroup(FName ViewModelName, FGuid ViewModelId)
+{
+	if (!ViewModelId.IsValid())
+	{
+		SetGroup(ViewModelName);
+	}
+	else
+	{
+		check(RowType == ERowType::None);
+		RowType = ERowType::Group;
+		Name = ViewModelName;
+		BindingId = ViewModelId;
+		bGroupIsWidget = false;
+	}
+}
+
+void FBindingEntry::SetBindingId(FGuid Id)
+{
+	check(RowType == ERowType::None);
+	RowType = ERowType::Binding;
+	BindingId = Id;
+}
+
+void FBindingEntry::SetBindingParameter(FGuid Id, FName ParameterName)
+{
+	check(RowType == ERowType::None);
+	RowType = ERowType::BindingParameter;
+	BindingId = Id;
+	Name = ParameterName;
+}
+
+UMVVMBlueprintViewEvent* FBindingEntry::GetEvent() const
+{
+	return Event.Get();
+}
+
+void FBindingEntry::SetEvent(UMVVMBlueprintViewEvent* InEvent)
+{
+	check(RowType == ERowType::None);
+	RowType = ERowType::Event;
+	Event = InEvent;
+}
+
+void FBindingEntry::SetEventParameter(UMVVMBlueprintViewEvent* InEvent, FName ParameterName)
+{
+	check(RowType == ERowType::None);
+	RowType = ERowType::EventParameter;
+	Event = InEvent;
+	Name = ParameterName;
+}
+
+void FBindingEntry::AddChild(TSharedPtr<FBindingEntry> Child)
+{
+	AllChildren.Add(Child);
+}
+	
+void FBindingEntry::AddFilteredChild(TSharedPtr<FBindingEntry> Child)
+{
+	FilteredChildren.Add(Child);
+	bUseFilteredChildren = true;
+}
+
+void FBindingEntry::ResetChildren()
+{
+	AllChildren.Reset();
+	FilteredChildren.Reset();
+	bUseFilteredChildren = false;
+}
+
+void FBindingEntry::SetUseFilteredChildList()
+{
+	bUseFilteredChildren = true;
+}
+
+bool FBindingEntry::operator==(const FBindingEntry& Other) const
+{
+	return RowType == Other.RowType
+		&& Name == Other.Name
+		&& BindingId == Other.BindingId
+		&& Event == Other.Event;
+}
+
+FString FBindingEntry::GetSearchNameString(UMVVMBlueprintView* View, UWidgetBlueprint* WidgetBP)
+{
+	FString RowToString;
+	FString FunctionKeywords;
+
+	switch (RowType)
+	{
+	case UE::MVVM::FBindingEntry::ERowType::Group:
+	case UE::MVVM::FBindingEntry::ERowType::BindingParameter:
+	case UE::MVVM::FBindingEntry::ERowType::EventParameter:
+		RowToString = Name.ToString();
+		break;
+	case UE::MVVM::FBindingEntry::ERowType::Binding:
+		{
+			FMVVMBlueprintViewBinding* BindingInRow = GetBinding(View);
+			check(BindingInRow);
+			RowToString.Append(BindingInRow->GetSearchableString(WidgetBP));
+		}
+		break;
+	case UE::MVVM::FBindingEntry::ERowType::Event:
+		if (UMVVMBlueprintViewEvent* EventInRow = GetEvent())
+		{
+			RowToString.Append(EventInRow->GetSearchableString());
+		}
+		break;
+	default:
+		break;
+	}
+
+	RowToString.ReplaceInline(TEXT(" "), TEXT(""));
+
+	return RowToString;
+}
+
+} // namespace
+
+#undef LOCTEXT_NAMESPACE
\ No newline at end of file
diff --git a/Engine/Plugins/Runtime/ModelViewViewModel/Source/ModelViewViewModelEditor/Private/Types/MVVMBindingEntry.h b/Engine/Plugins/Runtime/ModelViewViewModel/Source/ModelViewViewModelEditor/Private/Types/MVVMBindingEntry.h
new file mode 100644
--- /dev/null
+++ b/Engine/Plugins/Runtime/ModelViewViewModel/Source/ModelViewViewModelEditor/Private/Types/MVVMBindingEntry.h
@@ -0,0 +1,124 @@
+// Copyright Epic Games, Inc. All Rights Reserved.
+
+#pragma once
+
+#include "CoreMinimal.h"
+#include "UObject/Object.h"
+#include "UObject/WeakObjectPtr.h"
+#include "UObject/WeakObjectPtrTemplates.h"
+
+class UMVVMBlueprintView;
+class UMVVMBlueprintViewEvent;
+struct FMVVMBlueprintViewBinding;
+class UWidgetBlueprint;
+
+namespace UE::MVVM
+{
+
+/**
+ * A structure for the different entries in the BindingList
+ */
+struct FBindingEntry
+{
+	enum class ERowType
+	{
+		None,
+		Group,
+		Binding,
+		BindingParameter,
+		Event,
+		EventParameter,
+	};
+
+	FMVVMBlueprintViewBinding* GetBinding(UMVVMBlueprintView* View) const;
+
+	const FMVVMBlueprintViewBinding* GetBinding(const UMVVMBlueprintView* View) const;
+
+	ERowType GetRowType() const
+	{
+		return RowType;
+	}
+
+	//~ group
+	FName GetGroupName() const
+	{
+		return Name;
+	}
+	
+	FGuid GetGroupAsViewModel() const
+	{
+		return BindingId;
+	}
+	
+	bool IsGroupWidget() const
+	{
+		return bGroupIsWidget;
+	}
+	
+	void SetGroup(FName WidgetName);
+	void SetGroup(FName ViewModelName, FGuid ViewModelId);
+
+	//~ binding
+	FGuid GetBindingId() const
+	{
+		return BindingId;
+	}
+
+	void SetBindingId(FGuid Id);
+
+	//~ binding parameter
+	FName GetBindingParameterName() const
+	{
+		return Name;
+	}
+
+	void SetBindingParameter(FGuid Id, FName ParameterName);
+
+	//~ event
+	UMVVMBlueprintViewEvent* GetEvent() const;
+
+	void SetEvent(UMVVMBlueprintViewEvent* InEvent);
+
+	//~ event parameter
+	FName GetEventParameterName() const
+	{
+		return Name;
+	}
+
+	void SetEventParameter(UMVVMBlueprintViewEvent* InEvent, FName ParameterName);
+
+	//~ children
+	TConstArrayView<TSharedPtr<FBindingEntry>> GetAllChildren() const
+	{
+		return AllChildren;
+	}
+
+	TConstArrayView<TSharedPtr<FBindingEntry>> GetFilteredChildren() const
+	{
+		return bUseFilteredChildren ? FilteredChildren : AllChildren;
+	}
+
+	void AddChild(TSharedPtr<FBindingEntry> Child);
+
+	void AddFilteredChild(TSharedPtr<FBindingEntry> Child);
+
+	void ResetChildren();
+
+	void SetUseFilteredChildList();
+
+	bool operator==(const FBindingEntry& Other) const;
+
+	FString GetSearchNameString(UMVVMBlueprintView* View, UWidgetBlueprint* WidgetBP);
+	
+private:
+	ERowType RowType = ERowType::None;
+	FName Name;
+	FGuid BindingId;
+	TWeakObjectPtr<UMVVMBlueprintViewEvent> Event;
+	TArray<TSharedPtr<FBindingEntry>> AllChildren;
+	TArray<TSharedPtr<FBindingEntry>> FilteredChildren;
+	bool bGroupIsWidget = false;
+	bool bUseFilteredChildren = false;
+};
+
+} // namespace
diff --git a/Engine/Plugins/Runtime/ModelViewViewModel/Source/ModelViewViewModelEditor/Private/Widgets/BindingEntry/SMVVMBaseRow.cpp b/Engine/Plugins/Runtime/ModelViewViewModel/Source/ModelViewViewModelEditor/Private/Widgets/BindingEntry/SMVVMBaseRow.cpp
new file mode 100644
--- /dev/null
+++ b/Engine/Plugins/Runtime/ModelViewViewModel/Source/ModelViewViewModelEditor/Private/Widgets/BindingEntry/SMVVMBaseRow.cpp
@@ -0,0 +1,50 @@
+// Copyright Epic Games, Inc. All Rights Reserved.
+
+#include "Widgets/BindingEntry/SMVVMBaseRow.h"
+
+#include "Styling/MVVMEditorStyle.h"
+
+#include "Widgets/SNullWidget.h"
+
+namespace UE::MVVM::BindingEntry
+{
+
+void SBaseRow::Construct(const FArguments& Args, const TSharedRef<STableViewBase>& OwnerTableView, const TSharedPtr<FWidgetBlueprintEditor>& InBlueprintEditor, UWidgetBlueprint* InBlueprint, const TSharedPtr<FBindingEntry>& InEntry)
+{
+	BindingEntry = InEntry;
+	WeakBlueprint = InBlueprint;
+	WeakBlueprintEditor = InBlueprintEditor;
+
+	STableRow<TSharedPtr<FBindingEntry>>::Construct(
+		STableRow<TSharedPtr<FBindingEntry>>::FArguments()
+		.Padding(2.0f)
+		.Style(FMVVMEditorStyle::Get(), GetTableRowStyle())
+		[
+			BuildRowWidget()
+		],
+		OwnerTableView
+	);
+}
+
+TSharedRef<SWidget> SBaseRow::BuildRowWidget()
+{
+	return SNullWidget::NullWidget;
+}
+	
+UWidgetBlueprint* SBaseRow::GetBlueprint() const
+{
+	return WeakBlueprint.Get();
+}
+
+UMVVMBlueprintView* SBaseRow::GetBlueprintView() const
+{
+	return GEditor->GetEditorSubsystem<UMVVMEditorSubsystem>()->GetView(GetBlueprint());
+}
+
+UMVVMEditorSubsystem* SBaseRow::GetEditorSubsystem() const
+{
+	return GEditor->GetEditorSubsystem<UMVVMEditorSubsystem>();
+}
+
+
+} // namespace UE::MVVM::BindingEntry
diff --git a/Engine/Plugins/Runtime/ModelViewViewModel/Source/ModelViewViewModelEditor/Private/Widgets/BindingEntry/SMVVMBaseRow.h b/Engine/Plugins/Runtime/ModelViewViewModel/Source/ModelViewViewModelEditor/Private/Widgets/BindingEntry/SMVVMBaseRow.h
new file mode 100644
--- /dev/null
+++ b/Engine/Plugins/Runtime/ModelViewViewModel/Source/ModelViewViewModelEditor/Private/Widgets/BindingEntry/SMVVMBaseRow.h
@@ -0,0 +1,52 @@
+// Copyright Epic Games, Inc. All Rights Reserved.
+
+#pragma once
+
+#include "Types/MVVMBindingEntry.h"
+#include "Widgets/Views/SListView.h"
+#include "Widgets/Views/STableRow.h"
+
+#include "MVVMEditorSubsystem.h"
+#include "MVVMBlueprintView.h"
+#include "WidgetBlueprint.h"
+#include "WidgetBlueprintEditor.h"
+
+namespace UE::MVVM::BindingEntry
+{
+
+/**
+ * 
+ */
+class SBaseRow : public STableRow<TSharedPtr<FBindingEntry>>
+{
+public:
+	SLATE_BEGIN_ARGS(SBaseRow) {}
+	SLATE_END_ARGS()
+
+	void Construct(const FArguments& Args, const TSharedRef<STableViewBase>& OwnerTableView, const TSharedPtr<FWidgetBlueprintEditor>& InBlueprintEditor, UWidgetBlueprint* InBlueprint, const TSharedPtr<FBindingEntry>& InEntry);
+
+	TSharedPtr<FBindingEntry> GetEntry() const
+	{
+		return BindingEntry;
+	}
+
+	TSharedPtr<FWidgetBlueprintEditor> GetBlueprintEditor() const
+	{
+		return WeakBlueprintEditor.Pin();
+	}
+	
+	UWidgetBlueprint* GetBlueprint() const;
+	UMVVMBlueprintView* GetBlueprintView() const;
+	UMVVMEditorSubsystem* GetEditorSubsystem() const;
+
+protected:
+	virtual TSharedRef<SWidget> BuildRowWidget() = 0;
+	virtual const ANSICHAR* GetTableRowStyle() const = 0;
+
+private:
+	TSharedPtr<FBindingEntry> BindingEntry;
+	TWeakObjectPtr<UWidgetBlueprint> WeakBlueprint;
+	TWeakPtr<FWidgetBlueprintEditor> WeakBlueprintEditor;
+};
+
+} // namespace UE::MVVM
diff --git a/Engine/Plugins/Runtime/ModelViewViewModel/Source/ModelViewViewModelEditor/Private/Widgets/BindingEntry/SMVVMBindingRow.cpp b/Engine/Plugins/Runtime/ModelViewViewModel/Source/ModelViewViewModelEditor/Private/Widgets/BindingEntry/SMVVMBindingRow.cpp
new file mode 100644
--- /dev/null
+++ b/Engine/Plugins/Runtime/ModelViewViewModel/Source/ModelViewViewModelEditor/Private/Widgets/BindingEntry/SMVVMBindingRow.cpp
@@ -0,0 +1,814 @@
+// Copyright Epic Games, Inc. All Rights Reserved.
+
+
+#include "Widgets/BindingEntry/SMVVMBindingRow.h"
+
+#include "Bindings/MVVMBindingHelper.h"
+#include "Blueprint/WidgetTree.h"
+#include "MVVMBlueprintViewBinding.h"
+#include "MVVMBlueprintViewEvent.h"
+#include "MVVMDeveloperProjectSettings.h"
+#include "MVVMPropertyPath.h"
+
+#include "Details/WidgetPropertyDragDropOp.h"
+#include "DragAndDrop/DecoratedDragDropOp.h"
+#include "Styling/AppStyle.h"
+#include "Styling/MVVMEditorStyle.h"
+
+#include "Dialog/SCustomDialog.h"
+#include "Widgets/Input/SCheckBox.h"
+#include "Widgets/Input/SComboBox.h"
+#include "Widgets/Layout/SBorder.h"
+#include "Widgets/Layout/SSpacer.h"
+#include "Widgets/Text/STextBlock.h"
+#include "Widgets/SBoxPanel.h"
+#include "Widgets/SMVVMFieldSelector.h"
+#include "Widgets/SMVVMSourceSelector.h"
+#include "Widgets/ViewModelFieldDragDropOp.h"
+#include "SSimpleButton.h"
+
+#define LOCTEXT_NAMESPACE "BindingListView_BindingRow"
+
+namespace UE::MVVM::BindingEntry
+{
+
+namespace Private
+{
+
+	TArray<FName>* GetBindingModeNames()
+	{
+		static TArray<FName> BindingModeNames;
+
+		if (BindingModeNames.IsEmpty())
+		{
+			UEnum* ModeEnum = StaticEnum<EMVVMBindingMode>();
+
+			BindingModeNames.Reserve(ModeEnum->NumEnums());
+
+			for (int32 BindingIndex = 0; BindingIndex < ModeEnum->NumEnums() - 1; ++BindingIndex)
+			{
+				const bool bIsHidden = ModeEnum->HasMetaData(TEXT("Hidden"), BindingIndex);
+				if (!bIsHidden)
+				{
+					BindingModeNames.Add(ModeEnum->GetNameByIndex(BindingIndex));
+				}
+			}
+		}
+
+		return &BindingModeNames;
+	}
+
+} // namespace
+
+void SBindingRow::Construct(const FArguments& Args, const TSharedRef<STableViewBase>& OwnerTableView, const TSharedPtr<FWidgetBlueprintEditor>& InBlueprintEditor, UWidgetBlueprint* InBlueprint, const TSharedPtr<FBindingEntry>& InEntry)
+{
+	static IConsoleVariable* CVarDefaultExecutionMode = IConsoleManager::Get().FindConsoleVariable(TEXT("MVVM.DefaultExecutionMode"));
+	ensure(CVarDefaultExecutionMode);
+	DefaultExecutionMode = CVarDefaultExecutionMode ? (EMVVMExecutionMode)CVarDefaultExecutionMode->GetInt() : EMVVMExecutionMode::DelayedWhenSharedElseImmediate;
+
+	SBaseRow::Construct(SBaseRow::FArguments(), OwnerTableView, InBlueprintEditor, InBlueprint, InEntry);
+
+	TSharedPtr<SWidget> ChildContent = ChildSlot.DetachWidget();
+	ChildSlot
+	[
+		// Add a single pixel top and bottom border for this widget.
+		SNew(SBorder)
+		.BorderImage(FAppStyle::Get().GetBrush("Brushes.Recessed"))
+		.Padding(0.0f, 2.0f, 0.0f, 1.0f)
+		[
+			// Restore the border that we're meant to have that reacts to selection/hover/etc.
+			SNew(SBorder)
+			.BorderImage(this, &SBindingRow::GetBorderImage)
+			.Padding(0.0f)
+			[
+				ChildContent.ToSharedRef()
+			]
+		]
+	];
+}
+
+TSharedRef<SWidget> SBindingRow::BuildRowWidget()
+{
+	FMVVMBlueprintViewBinding* ViewBinding = GetThisViewBinding();
+	
+	return SNew(SBorder)
+	.BorderImage(FAppStyle::Get().GetBrush("PlainBorder"))
+	.Padding(0)
+	.BorderBackgroundColor(this, &SBindingRow::GetErrorBorderColor)
+	[
+		SNew(SBox)
+		.HeightOverride(30)
+		[
+			SNew(SHorizontalBox)
+
+			+ SHorizontalBox::Slot()
+			.Padding(2, 0)
+			.VAlign(VAlign_Center)
+			.HAlign(HAlign_Left)
+			.AutoWidth()
+			[
+				SNew(SCheckBox)
+				.IsChecked(this, &SBindingRow::IsBindingCompiled)
+				.OnCheckStateChanged(this, &SBindingRow::OnIsBindingCompileChanged)
+			]
+
+			+ SHorizontalBox::Slot()
+			.Padding(4.f, 0.f)
+			.VAlign(VAlign_Center)
+			.HAlign(HAlign_Left)
+			.AutoWidth()
+			[
+				SNew(SBox)
+				.MinDesiredWidth(150)
+				[
+					SNew(SFieldSelector, GetBlueprint())
+					.OnGetPropertyPath(this, &SBindingRow::GetSelectedPropertyPath, false)
+					.OnGetConversionFunction(this, &SBindingRow::GetSelectedConversionFunction, false)
+					.OnFieldSelectionChanged(this, &SBindingRow::HandleFieldSelectionChanged, false)
+					.OnGetSelectionContext(this, &SBindingRow::GetSelectedSelectionContext, false)
+					.OnDrop(this, &SBindingRow::HandleFieldSelectorDrop, false)
+					.OnDragEnter(this, &SBindingRow::HandleFieldSelectorDragEnter, false)
+					.ShowContext(false)
+				]
+			]
+
+			+ SHorizontalBox::Slot()
+			.Padding(4.f, 0.f)
+			.VAlign(VAlign_Center)
+			.HAlign(HAlign_Left)
+			.AutoWidth()
+			[
+				SNew(SComboBox<FName>)
+				.OptionsSource(Private::GetBindingModeNames())
+				.InitiallySelectedItem(StaticEnum<EMVVMBindingMode>()->GetNameByValue((int64)ViewBinding->BindingType))
+				.OnSelectionChanged(this, &SBindingRow::OnBindingModeSelectionChanged)
+				.OnGenerateWidget(this, &SBindingRow::GenerateBindingModeWidget)
+				.ToolTipText(this, &SBindingRow::GetCurrentBindingModeLabel)
+				.Content()
+				[
+					SNew(SBox)
+					.HAlign(HAlign_Center)
+					.VAlign(VAlign_Center)
+					.WidthOverride(16.0f)
+					.HeightOverride(16.0f)
+					[
+						SNew(SImage)
+						.Image(this, &SBindingRow::GetCurrentBindingModeBrush)
+					]
+				]
+			]
+
+			+ SHorizontalBox::Slot()
+			.Padding(4.f, 0.f)
+			.VAlign(VAlign_Center)
+			.HAlign(HAlign_Left)
+			.AutoWidth()
+			[
+				SNew(SBox)
+				.MinDesiredWidth(150.0f)
+				[
+					SNew(SFieldSelector, GetBlueprint())
+					.OnGetPropertyPath(this, &SBindingRow::GetSelectedPropertyPath, true)
+					.OnGetConversionFunction(this, &SBindingRow::GetSelectedConversionFunction, true)
+					.OnFieldSelectionChanged(this, &SBindingRow::HandleFieldSelectionChanged, true)
+					.OnGetSelectionContext(this, &SBindingRow::GetSelectedSelectionContext, true)
+					.OnDrop(this, &SBindingRow::HandleFieldSelectorDrop, true)
+					.OnDragEnter(this, &SBindingRow::HandleFieldSelectorDragEnter, true)
+				]
+			]
+
+			+ SHorizontalBox::Slot()
+			[
+				SNew(SSpacer)
+			]
+
+			+ SHorizontalBox::Slot()
+			.Padding(2.0f, 1.0f)
+			.VAlign(VAlign_Fill)
+			.HAlign(HAlign_Right)
+			.AutoWidth()
+			[
+				SNew(SCheckBox)
+				.IsChecked(this, &SBindingRow::IsExecutionModeOverrideChecked)
+				.OnCheckStateChanged(this, &SBindingRow::OnExecutionModeOverrideChanged)
+			]
+
+			+ SHorizontalBox::Slot()
+			.Padding(2.0f, 1.0f)
+			.VAlign(VAlign_Fill)
+			.HAlign(HAlign_Right)
+			.AutoWidth()
+			[
+				SNew(SComboButton)
+				.ContentPadding(FMargin(4.f, 0.f))
+				.OnGetMenuContent(this, &SBindingRow::OnGetExecutionModeMenuContent)
+				.IsEnabled(this, &SBindingRow::IsExecutionModeOverridden)
+				.ButtonContent()
+				[
+					SNew(STextBlock)
+					.Text(this, &SBindingRow::GetExecutioModeValue)
+					.ToolTipText(this, &SBindingRow::GetExecutioModeValueToolTip)
+				]
+			]
+
+			+ SHorizontalBox::Slot()
+			.Padding(2.0f, 0.0f)
+			.VAlign(VAlign_Center)
+			.HAlign(HAlign_Left)
+			.AutoWidth()
+			[
+				SNew(SSimpleButton)
+				.Icon(FAppStyle::Get().GetBrush("Icons.Error"))
+				.Visibility(this, &SBindingRow::GetErrorButtonVisibility)
+				.ToolTipText(this, &SBindingRow::GetErrorButtonToolTip)
+				.OnClicked(this, &SBindingRow::OnErrorButtonClicked)
+			]
+		]
+	];
+}
+
+FMVVMBlueprintViewBinding* SBindingRow::GetThisViewBinding() const
+{
+	UMVVMBlueprintView* BlueprintView = GetBlueprintView();
+	return BlueprintView ? GetEntry()->GetBinding(BlueprintView) : nullptr;
+}
+
+FSlateColor SBindingRow::GetErrorBorderColor() const
+{
+	if (const FMVVMBlueprintViewBinding* ViewModelBinding = GetThisViewBinding())
+	{
+		UMVVMBlueprintView* BlueprintView = GetBlueprintView();
+		if (BlueprintView->HasBindingMessage(ViewModelBinding->BindingId, EBindingMessageType::Error))
+		{
+			return FStyleColors::Error;
+		}
+		else if (BlueprintView->HasBindingMessage(ViewModelBinding->BindingId, EBindingMessageType::Warning))
+		{
+			return FStyleColors::Warning;
+		}
+	}
+	return FStyleColors::Transparent;
+}
+
+ECheckBoxState SBindingRow::IsBindingEnabled() const
+{
+	if (const FMVVMBlueprintViewBinding* ViewModelBinding = GetThisViewBinding())
+	{
+		return ViewModelBinding->bEnabled ? ECheckBoxState::Checked : ECheckBoxState::Unchecked;
+	}
+	return ECheckBoxState::Undetermined;
+}
+
+ECheckBoxState SBindingRow::IsBindingCompiled() const
+{
+	if (const FMVVMBlueprintViewBinding* ViewModelBinding = GetThisViewBinding())
+	{
+		return ViewModelBinding->bCompile ? ECheckBoxState::Checked : ECheckBoxState::Unchecked;
+	}
+	return ECheckBoxState::Undetermined;
+}
+
+EVisibility SBindingRow::GetErrorButtonVisibility() const
+{
+	if (const FMVVMBlueprintViewBinding* ViewModelBinding = GetThisViewBinding())
+	{
+		UMVVMBlueprintView* BlueprintView = GetBlueprintView();
+		bool HasBindingError = BlueprintView->HasBindingMessage(ViewModelBinding->BindingId, EBindingMessageType::Error)
+			|| BlueprintView->HasBindingMessage(ViewModelBinding->BindingId, EBindingMessageType::Warning);
+		return HasBindingError ? EVisibility::Visible : EVisibility::Hidden;
+	}
+	return EVisibility::Collapsed;
+}
+
+FText SBindingRow::GetErrorButtonToolTip() const
+{
+	// Get error messages of this binding
+	if (const FMVVMBlueprintViewBinding* ViewModelBinding = GetThisViewBinding())
+	{
+		UMVVMBlueprintView* BlueprintView = GetBlueprintView();
+		TArray<FText> BindingErrorList = BlueprintView->GetBindingMessages(ViewModelBinding->BindingId, EBindingMessageType::Error);
+		TArray<FText> BindingWarningList = BlueprintView->GetBindingMessages(ViewModelBinding->BindingId, EBindingMessageType::Warning);
+		BindingErrorList.Append(BindingWarningList);
+
+		static const FText NewLineText = FText::FromString(TEXT("\n"));
+		FText HintText = LOCTEXT("ErrorButtonText", "Errors and Warnings: (Click to show in a separate window)");
+		FText ErrorsText = FText::Join(NewLineText, BindingErrorList);
+		return FText::Join(NewLineText, HintText, ErrorsText);
+	}
+	return FText();
+}
+
+FReply SBindingRow::OnErrorButtonClicked()
+{
+	ErrorItems.Reset();
+
+	const UWidgetBlueprint* WidgetBlueprint = GetBlueprint();
+	const UMVVMBlueprintView* View = GetBlueprintView();
+	const FMVVMBlueprintViewBinding* ViewBinding = View ? GetEntry()->GetBinding(View) : nullptr;
+	if (WidgetBlueprint && View && ViewBinding)
+	{
+		for (const FText& ErrorText : View->GetBindingMessages(ViewBinding->BindingId, EBindingMessageType::Error))
+		{
+			ErrorItems.Add(MakeShared<FText>(ErrorText));
+		}
+
+		for (const FText& WarningText : View->GetBindingMessages(ViewBinding->BindingId, EBindingMessageType::Warning))
+		{
+			ErrorItems.Add(MakeShared<FText>(WarningText));
+		}
+
+		const FText BindingDisplayName = FText::FromString(ViewBinding->GetDisplayNameString(WidgetBlueprint));
+		TSharedRef<SCustomDialog> ErrorDialog = SNew(SCustomDialog)
+			.Title(FText::Format(LOCTEXT("Compilation Errors and Warnings", "Compilation Errors and Warnings for {0}"), BindingDisplayName))
+			.Buttons({
+				SCustomDialog::FButton(LOCTEXT("OK", "OK"))
+				})
+			.Content()
+			[
+				SNew(SListView<TSharedPtr<FText>>)
+				.ListItemsSource(&ErrorItems)
+				.OnGenerateRow(this, &SBindingRow::OnGenerateErrorRow)
+			];
+
+		ErrorDialog->Show();
+	}
+
+	return FReply::Handled();
+}
+
+TSharedRef<ITableRow> SBindingRow::OnGenerateErrorRow(TSharedPtr<FText> Text, const TSharedRef<STableViewBase>& TableView) const
+{
+	return SNew(STableRow<TSharedPtr<FText>>, TableView)
+	.Content()
+	[
+		SNew(SBorder)
+		.BorderBackgroundColor(FStyleColors::Background)
+		[
+			SNew(STextBlock)
+			.Text(*Text.Get())
+		]
+	];
+}
+
+TArray<FBindingSource> SBindingRow::GetAvailableViewModels() const
+{
+	UMVVMEditorSubsystem* EditorSubsystem = GetEditorSubsystem();
+	return EditorSubsystem->GetAllViewModels(GetBlueprint());
+}
+
+FMVVMBlueprintPropertyPath SBindingRow::GetSelectedPropertyPath(bool bSource) const
+{
+	if (FMVVMBlueprintViewBinding* ViewBinding = GetThisViewBinding())
+	{
+		return bSource ? ViewBinding->SourcePath : ViewBinding->DestinationPath;
+	}
+	return FMVVMBlueprintPropertyPath();
+}
+
+const UFunction* SBindingRow::GetSelectedConversionFunction(bool bSourceToDest) const
+{
+	if (FMVVMBlueprintViewBinding* ViewBinding = GetThisViewBinding())
+	{
+		UMVVMEditorSubsystem* EditorSubsystem = GetEditorSubsystem();
+		return EditorSubsystem->GetConversionFunction(GetBlueprint(), *ViewBinding, bSourceToDest);
+	}
+	return nullptr;
+}
+
+void SBindingRow::HandleFieldSelectionChanged(FMVVMBlueprintPropertyPath SelectedField, const UFunction* Function, bool bSource)
+{
+	UWidgetBlueprint* WidgetBlueprint = GetBlueprint();
+	FMVVMBlueprintViewBinding* ViewBinding = GetThisViewBinding();
+	if (WidgetBlueprint && ViewBinding)
+	{
+		UMVVMEditorSubsystem* Subsystem = GetEditorSubsystem();
+		if (bSource)
+		{
+			Subsystem->SetSourceToDestinationConversionFunction(WidgetBlueprint, *ViewBinding, Function);
+			if (ViewBinding->SourcePath != SelectedField)
+			{
+				Subsystem->SetSourcePathForBinding(WidgetBlueprint, *ViewBinding, SelectedField);
+			}
+		}
+		else
+		{
+			Subsystem->SetDestinationToSourceConversionFunction(WidgetBlueprint, *ViewBinding, Function);
+			if (ViewBinding->DestinationPath != SelectedField)
+			{
+				Subsystem->SetDestinationPathForBinding(WidgetBlueprint, *ViewBinding, SelectedField);
+			}
+		}
+	}
+}
+
+FFieldSelectionContext SBindingRow::GetSelectedSelectionContext(bool bSource) const
+{
+	FFieldSelectionContext Result;
+
+	const UWidgetBlueprint* WidgetBlueprint = GetBlueprint();
+	const FMVVMBlueprintViewBinding* ViewBinding = GetThisViewBinding();
+	if (WidgetBlueprint == nullptr || ViewBinding == nullptr)
+	{
+		return Result;
+	}
+
+	Result.BindingMode = ViewBinding->BindingType;
+
+	{
+		TArray<FMVVMConstFieldVariant> Fields = bSource ? ViewBinding->DestinationPath.GetFields(WidgetBlueprint->SkeletonGeneratedClass) : ViewBinding->SourcePath.GetFields(WidgetBlueprint->SkeletonGeneratedClass);
+		if (Fields.Num() > 0)
+		{
+			FMVVMConstFieldVariant LastField = Fields.Last();
+			if (LastField.IsProperty())
+			{
+				Result.AssignableTo = LastField.GetProperty();
+			}
+			else if (LastField.IsFunction())
+			{
+				if (const UFunction* Function = LastField.GetFunction())
+				{
+					Result.AssignableTo = BindingHelper::GetReturnProperty(Function);
+				}
+			}
+		}
+	}
+
+	if (!bSource && !ViewBinding->DestinationPath.GetWidgetName().IsNone())
+	{
+		Result.FixedBindingSource = FBindingSource::CreateForWidget(WidgetBlueprint, ViewBinding->DestinationPath.GetWidgetName());
+	}
+
+	Result.bAllowWidgets = true;
+	Result.bAllowViewModels = bSource;
+	Result.bAllowConversionFunctions = false;
+	bool bIsReadingValue = (IsForwardBinding(ViewBinding->BindingType) && bSource)
+		|| (IsBackwardBinding(ViewBinding->BindingType) && !bSource);
+	if (!(IsBackwardBinding(ViewBinding->BindingType) && IsForwardBinding(ViewBinding->BindingType)))
+	{
+		Result.bAllowConversionFunctions = bIsReadingValue;
+	}
+
+	Result.bReadable = bIsReadingValue;
+	Result.bWritable = (IsForwardBinding(ViewBinding->BindingType) && !bSource)
+		|| (IsBackwardBinding(ViewBinding->BindingType) && bSource);
+	
+	return Result;
+}
+
+FReply SBindingRow::HandleFieldSelectorDrop(const FGeometry& MyGeometry, const FDragDropEvent& DragDropEvent, bool bSource)
+{
+	TSharedPtr<FDecoratedDragDropOp> DragDropOp = DragDropEvent.GetOperationAs<FDecoratedDragDropOp>();
+	if (!DragDropOp.IsValid())
+	{
+		return FReply::Unhandled();
+	}
+
+	// Accept all drag-drop operations that are widget properties, but only accept view model fields when we are dropping into the Source box.
+	if (!DragDropOp->IsOfType<FWidgetPropertyDragDropOp>() && (!DragDropOp->IsOfType<FViewModelFieldDragDropOp>() || !bSource))
+	{
+		return FReply::Unhandled();
+	}
+
+	FMVVMBlueprintViewBinding* ViewBinding = GetThisViewBinding();
+	UWidgetBlueprint* WidgetBlueprint = GetBlueprint();
+	if (ViewBinding == nullptr || WidgetBlueprint == nullptr)
+	{
+		return FReply::Unhandled();
+	}
+
+	TSharedPtr<FViewModelFieldDragDropOp> ViewModelFieldDragDropOp = DragDropEvent.GetOperationAs<FViewModelFieldDragDropOp>();
+	TSharedPtr<FWidgetPropertyDragDropOp> WidgetPropertyDragDropOp = DragDropEvent.GetOperationAs<FWidgetPropertyDragDropOp>();
+	bool bIsViewModelProperty = false;
+
+	if (ViewModelFieldDragDropOp)
+	{
+		bIsViewModelProperty = true;
+	}
+
+	UWidgetBlueprint* DragDropWidgetBP = bIsViewModelProperty ? ViewModelFieldDragDropOp->WidgetBP.Get() : WidgetPropertyDragDropOp->WidgetBP.Get();
+
+	if (WidgetBlueprint == DragDropWidgetBP)
+	{
+		TArray<FFieldVariant> FieldPath = bIsViewModelProperty ? ViewModelFieldDragDropOp->DraggedField : WidgetPropertyDragDropOp->DraggedPropertyPath;
+		FMVVMBlueprintPropertyPath PropertyPath;
+
+		PropertyPath.ResetPropertyPath();
+		for (const FFieldVariant& Field : FieldPath)
+		{
+			PropertyPath.AppendPropertyPath(WidgetBlueprint, FMVVMConstFieldVariant(Field));
+		}
+
+		UMVVMEditorSubsystem* Subsystem = GetEditorSubsystem();
+		if (bIsViewModelProperty)
+		{
+			if (ViewModelFieldDragDropOp->ViewModelId.IsValid())
+			{
+				PropertyPath.SetViewModelId(ViewModelFieldDragDropOp->ViewModelId);
+			}
+			else
+			{
+				return FReply::Unhandled();
+			}
+		}
+		else
+		{
+			if (UWidget* OwnerWidgetPtr = WidgetPropertyDragDropOp->OwnerWidget.Get())
+			{
+				// Search for the widget by its name in the widget tree
+				// If the widget is not found, we know it is the root preview widget so we use the blueprint name.
+				if (WidgetBlueprint->WidgetTree->FindWidget(OwnerWidgetPtr->GetFName()) && WidgetBlueprint->GetFName() != OwnerWidgetPtr->GetFName())
+				{
+					PropertyPath.SetWidgetName(OwnerWidgetPtr->GetFName());
+				}
+				else
+				{
+					PropertyPath.SetSelfContext();
+				}
+			}
+		}
+
+		if (bSource)
+		{
+			Subsystem->SetSourceToDestinationConversionFunction(WidgetBlueprint, *ViewBinding, nullptr);
+			Subsystem->SetSourcePathForBinding(WidgetBlueprint, *ViewBinding, PropertyPath);
+		}
+		else
+		{
+			Subsystem->SetDestinationToSourceConversionFunction(WidgetBlueprint, *ViewBinding, nullptr);
+			Subsystem->SetDestinationPathForBinding(WidgetBlueprint, *ViewBinding, PropertyPath);
+		}
+		return FReply::Handled();
+	}
+
+	return FReply::Unhandled();
+}
+
+void SBindingRow::HandleFieldSelectorDragEnter(const FGeometry& MyGeometry, const FDragDropEvent& DragDropEvent, bool bSource)
+{
+	TSharedPtr<FDecoratedDragDropOp> DragDropOp = DragDropEvent.GetOperationAs<FDecoratedDragDropOp>();
+	if (DragDropOp.IsValid())
+	{
+		// Accept all drag-drop operations that are widget properties, but only accept view model fields when we are dropping into the Source box.
+		if (DragDropOp->IsOfType<FWidgetPropertyDragDropOp>() || (DragDropOp->IsOfType<FViewModelFieldDragDropOp>() && bSource))
+		{
+			if (UWidgetBlueprint* WidgetBlueprint = GetBlueprint())
+			{
+				TSharedPtr<FViewModelFieldDragDropOp> ViewModelFieldDragDropOp = DragDropEvent.GetOperationAs<FViewModelFieldDragDropOp>();
+				TSharedPtr<FWidgetPropertyDragDropOp> WidgetPropertyDragDropOp = DragDropEvent.GetOperationAs<FWidgetPropertyDragDropOp>();
+				bool IsViewModelProperty = false;
+
+				if (ViewModelFieldDragDropOp)
+				{
+					IsViewModelProperty = true;
+				}
+
+				UWidgetBlueprint* DragDropWidgetBP = IsViewModelProperty ? ViewModelFieldDragDropOp->WidgetBP.Get() : WidgetPropertyDragDropOp->WidgetBP.Get();
+
+				if (DragDropWidgetBP && DragDropWidgetBP == WidgetBlueprint)
+				{
+					DragDropOp->CurrentIconBrush = FAppStyle::GetBrush(TEXT("Graph.ConnectorFeedback.OK"));
+					return;
+				}
+			}
+		}
+		else
+		{
+			DragDropOp->CurrentIconBrush = FAppStyle::GetBrush(TEXT("Graph.ConnectorFeedback.Error"));
+		}
+	}
+}
+
+ECheckBoxState SBindingRow::IsExecutionModeOverrideChecked() const
+{
+	if (FMVVMBlueprintViewBinding* ViewBinding = GetThisViewBinding())
+	{
+		return ViewBinding->bOverrideExecutionMode ? ECheckBoxState::Checked : ECheckBoxState::Unchecked;
+	}
+	return ECheckBoxState::Undetermined;
+}
+
+void SBindingRow::OnExecutionModeOverrideChanged(ECheckBoxState NewState)
+{
+	if (FMVVMBlueprintViewBinding* ViewBinding = GetThisViewBinding())
+	{
+		UMVVMEditorSubsystem* Subsystem = GetEditorSubsystem();
+		if (NewState == ECheckBoxState::Checked)
+		{
+			Subsystem->OverrideExecutionModeForBinding(GetBlueprint(), *ViewBinding, DefaultExecutionMode);
+		}
+		else
+		{
+			Subsystem->ResetExecutionModeForBinding(GetBlueprint(), *ViewBinding);
+		}
+	}
+}
+
+void SBindingRow::OnExecutionModeSelectionChanged(EMVVMExecutionMode Value)
+{
+	if (FMVVMBlueprintViewBinding* ViewBinding = GetThisViewBinding())
+	{
+		UMVVMEditorSubsystem* Subsystem = GetEditorSubsystem();
+		Subsystem->OverrideExecutionModeForBinding(GetBlueprint(), *ViewBinding, Value);
+	}
+}
+
+TSharedRef<SWidget> SBindingRow::OnGetExecutionModeMenuContent()
+{
+	const bool bCloseAfterSelection = true;
+	FMenuBuilder MenuBuilder(bCloseAfterSelection, nullptr, nullptr, true);
+
+	UEnum* Enum = StaticEnum<EMVVMExecutionMode>();
+	for (int32 Index = 0; Index < Enum->NumEnums() - 1; Index++)
+	{
+		EMVVMExecutionMode Mode = static_cast<EMVVMExecutionMode>(Enum->GetValueByIndex(Index));
+		if (!GetDefault<UMVVMDeveloperProjectSettings>()->IsExecutionModeAllowed(Mode))
+		{
+			continue;
+		}
+
+		MenuBuilder.AddMenuEntry(
+			Enum->GetDisplayNameTextByIndex(Index),
+			Enum->GetToolTipTextByIndex(Index),
+			FSlateIcon(),
+			FUIAction
+			(
+				FExecuteAction::CreateLambda([this, Mode]()
+					{
+						OnExecutionModeSelectionChanged(Mode);
+					})
+			)
+		);
+	}
+
+	return MenuBuilder.MakeWidget();
+}
+
+FText SBindingRow::GetExecutioModeValue() const
+{
+	EMVVMExecutionMode ExecutionMode = DefaultExecutionMode;
+	if (FMVVMBlueprintViewBinding* ViewBinding = GetThisViewBinding())
+	{
+		if (ViewBinding->bOverrideExecutionMode)
+		{
+			ExecutionMode = ViewBinding->OverrideExecutionMode;
+		}
+	}
+	return StaticEnum<EMVVMExecutionMode>()->GetDisplayNameTextByValue(static_cast<int64>(ExecutionMode));
+}
+
+FText SBindingRow::GetExecutioModeValueToolTip() const
+{
+	EMVVMExecutionMode ExecutionMode = DefaultExecutionMode;
+	if (FMVVMBlueprintViewBinding* ViewBinding = GetThisViewBinding())
+	{
+		if (ViewBinding->bOverrideExecutionMode)
+		{
+			ExecutionMode = ViewBinding->OverrideExecutionMode;
+		}
+	}
+	return StaticEnum<EMVVMExecutionMode>()->GetToolTipTextByIndex(static_cast<int64>(ExecutionMode));
+}
+
+bool SBindingRow::IsExecutionModeOverridden() const
+{
+	if (FMVVMBlueprintViewBinding* ViewBinding = GetThisViewBinding())
+	{
+		return ViewBinding->bOverrideExecutionMode;
+	}
+	return false;
+}
+
+void SBindingRow::OnIsBindingEnableChanged(ECheckBoxState NewState)
+{
+	if (NewState == ECheckBoxState::Undetermined)
+	{
+		return;
+	}
+
+	if (FMVVMBlueprintViewBinding* ViewBinding = GetThisViewBinding())
+	{
+		UMVVMEditorSubsystem* Subsystem = GetEditorSubsystem();
+		Subsystem->SetEnabledForBinding(GetBlueprint(), *ViewBinding, NewState == ECheckBoxState::Checked);
+	}
+}
+
+void SBindingRow::OnIsBindingCompileChanged(ECheckBoxState NewState)
+{
+	if (NewState == ECheckBoxState::Undetermined)
+	{
+		return;
+	}
+
+	if (FMVVMBlueprintViewBinding* ViewBinding = GetThisViewBinding())
+	{
+		UMVVMEditorSubsystem* Subsystem = GetEditorSubsystem();
+		Subsystem->SetCompileForBinding(GetBlueprint(), *ViewBinding, NewState == ECheckBoxState::Checked);
+	}
+}
+
+const FSlateBrush* SBindingRow::GetBindingModeBrush(EMVVMBindingMode BindingMode) const
+{
+	switch (BindingMode)
+	{
+	case EMVVMBindingMode::OneTimeToDestination:
+		return FMVVMEditorStyle::Get().GetBrush("BindingMode.OneTimeOneWay");
+	case EMVVMBindingMode::OneWayToDestination:
+		return FMVVMEditorStyle::Get().GetBrush("BindingMode.OneWayToSource");
+	case EMVVMBindingMode::OneWayToSource:
+		return FMVVMEditorStyle::Get().GetBrush("BindingMode.OneWay");
+	case EMVVMBindingMode::OneTimeToSource:
+		return nullptr;
+	case EMVVMBindingMode::TwoWay:
+		return FMVVMEditorStyle::Get().GetBrush("BindingMode.TwoWay");
+	default:
+		return nullptr;
+	}
+}
+
+const FSlateBrush* SBindingRow::GetCurrentBindingModeBrush() const
+{
+	if (FMVVMBlueprintViewBinding* ViewBinding = GetThisViewBinding())
+	{
+		return GetBindingModeBrush(ViewBinding->BindingType);
+	}
+	return nullptr;
+}
+
+FText SBindingRow::GetCurrentBindingModeLabel() const
+{
+	if (FMVVMBlueprintViewBinding* ViewBinding = GetThisViewBinding())
+	{
+		return GetBindingModeLabel(ViewBinding->BindingType);
+	}
+	return FText::GetEmpty();
+}
+
+FText SBindingRow::GetBindingModeLabel(EMVVMBindingMode BindingMode) const
+{
+	static FText OneTimeToDestinationLabel = LOCTEXT("OneTimeToDestinationLabel", "One Time To Widget");
+	static FText OneWayToDestinationLabel = LOCTEXT("OneWayToDestinationLabel", "One Way To Widget");
+	static FText OneWayToSourceLabel = LOCTEXT("OneWayToSourceLabel", "One Way To View Model");
+	static FText OneTimeToSourceLabel = LOCTEXT("OneTimeToSourceLabel", "One Time To View Model");
+	static FText TwoWayLabel = LOCTEXT("TwoWayLabel", "Two Way");
+
+	switch (BindingMode)
+	{
+	case EMVVMBindingMode::OneTimeToDestination:
+		return OneTimeToDestinationLabel;
+	case EMVVMBindingMode::OneWayToDestination:
+		return OneWayToDestinationLabel;
+	case EMVVMBindingMode::OneWayToSource:
+		return OneWayToSourceLabel;
+	case EMVVMBindingMode::OneTimeToSource:
+		return OneTimeToSourceLabel;
+	case EMVVMBindingMode::TwoWay:
+		return TwoWayLabel;
+	default:
+		return FText::GetEmpty();
+	}
+}
+
+TSharedRef<SWidget> SBindingRow::GenerateBindingModeWidget(FName ValueName) const
+{
+	const UEnum* ModeEnum = StaticEnum<EMVVMBindingMode>();
+	int32 Index = ModeEnum->GetIndexByName(ValueName);
+	EMVVMBindingMode MVVMBindingMode = EMVVMBindingMode(Index);
+	return SNew(SHorizontalBox)
+	+ SHorizontalBox::Slot()
+	.AutoWidth()
+	[
+		SNew(SBox)
+		.HAlign(HAlign_Center)
+		.VAlign(VAlign_Center)
+		.WidthOverride(16.0f)
+		.HeightOverride(16.0f)
+		[
+			SNew(SImage)
+			.Image(GetBindingModeBrush(MVVMBindingMode))
+		]
+	]
+	+ SHorizontalBox::Slot()
+	.AutoWidth()
+	.Padding(2.0f, 0.0f, 0.0f, 0.0f)
+	[
+		SNew(STextBlock)
+		.Text(GetBindingModeLabel(MVVMBindingMode))
+		.ToolTipText(ModeEnum->GetToolTipTextByIndex(Index))
+	];
+}
+
+void SBindingRow::OnBindingModeSelectionChanged(FName ValueName, ESelectInfo::Type)
+{
+	if (FMVVMBlueprintViewBinding* ViewBinding = GetThisViewBinding())
+	{
+		const UEnum* ModeEnum = StaticEnum<EMVVMBindingMode>();
+		EMVVMBindingMode NewMode = (EMVVMBindingMode)ModeEnum->GetValueByName(ValueName);
+
+		UMVVMEditorSubsystem* Subsystem = GetEditorSubsystem();
+		Subsystem->SetBindingTypeForBinding(GetBlueprint(), *ViewBinding, NewMode);
+	}
+}
+
+} // namespace
+
+#undef LOCTEXT_NAMESPACE
diff --git a/Engine/Plugins/Runtime/ModelViewViewModel/Source/ModelViewViewModelEditor/Private/Widgets/BindingEntry/SMVVMBindingRow.h b/Engine/Plugins/Runtime/ModelViewViewModel/Source/ModelViewViewModelEditor/Private/Widgets/BindingEntry/SMVVMBindingRow.h
new file mode 100644
--- /dev/null
+++ b/Engine/Plugins/Runtime/ModelViewViewModel/Source/ModelViewViewModelEditor/Private/Widgets/BindingEntry/SMVVMBindingRow.h
@@ -0,0 +1,101 @@
+// Copyright Epic Games, Inc. All Rights Reserved.
+
+#pragma once
+
+#include "Input/Reply.h"
+#include "Types/MVVMBindingSource.h"
+#include "Types/MVVMExecutionMode.h"
+#include "Widgets/BindingEntry/SMVVMBaseRow.h"
+#include "Widgets/SMVVMFieldSelectorMenu.h"
+
+struct FMVVMBlueprintViewBinding;
+
+
+namespace UE::MVVM::BindingEntry
+{
+
+/**
+ *
+ */
+class SBindingRow : public SBaseRow
+{
+public:
+	SLATE_BEGIN_ARGS(SBindingRow) {}
+	SLATE_END_ARGS()
+
+	void Construct(const FArguments& Args, const TSharedRef<STableViewBase>& OwnerTableView, const TSharedPtr<FWidgetBlueprintEditor>& InBlueprintEditor, UWidgetBlueprint* InBlueprint, const TSharedPtr<FBindingEntry>& InEntry);
+
+protected:
+	virtual TSharedRef<SWidget> BuildRowWidget() override;
+	virtual const ANSICHAR* GetTableRowStyle() const override
+	{
+		return "BindingView.BindingRow";
+	}
+
+private:
+	FMVVMBlueprintViewBinding* GetThisViewBinding() const;
+
+	FSlateColor GetErrorBorderColor() const;
+
+	ECheckBoxState IsBindingEnabled() const;
+
+	ECheckBoxState IsBindingCompiled() const;
+
+	EVisibility GetErrorButtonVisibility() const;
+
+	FText GetErrorButtonToolTip() const;
+
+	FReply OnErrorButtonClicked();
+
+	TSharedRef<ITableRow> OnGenerateErrorRow(TSharedPtr<FText> Text, const TSharedRef<STableViewBase>& TableView) const;
+
+	TArray<FBindingSource> GetAvailableViewModels() const;
+
+	FMVVMBlueprintPropertyPath GetSelectedPropertyPath(bool bSource) const;
+
+	const UFunction* GetSelectedConversionFunction(bool bSourceToDest) const;
+
+	void HandleFieldSelectionChanged(FMVVMBlueprintPropertyPath SelectedField, const UFunction* Function, bool bSource);
+
+	FFieldSelectionContext GetSelectedSelectionContext(bool bSource) const;
+
+	FReply HandleFieldSelectorDrop(const FGeometry& MyGeometry, const FDragDropEvent& DragDropEvent, bool bSource);
+
+	void HandleFieldSelectorDragEnter(const FGeometry& MyGeometry, const FDragDropEvent& DragDropEvent, bool bSource);
+
+	ECheckBoxState IsExecutionModeOverrideChecked() const;
+
+	void OnExecutionModeOverrideChanged(ECheckBoxState NewState);
+
+	void OnExecutionModeSelectionChanged(EMVVMExecutionMode Value);
+
+	TSharedRef<SWidget> OnGetExecutionModeMenuContent();
+
+	FText GetExecutioModeValue() const;
+
+	FText GetExecutioModeValueToolTip() const;
+
+	bool IsExecutionModeOverridden() const;
+
+	void OnIsBindingEnableChanged(ECheckBoxState NewState);
+
+	void OnIsBindingCompileChanged(ECheckBoxState NewState);
+
+	const FSlateBrush* GetBindingModeBrush(EMVVMBindingMode BindingMode) const;
+
+	const FSlateBrush* GetCurrentBindingModeBrush() const;
+
+	FText GetCurrentBindingModeLabel() const;
+
+	FText GetBindingModeLabel(EMVVMBindingMode BindingMode) const;
+
+	TSharedRef<SWidget> GenerateBindingModeWidget(FName ValueName) const;
+
+	void OnBindingModeSelectionChanged(FName ValueName, ESelectInfo::Type);
+
+private:
+	TArray<TSharedPtr<FText>> ErrorItems;
+	EMVVMExecutionMode DefaultExecutionMode = EMVVMExecutionMode::DelayedWhenSharedElseImmediate;
+};
+
+} // namespace UE::MVVM
diff --git a/Engine/Plugins/Runtime/ModelViewViewModel/Source/ModelViewViewModelEditor/Private/Widgets/BindingEntry/SMVVMEventRow.cpp b/Engine/Plugins/Runtime/ModelViewViewModel/Source/ModelViewViewModelEditor/Private/Widgets/BindingEntry/SMVVMEventRow.cpp
new file mode 100644
--- /dev/null
+++ b/Engine/Plugins/Runtime/ModelViewViewModel/Source/ModelViewViewModelEditor/Private/Widgets/BindingEntry/SMVVMEventRow.cpp
@@ -0,0 +1,476 @@
+// Copyright Epic Games, Inc. All Rights Reserved.
+
+
+#include "Widgets/BindingEntry/SMVVMEventRow.h"
+
+#include "Bindings/MVVMBindingHelper.h"
+#include "Blueprint/WidgetTree.h"
+#include "MVVMBlueprintViewBinding.h"
+#include "MVVMBlueprintViewEvent.h"
+#include "MVVMDeveloperProjectSettings.h"
+#include "MVVMPropertyPath.h"
+
+#include "Details/WidgetPropertyDragDropOp.h"
+#include "DragAndDrop/DecoratedDragDropOp.h"
+#include "Styling/AppStyle.h"
+#include "Styling/MVVMEditorStyle.h"
+
+#include "Dialog/SCustomDialog.h"
+#include "Widgets/Input/SCheckBox.h"
+#include "Widgets/Input/SComboBox.h"
+#include "Widgets/Layout/SBorder.h"
+#include "Widgets/Layout/SSpacer.h"
+#include "Widgets/Text/STextBlock.h"
+#include "Widgets/SBoxPanel.h"
+#include "Widgets/SMVVMFieldSelector.h"
+#include "Widgets/SMVVMSourceSelector.h"
+#include "Widgets/ViewModelFieldDragDropOp.h"
+#include "SSimpleButton.h"
+
+#define LOCTEXT_NAMESPACE "BindingListView_EventRow"
+
+namespace UE::MVVM::BindingEntry
+{
+
+void SEventRow::Construct(const FArguments& Args, const TSharedRef<STableViewBase>& OwnerTableView, const TSharedPtr<FWidgetBlueprintEditor>& InBlueprintEditor, UWidgetBlueprint* InBlueprint, const TSharedPtr<FBindingEntry>& InEntry)
+{
+	SBaseRow::Construct(SBaseRow::FArguments(), OwnerTableView, InBlueprintEditor, InBlueprint, InEntry);
+
+	TSharedPtr<SWidget> ChildContent = ChildSlot.DetachWidget();
+	ChildSlot
+	[
+		// Add a single pixel top and bottom border for this widget.
+		SNew(SBorder)
+		.BorderImage(FAppStyle::Get().GetBrush("Brushes.Recessed"))
+		.Padding(0.0f, 2.0f, 0.0f, 1.0f)
+		[
+			// Restore the border that we're meant to have that reacts to selection/hover/etc.
+			SNew(SBorder)
+			.BorderImage(this, &SEventRow::GetBorderImage)
+			.Padding(0.0f)
+			[
+				ChildContent.ToSharedRef()
+			]
+		]
+	];
+}
+
+TSharedRef<SWidget> SEventRow::BuildRowWidget()
+{
+	return SNew(SBorder)
+	.BorderImage(FAppStyle::Get().GetBrush("PlainBorder"))
+	.Padding(0.0f)
+	.BorderBackgroundColor(this, &SEventRow::GetErrorBorderColor)
+	[
+		SNew(SBox)
+		.HeightOverride(30)
+		[
+			SNew(SHorizontalBox)
+			+ SHorizontalBox::Slot()
+			.Padding(2.0f, 0.0f)
+			.VAlign(VAlign_Center)
+			.HAlign(HAlign_Left)
+			.AutoWidth()
+			[
+				SNew(SCheckBox)
+				.IsChecked(this, &SEventRow::IsEventCompiled)
+				.OnCheckStateChanged(this, &SEventRow::OnIsEventCompileChanged)
+			]
+
+			+ SHorizontalBox::Slot()
+			.Padding(4.0f, 0.0f)
+			.VAlign(VAlign_Center)
+			.HAlign(HAlign_Left)
+			.AutoWidth()
+			[
+				SNew(SBox)
+				.MinDesiredWidth(150.0f)
+				[
+					SNew(SFieldSelector, GetBlueprint())
+					.OnGetPropertyPath(this, &SEventRow::GetSelectedPropertyPath, true)
+					.OnGetConversionFunction(this, &SEventRow::GetSelectedConversionFunction, true)
+					.OnFieldSelectionChanged(this, &SEventRow::HandleFieldSelectionChanged, true)
+					.OnGetSelectionContext(this, &SEventRow::GetSelectedSelectionContext, true)
+					.OnDrop(this, &SEventRow::HandleFieldSelectorDrop, true)
+					.OnDragEnter(this, &SEventRow::HandleFieldSelectorDragEnter, true)
+					.ShowContext(false)
+				]
+			]
+
+			+ SHorizontalBox::Slot()
+			.Padding(4.0f, 0.0f)
+			.VAlign(VAlign_Fill)
+			.HAlign(HAlign_Left)
+			.AutoWidth()
+			[
+				SNew(SBox)
+				.HAlign(HAlign_Center)
+				.VAlign(VAlign_Center)
+				.WidthOverride(16.0f)
+				.HeightOverride(16.0f)
+				[
+					SNew(SImage)
+					.Image(FMVVMEditorStyle::Get().GetBrush("BindingMode.OneWay"))
+				]
+			]
+
+			+ SHorizontalBox::Slot()
+			.Padding(4.0f, 0.0f, 2.0f, 0.0f)
+			.VAlign(VAlign_Center)
+			.HAlign(HAlign_Left)
+			.AutoWidth()
+			[
+				SNew(SBox)
+				.MinDesiredWidth(150.0f)
+				[
+					SNew(SFieldSelector, GetBlueprint())
+					.OnGetPropertyPath(this, &SEventRow::GetSelectedPropertyPath, false)
+					.OnGetConversionFunction(this, &SEventRow::GetSelectedConversionFunction, false)
+					.OnFieldSelectionChanged(this, &SEventRow::HandleFieldSelectionChanged, false)
+					.OnGetSelectionContext(this, &SEventRow::GetSelectedSelectionContext, false)
+					.OnDrop(this, &SEventRow::HandleFieldSelectorDrop, false)
+					.OnDragEnter(this, &SEventRow::HandleFieldSelectorDragEnter, false)
+				]
+			]
+
+			+ SHorizontalBox::Slot()
+			[
+				SNew(SSpacer)
+			]
+
+			+ SHorizontalBox::Slot()
+			.Padding(2.0f, 0.0f)
+			.VAlign(VAlign_Center)
+			.HAlign(HAlign_Left)
+			.AutoWidth()
+			[
+				SNew(SSimpleButton)
+				.Icon(FAppStyle::Get().GetBrush("Icons.Error"))
+				.Visibility(this, &SEventRow::GetErrorButtonVisibility)
+				.ToolTipText(this, &SEventRow::GetErrorButtonToolTip)
+				.OnClicked(this, &SEventRow::OnErrorButtonClicked)
+			]
+		]
+	];
+}
+
+UMVVMBlueprintViewEvent* SEventRow::GetEvent() const
+{
+	return GetEntry()->GetEvent();
+}
+
+FSlateColor SEventRow::GetErrorBorderColor() const
+{
+	if (const UMVVMBlueprintViewEvent* Event = GetEvent())
+	{
+		if (Event->HasCompilationMessage(UMVVMBlueprintViewEvent::EMessageType::Error))
+		{
+			return FStyleColors::Error;
+		}
+		else if (Event->HasCompilationMessage(UMVVMBlueprintViewEvent::EMessageType::Warning))
+		{
+			return FStyleColors::Warning;
+		}
+	}
+	return FStyleColors::Transparent;
+}
+
+EVisibility SEventRow::GetErrorButtonVisibility() const
+{
+	if (const UMVVMBlueprintViewEvent* Event = GetEvent())
+	{
+		bool bHasBindingError = Event->HasCompilationMessage(UMVVMBlueprintViewEvent::EMessageType::Error);
+		bool bHasBindingWarning = Event->HasCompilationMessage(UMVVMBlueprintViewEvent::EMessageType::Warning);
+		return bHasBindingError || bHasBindingWarning ? EVisibility::Visible : EVisibility::Hidden;
+	}
+	return EVisibility::Collapsed;
+}
+
+FText SEventRow::GetErrorButtonToolTip() const
+{
+	if (const UMVVMBlueprintViewEvent* Event = GetEvent())
+	{
+		TArray<FText> BindingErrorList = Event->GetCompilationMessages(UMVVMBlueprintViewEvent::EMessageType::Error);
+		TArray<FText> BindingWarningList = Event->GetCompilationMessages(UMVVMBlueprintViewEvent::EMessageType::Warning);
+		BindingErrorList.Append(BindingWarningList);
+
+		static const FText NewLineText = FText::FromString(TEXT("\n"));
+		FText HintText = LOCTEXT("ErrorButtonText", "Errors and Warnings: (Click to show in a separate window)");
+		FText ErrorsText = FText::Join(NewLineText, BindingErrorList);
+		return FText::Join(NewLineText, HintText, ErrorsText);
+	}
+	return FText();
+}
+
+FReply SEventRow::OnErrorButtonClicked()
+{
+	ErrorItems.Reset();
+
+	if (const UMVVMBlueprintViewEvent* Event = GetEvent())
+	{
+		for (const FText& ErrorText : Event->GetCompilationMessages(UMVVMBlueprintViewEvent::EMessageType::Error))
+		{
+			ErrorItems.Add(MakeShared<FText>(ErrorText));
+		}
+
+		for (const FText& WarningText : Event->GetCompilationMessages(UMVVMBlueprintViewEvent::EMessageType::Warning))
+		{
+			ErrorItems.Add(MakeShared<FText>(WarningText));
+		}
+
+		const FText BindingDisplayName = Event->GetDisplayName(true);
+		TSharedRef<SCustomDialog> ErrorDialog = SNew(SCustomDialog)
+			.Title(FText::Format(LOCTEXT("Compilation Errors and Warnings", "Compilation Errors and Warnings for {0}"), BindingDisplayName))
+			.Buttons({
+				SCustomDialog::FButton(LOCTEXT("OK", "OK"))
+				})
+			.Content()
+			[
+				SNew(SListView<TSharedPtr<FText>>)
+					.ListItemsSource(&ErrorItems)
+					.OnGenerateRow(this, &SEventRow::OnGenerateErrorRow)
+			];
+
+		ErrorDialog->Show();
+	}
+
+	return FReply::Handled();
+}
+
+TSharedRef<ITableRow> SEventRow::OnGenerateErrorRow(TSharedPtr<FText> Text, const TSharedRef<STableViewBase>& TableView) const
+{
+	return SNew(STableRow<TSharedPtr<FText>>, TableView)
+	.Content()
+	[
+		SNew(SBorder)
+		.BorderBackgroundColor(FStyleColors::Background)
+		[
+			SNew(STextBlock)
+			.Text(*Text.Get())
+		]
+	];
+}
+
+ECheckBoxState SEventRow::IsEventEnabled() const
+{
+	if (const UMVVMBlueprintViewEvent* Event = GetEvent())
+	{
+		return Event->bEnabled ? ECheckBoxState::Checked : ECheckBoxState::Unchecked;
+	}
+	return ECheckBoxState::Undetermined;
+}
+
+void SEventRow::OnIsEventEnableChanged(ECheckBoxState NewState)
+{
+	if (NewState == ECheckBoxState::Undetermined)
+	{
+		return;
+	}
+
+	if (UMVVMBlueprintViewEvent* Event = GetEvent())
+	{
+		GetEditorSubsystem()->SetEnabledForEvent(Event, NewState == ECheckBoxState::Checked);
+	}
+}
+
+ECheckBoxState SEventRow::IsEventCompiled() const
+{
+	if (const UMVVMBlueprintViewEvent* Event = GetEvent())
+	{
+		return Event->bCompile ? ECheckBoxState::Checked : ECheckBoxState::Unchecked;
+	}
+	return ECheckBoxState::Undetermined;
+}
+
+void SEventRow::OnIsEventCompileChanged(ECheckBoxState NewState)
+{
+	if (NewState == ECheckBoxState::Undetermined)
+	{
+		return;
+	}
+
+	if (UMVVMBlueprintViewEvent* Event = GetEvent())
+	{
+		GetEditorSubsystem()->SetCompileForEvent(Event, NewState == ECheckBoxState::Checked);
+	}
+}
+
+FMVVMBlueprintPropertyPath SEventRow::GetSelectedPropertyPath(bool bEvent) const
+{
+	if (const UMVVMBlueprintViewEvent* Event = GetEvent())
+	{
+		return bEvent ? Event->GetEventPath() : Event->GetDestinationPath();
+	}
+	return FMVVMBlueprintPropertyPath();
+}
+
+const UFunction* SEventRow::GetSelectedConversionFunction(bool bSourceToDest) const
+{
+	return nullptr;
+}
+
+void SEventRow::HandleFieldSelectionChanged(FMVVMBlueprintPropertyPath SelectedField, const UFunction* Function, bool bEvent)
+{
+	UWidgetBlueprint* WidgetBlueprint = GetBlueprint();
+	UMVVMBlueprintViewEvent* Event = GetEvent();
+	if (WidgetBlueprint && Event)
+	{
+		UMVVMEditorSubsystem* Subsystem = GetEditorSubsystem();
+		if (bEvent)
+		{
+			if (Event->GetEventPath() != SelectedField)
+			{
+				Subsystem->SetEventPath(Event, SelectedField);
+			}
+		}
+		else
+		{
+			if (Event->GetDestinationPath() != SelectedField)
+			{
+				Subsystem->SetEventDestinationPath(Event, SelectedField);
+			}
+		}
+	}
+}
+
+FFieldSelectionContext SEventRow::GetSelectedSelectionContext(bool bEvent) const
+{
+	FFieldSelectionContext Result;
+	const UWidgetBlueprint* WidgetBlueprintPtr = GetBlueprint();
+	UMVVMBlueprintViewEvent* Event = GetEvent();
+	if (WidgetBlueprintPtr == nullptr || Event == nullptr)
+	{
+		return Result;
+	}
+
+	Result.BindingMode = EMVVMBindingMode::OneTimeToDestination;
+	if (bEvent && !Event->GetEventPath().GetWidgetName().IsNone())
+	{
+		Result.FixedBindingSource = FBindingSource::CreateForWidget(WidgetBlueprintPtr, Event->GetEventPath().GetWidgetName());
+	}
+
+	Result.bAllowWidgets = true;
+	Result.bAllowViewModels = !bEvent;
+	Result.bAllowConversionFunctions = false;
+	Result.bReadable = bEvent;
+	Result.bWritable = !bEvent;
+
+	return Result;
+}
+
+FReply SEventRow::HandleFieldSelectorDrop(const FGeometry& MyGeometry, const FDragDropEvent& DragDropEvent, bool bEvent)
+{
+	TSharedPtr<FDecoratedDragDropOp> DragDropOp = DragDropEvent.GetOperationAs<FDecoratedDragDropOp>();
+	if (!DragDropOp.IsValid())
+	{
+		return FReply::Unhandled();
+	}
+
+	// Accept all drag-drop operations that are widget properties, but only accept view model fields when we are dropping into the Source box.
+	if (!DragDropOp->IsOfType<FWidgetPropertyDragDropOp>() && (!DragDropOp->IsOfType<FViewModelFieldDragDropOp>() || !bEvent))
+	{
+		return FReply::Unhandled();
+	}
+
+	UMVVMBlueprintViewEvent* Event = GetEvent();
+	UWidgetBlueprint* WidgetBlueprintPtr = GetBlueprint();
+	if (Event == nullptr || WidgetBlueprintPtr == nullptr)
+	{
+		return FReply::Unhandled();
+	}
+
+	TSharedPtr<FViewModelFieldDragDropOp> ViewModelFieldDragDropOp = DragDropEvent.GetOperationAs<FViewModelFieldDragDropOp>();
+	TSharedPtr<FWidgetPropertyDragDropOp> WidgetPropertyDragDropOp = DragDropEvent.GetOperationAs<FWidgetPropertyDragDropOp>();
+	bool bIsViewModelProperty = ViewModelFieldDragDropOp != nullptr;
+
+	UWidgetBlueprint* DragDropWidgetBP = bIsViewModelProperty ? ViewModelFieldDragDropOp->WidgetBP.Get() : WidgetPropertyDragDropOp->WidgetBP.Get();
+	if (WidgetBlueprintPtr != DragDropWidgetBP)
+	{
+		return FReply::Unhandled();
+	}
+
+	TArray<FFieldVariant> FieldPath = bIsViewModelProperty ? ViewModelFieldDragDropOp->DraggedField : WidgetPropertyDragDropOp->DraggedPropertyPath;
+	FMVVMBlueprintPropertyPath PropertyPath;
+	for (const FFieldVariant& Field : FieldPath)
+	{
+		PropertyPath.AppendPropertyPath(WidgetBlueprintPtr, FMVVMConstFieldVariant(Field));
+	}
+
+	UMVVMEditorSubsystem* Subsystem = GetEditorSubsystem();
+	if (bIsViewModelProperty)
+	{
+		if (ViewModelFieldDragDropOp->ViewModelId.IsValid())
+		{
+			PropertyPath.SetViewModelId(ViewModelFieldDragDropOp->ViewModelId);
+		}
+		else
+		{
+			return FReply::Unhandled();
+		}
+	}
+	else
+	{
+		if (UWidget* OwnerWidgetPtr = WidgetPropertyDragDropOp->OwnerWidget.Get())
+		{
+			// Search for the widget by its name in the widget tree
+			// If the widget is not found, we know it is the root preview widget so we use the blueprint name.
+			if (WidgetBlueprintPtr->WidgetTree->FindWidget(OwnerWidgetPtr->GetFName()) && WidgetBlueprintPtr->GetFName() != OwnerWidgetPtr->GetFName())
+			{
+				PropertyPath.SetWidgetName(OwnerWidgetPtr->GetFName());
+			}
+			else
+			{
+				PropertyPath.SetSelfContext();
+			}
+		}
+	}
+
+	if (bEvent)
+	{
+		Subsystem->SetEventPath(Event, PropertyPath);
+	}
+	else
+	{
+		Subsystem->SetEventDestinationPath(Event, PropertyPath);
+	}
+	return FReply::Handled();
+}
+
+void SEventRow::HandleFieldSelectorDragEnter(const FGeometry& MyGeometry, const FDragDropEvent& DragDropEvent, bool bEvent)
+{
+	TSharedPtr<FDecoratedDragDropOp> DragDropOp = DragDropEvent.GetOperationAs<FDecoratedDragDropOp>();
+	if (DragDropOp.IsValid())
+	{
+		// Accept all drag-drop operations that are widget properties, but only accept view model fields when we are dropping into the Source box.
+		if (DragDropOp->IsOfType<FWidgetPropertyDragDropOp>() || (DragDropOp->IsOfType<FViewModelFieldDragDropOp>() && !bEvent))
+		{
+			if (UWidgetBlueprint* WidgetBlueprintPtr = GetBlueprint())
+			{
+				TSharedPtr<FViewModelFieldDragDropOp> ViewModelFieldDragDropOp = DragDropEvent.GetOperationAs<FViewModelFieldDragDropOp>();
+				TSharedPtr<FWidgetPropertyDragDropOp> WidgetPropertyDragDropOp = DragDropEvent.GetOperationAs<FWidgetPropertyDragDropOp>();
+				bool IsViewModelProperty = false;
+
+				if (ViewModelFieldDragDropOp)
+				{
+					IsViewModelProperty = true;
+				}
+
+				UWidgetBlueprint* DragDropWidgetBP = IsViewModelProperty ? ViewModelFieldDragDropOp->WidgetBP.Get() : WidgetPropertyDragDropOp->WidgetBP.Get();
+
+				if (DragDropWidgetBP && DragDropWidgetBP == WidgetBlueprintPtr)
+				{
+					DragDropOp->CurrentIconBrush = FAppStyle::GetBrush(TEXT("Graph.ConnectorFeedback.OK"));
+					return;
+				}
+			}
+		}
+		else
+		{
+			DragDropOp->CurrentIconBrush = FAppStyle::GetBrush(TEXT("Graph.ConnectorFeedback.Error"));
+		}
+	}
+}
+
+} // namespace
+
+#undef LOCTEXT_NAMESPACE
diff --git a/Engine/Plugins/Runtime/ModelViewViewModel/Source/ModelViewViewModelEditor/Private/Widgets/BindingEntry/SMVVMEventRow.h b/Engine/Plugins/Runtime/ModelViewViewModel/Source/ModelViewViewModelEditor/Private/Widgets/BindingEntry/SMVVMEventRow.h
new file mode 100644
--- /dev/null
+++ b/Engine/Plugins/Runtime/ModelViewViewModel/Source/ModelViewViewModelEditor/Private/Widgets/BindingEntry/SMVVMEventRow.h
@@ -0,0 +1,70 @@
+// Copyright Epic Games, Inc. All Rights Reserved.
+
+#pragma once
+
+#include "Input/Reply.h"
+#include "Types/MVVMBindingSource.h"
+#include "Types/MVVMExecutionMode.h"
+#include "Widgets/BindingEntry/SMVVMBaseRow.h"
+#include "Widgets/SMVVMFieldSelectorMenu.h"
+
+
+namespace UE::MVVM::BindingEntry
+{
+
+/**
+ *
+ */
+class SEventRow : public SBaseRow
+{
+public:
+	SLATE_BEGIN_ARGS(SEventRow) {}
+	SLATE_END_ARGS()
+
+	void Construct(const FArguments& Args, const TSharedRef<STableViewBase>& OwnerTableView, const TSharedPtr<FWidgetBlueprintEditor>& InBlueprintEditor, UWidgetBlueprint* InBlueprint, const TSharedPtr<FBindingEntry>& InEntry);
+
+protected:
+	virtual TSharedRef<SWidget> BuildRowWidget() override;
+	virtual const ANSICHAR* GetTableRowStyle() const override
+	{
+		return "BindingView.BindingRow";
+	}
+
+private:
+	UMVVMBlueprintViewEvent* GetEvent() const;
+
+	FSlateColor GetErrorBorderColor() const;
+
+	EVisibility GetErrorButtonVisibility() const;
+
+	FText GetErrorButtonToolTip() const;
+
+	FReply OnErrorButtonClicked();
+
+	TSharedRef<ITableRow> OnGenerateErrorRow(TSharedPtr<FText> Text, const TSharedRef<STableViewBase>& TableView) const;
+
+	ECheckBoxState IsEventEnabled() const;
+
+	void OnIsEventEnableChanged(ECheckBoxState NewState);
+
+	ECheckBoxState IsEventCompiled() const;
+
+	void OnIsEventCompileChanged(ECheckBoxState NewState);
+
+	FMVVMBlueprintPropertyPath GetSelectedPropertyPath(bool bSource) const;
+
+	const UFunction* GetSelectedConversionFunction(bool bSourceToDest) const;
+
+	void HandleFieldSelectionChanged(FMVVMBlueprintPropertyPath SelectedField, const UFunction* Function, bool bSource);
+
+	FFieldSelectionContext GetSelectedSelectionContext(bool bSource) const;
+
+	FReply HandleFieldSelectorDrop(const FGeometry& MyGeometry, const FDragDropEvent& DragDropEvent, bool bSource);
+
+	void HandleFieldSelectorDragEnter(const FGeometry& MyGeometry, const FDragDropEvent& DragDropEvent, bool bSource);
+
+private:
+	TArray<TSharedPtr<FText>> ErrorItems;
+};
+
+} // namespace UE::MVVM
diff --git a/Engine/Plugins/Runtime/ModelViewViewModel/Source/ModelViewViewModelEditor/Private/Widgets/BindingEntry/SMVVMFunctionParameterRow.cpp b/Engine/Plugins/Runtime/ModelViewViewModel/Source/ModelViewViewModelEditor/Private/Widgets/BindingEntry/SMVVMFunctionParameterRow.cpp
new file mode 100644
--- /dev/null
+++ b/Engine/Plugins/Runtime/ModelViewViewModel/Source/ModelViewViewModelEditor/Private/Widgets/BindingEntry/SMVVMFunctionParameterRow.cpp
@@ -0,0 +1,137 @@
+// Copyright Epic Games, Inc. All Rights Reserved.
+
+#include "Widgets/BindingEntry/SMVVMFunctionParameterRow.h"
+
+#include "Types/MVVMBindingMode.h"
+#include "Types/MVVMBindingEntry.h"
+#include "MVVMBlueprintViewBinding.h"
+#include "MVVMBlueprintViewConversionFunction.h"
+#include "MVVMBlueprintViewEvent.h"
+
+#include "Kismet2/BlueprintEditorUtils.h"
+
+#include "Widgets/Images/SImage.h"
+#include "Widgets/Layout/SBorder.h"
+#include "Widgets/Text/STextBlock.h"
+#include "Widgets/SBoxPanel.h"
+#include "Widgets/SMVVMEventParameter.h"
+#include "Widgets/SMVVMFunctionParameter.h"
+
+#define LOCTEXT_NAMESPACE "BindingListView_FunctionParameterRow"
+
+namespace UE::MVVM::BindingEntry
+{
+
+TSharedRef<SWidget> SFunctionParameterRow::BuildRowWidget()
+{
+	UMVVMEditorSubsystem* EditorSubsystem = GetEditorSubsystem();
+
+	const FMVVMBlueprintPin* Pin = nullptr;
+	const UEdGraphPin* GraphPin = nullptr;
+	TSharedPtr<SWidget> ContentWidget;
+
+	if (GetEntry()->GetRowType() == FBindingEntry::ERowType::BindingParameter)
+	{
+		bool bSimpleConversionFunction = false;
+		UMVVMBlueprintView* View = EditorSubsystem->GetView(GetBlueprint());
+		FMVVMBlueprintViewBinding* Binding = GetEntry()->GetBinding(View);
+		check(Binding);
+
+		const bool bSourceToDestination = UE::MVVM::IsForwardBinding(Binding->BindingType);
+		UMVVMBlueprintViewConversionFunction* ConversionFunction = Binding->Conversion.GetConversionFunction(bSourceToDestination);
+		check(ConversionFunction);
+
+		Pin = ConversionFunction->GetPins().FindByPredicate([Name = GetEntry()->GetBindingParameterName()](const FMVVMBlueprintPin& Other) { return Other.GetName() == Name; });
+		GraphPin = ConversionFunction->GetOrCreateGraphPin(GetBlueprint(), GetEntry()->GetBindingParameterName());
+
+		ContentWidget = SNew(SFunctionParameter, GetBlueprint())
+			.BindingId(Binding->BindingId)
+			.ParameterName(GetEntry()->GetBindingParameterName())
+			.SourceToDestination(bSourceToDestination)
+			.AllowDefault(!bSimpleConversionFunction);
+	}
+	else if (GetEntry()->GetRowType() == FBindingEntry::ERowType::EventParameter)
+	{
+		UMVVMBlueprintViewEvent* ViewEvent = GetEntry()->GetEvent();
+		check(ViewEvent);
+
+		Pin = ViewEvent->GetPins().FindByPredicate([Name = GetEntry()->GetEventParameterName()](const FMVVMBlueprintPin& Other) { return Other.GetName() == Name; });
+		GraphPin = ViewEvent->GetOrCreateGraphPin(GetEntry()->GetEventParameterName());
+
+		ContentWidget = SNew(SEventParameter, GetBlueprint())
+			.Event(GetEntry()->GetEvent())
+			.ParameterName(GetEntry()->GetEventParameterName())
+			.AllowDefault(true);
+	}
+
+	FSlateColor PrimaryColor, SecondaryColor;
+	const FSlateBrush* PrimaryBrush = nullptr;
+	const FSlateBrush* SecondaryBrush = nullptr;
+	FText DisplayName, ToolTip;
+	bool bTextColorIsRed = false;
+	if (Pin == nullptr || GraphPin == nullptr || Pin->GetStatus() == EMVVMBlueprintPinStatus::Orphaned)
+	{
+		PrimaryBrush = FBlueprintEditorUtils::GetIconFromPin(FEdGraphPinType());
+		PrimaryColor = FLinearColor::Red;
+		DisplayName = GraphPin ? GraphPin->GetDisplayName() : (Pin ? FText::FromName(Pin->GetName()) : FText::GetEmpty());
+		bTextColorIsRed = true;
+	}
+	else
+	{
+		PrimaryBrush = FBlueprintEditor::GetVarIconAndColorFromPinType(GraphPin->PinType, PrimaryColor, SecondaryBrush, SecondaryColor);
+		DisplayName = GraphPin->GetDisplayName();
+		ToolTip = FText::FromString(GraphPin->PinToolTip);
+	}
+
+	TSharedPtr<STextBlock> DisplayNameTextBlock;
+	TSharedRef<SWidget> Result = SNew(SBox)
+	.HeightOverride(30)
+	.ToolTipText(ToolTip)
+	[
+		SNew(SHorizontalBox)
+
+		+ SHorizontalBox::Slot()
+		.HAlign(HAlign_Left)
+		.VAlign(VAlign_Center)
+		.Padding(0, 0, 8, 0)
+		.AutoWidth()
+		[
+			SNew(SHorizontalBox)
+			+ SHorizontalBox::Slot()
+			.VAlign(VAlign_Center)
+			.Padding(0, 0, 8, 0)
+			.AutoWidth()
+			[
+				SNew(SImage)
+				.Image(PrimaryBrush)
+				.ColorAndOpacity(PrimaryColor)
+				.DesiredSizeOverride(FVector2D(16, 16))
+			]
+
+			+ SHorizontalBox::Slot()
+			[
+				SAssignNew(DisplayNameTextBlock, STextBlock)
+				.Text(DisplayName)
+			]
+		]
+
+		+ SHorizontalBox::Slot()
+		.HAlign(HAlign_Left)
+		.VAlign(VAlign_Center)
+		.AutoWidth()
+		[
+			ContentWidget.ToSharedRef()
+		]
+	];
+
+	if (bTextColorIsRed)
+	{
+		DisplayNameTextBlock->SetColorAndOpacity(PrimaryColor);
+	}
+
+	return Result;
+}
+
+} // namespace
+
+#undef LOCTEXT_NAMESPACE
diff --git a/Engine/Plugins/Runtime/ModelViewViewModel/Source/ModelViewViewModelEditor/Private/Widgets/BindingEntry/SMVVMFunctionParameterRow.h b/Engine/Plugins/Runtime/ModelViewViewModel/Source/ModelViewViewModelEditor/Private/Widgets/BindingEntry/SMVVMFunctionParameterRow.h
new file mode 100644
--- /dev/null
+++ b/Engine/Plugins/Runtime/ModelViewViewModel/Source/ModelViewViewModelEditor/Private/Widgets/BindingEntry/SMVVMFunctionParameterRow.h
@@ -0,0 +1,26 @@
+// Copyright Epic Games, Inc. All Rights Reserved.
+
+#pragma once
+
+#include "Input/Reply.h"
+#include "Widgets/BindingEntry/SMVVMBaseRow.h"
+#include "Widgets/SMVVMFieldSelectorMenu.h"
+
+
+namespace UE::MVVM::BindingEntry
+{
+
+/**
+ *
+ */
+class SFunctionParameterRow : public SBaseRow
+{
+protected:
+	virtual TSharedRef<SWidget> BuildRowWidget() override;
+	virtual const ANSICHAR* GetTableRowStyle() const override
+	{
+		return "BindingView.ParameterRow";
+	}
+};
+
+} // namespace UE::MVVM
diff --git a/Engine/Plugins/Runtime/ModelViewViewModel/Source/ModelViewViewModelEditor/Private/Widgets/BindingEntry/SMVVMGroupRow.cpp b/Engine/Plugins/Runtime/ModelViewViewModel/Source/ModelViewViewModelEditor/Private/Widgets/BindingEntry/SMVVMGroupRow.cpp
new file mode 100644
--- /dev/null
+++ b/Engine/Plugins/Runtime/ModelViewViewModel/Source/ModelViewViewModelEditor/Private/Widgets/BindingEntry/SMVVMGroupRow.cpp
@@ -0,0 +1,146 @@
+// Copyright Epic Games, Inc. All Rights Reserved.
+
+#include "Widgets/BindingEntry/SMVVMGroupRow.h"
+
+#include "MVVMBlueprintViewBinding.h"
+#include "MVVMBlueprintViewEvent.h"
+#include "MVVMPropertyPath.h"
+
+#include "Styling/AppStyle.h"
+
+#include "Widgets/Layout/SBorder.h"
+#include "Widgets/SBoxPanel.h"
+#include "Widgets/SMVVMSourceSelector.h"
+#include "SSimpleButton.h"
+
+namespace UE::MVVM::BindingEntry
+{
+
+void SGroupRow::Construct(const FArguments& Args, const TSharedRef<STableViewBase>& OwnerTableView, const TSharedPtr<FWidgetBlueprintEditor>& InBlueprintEditor, UWidgetBlueprint* InBlueprint, const TSharedPtr<FBindingEntry>& InEntry)
+{
+	SBaseRow::Construct(SBaseRow::FArguments(), OwnerTableView, InBlueprintEditor, InBlueprint, InEntry);
+
+	TSharedPtr<SWidget> ChildContent = ChildSlot.DetachWidget();
+	ChildSlot
+	[
+		// Add a single pixel top and bottom border for this widget.
+		SNew(SBorder)
+		.BorderImage(FAppStyle::Get().GetBrush("Brushes.Recessed"))
+		.Padding(0.0f, 1.0f)
+		[
+			// Restore the border that we're meant to have that reacts to selection/hover/etc.
+			SNew(SBorder)
+			.BorderImage(this, &SGroupRow::GetBorderImage)
+			.Padding(0.0f)
+			[
+				ChildContent.ToSharedRef()
+			]
+		]
+	];
+}
+
+TSharedRef<SWidget> SGroupRow::BuildRowWidget()
+{
+	return SNew(SHorizontalBox)
+		+ SHorizontalBox::Slot()
+		.Padding(2.0f, 1.0f)
+		.HAlign(HAlign_Left)
+		.VAlign(VAlign_Fill)
+		.AutoWidth()
+		[
+			SNew(SBox)
+			.MinDesiredWidth(150.0f)
+			[
+				SNew(SBindingContextSelector, GetBlueprint())
+				.ShowClear(false)
+				.AutoRefresh(true)
+				.ViewModels(false)
+				.SelectedBindingSource(this, &SGroupRow::GetSelectedWidget)
+				.OnSelectionChanged(this, &SGroupRow::SetSelectedWidget)
+			]
+		]
+		+ SHorizontalBox::Slot()
+		.HAlign(HAlign_Left)
+		.VAlign(VAlign_Center)
+		.AutoWidth()
+		[
+			SNew(SSimpleButton)
+			.Icon(FAppStyle::Get().GetBrush("Icons.Plus"))
+			.IsEnabled_Lambda([this]() { return !GetEntry()->GetGroupName().IsNone(); })
+			.OnClicked(this, &SGroupRow::AddBinding)
+		];
+}
+
+FBindingSource SGroupRow::GetSelectedWidget() const
+{
+	return GetEntry()->IsGroupWidget()
+		? FBindingSource::CreateForWidget(GetBlueprint(), GetEntry()->GetGroupName())
+		: FBindingSource::CreateForViewModel(GetBlueprint(), GetEntry()->GetGroupAsViewModel());
+}
+
+void SGroupRow::SetSelectedWidget(FBindingSource Source)
+{
+	UWidgetBlueprint* WidgetBlueprint = GetBlueprint();
+	UMVVMBlueprintView* View = GetBlueprintView();
+	if (WidgetBlueprint && View)
+	{
+		UMVVMEditorSubsystem* EditorSubsystem = GetEditorSubsystem();
+
+		for (const TSharedPtr<FBindingEntry>& ChildEntry : GetEntry()->GetAllChildren())
+		{
+			if (ChildEntry->GetRowType() == FBindingEntry::ERowType::Binding)
+			{
+				if (FMVVMBlueprintViewBinding* Binding = ChildEntry->GetBinding(View))
+				{
+					FMVVMBlueprintPropertyPath CurrentPath = Binding->DestinationPath;
+					Source.SetSourceTo(CurrentPath);
+
+					EditorSubsystem->SetDestinationPathForBinding(WidgetBlueprint, *Binding, CurrentPath);
+				}
+			}
+
+			if (ChildEntry->GetRowType() == FBindingEntry::ERowType::Event)
+			{
+				if (UMVVMBlueprintViewEvent* Event = ChildEntry->GetEvent())
+				{
+					FMVVMBlueprintPropertyPath CurrentPath = Event->GetEventPath();
+					Source.SetSourceTo(CurrentPath);
+
+					EditorSubsystem->SetEventPath(Event, CurrentPath);
+				}
+			}
+		}
+	}
+}
+
+FReply SGroupRow::AddBinding() const
+{
+	if (UWidgetBlueprint* WidgetBlueprint = GetBlueprint())
+	{
+		UMVVMEditorSubsystem* EditorSubsystem = GetEditorSubsystem();
+		FMVVMBlueprintViewBinding& Binding = EditorSubsystem->AddBinding(WidgetBlueprint);
+		FMVVMBlueprintPropertyPath Path;
+		TSharedPtr<FBindingEntry> Entry = GetEntry();
+		if (Entry->IsGroupWidget())
+		{
+			if (Entry->GetGroupName() == WidgetBlueprint->GetFName())
+			{
+				Path.SetSelfContext();
+			}
+			else
+			{
+				Path.SetWidgetName(Entry->GetGroupName());
+			}
+		}
+		else
+		{
+			Path.SetViewModelId(Entry->GetGroupAsViewModel());
+		}
+
+		EditorSubsystem->SetDestinationPathForBinding(WidgetBlueprint, Binding, Path);
+	}
+
+	return FReply::Handled();
+}
+
+} // namespace
diff --git a/Engine/Plugins/Runtime/ModelViewViewModel/Source/ModelViewViewModelEditor/Private/Widgets/BindingEntry/SMVVMGroupRow.h b/Engine/Plugins/Runtime/ModelViewViewModel/Source/ModelViewViewModelEditor/Private/Widgets/BindingEntry/SMVVMGroupRow.h
new file mode 100644
--- /dev/null
+++ b/Engine/Plugins/Runtime/ModelViewViewModel/Source/ModelViewViewModelEditor/Private/Widgets/BindingEntry/SMVVMGroupRow.h
@@ -0,0 +1,39 @@
+// Copyright Epic Games, Inc. All Rights Reserved.
+
+#pragma once
+
+#include "Input/Reply.h"
+#include "Types/MVVMBindingSource.h"
+#include "Widgets/BindingEntry/SMVVMBaseRow.h"
+
+namespace UE::MVVM::BindingEntry
+{
+
+/**
+ *
+ */
+class SGroupRow : public SBaseRow
+{
+public:
+	SLATE_BEGIN_ARGS(SGroupRow) {}
+	SLATE_END_ARGS()
+
+	void Construct(const FArguments& Args, const TSharedRef<STableViewBase>& OwnerTableView, const TSharedPtr<FWidgetBlueprintEditor>& InBlueprintEditor, UWidgetBlueprint* InBlueprint, const TSharedPtr<FBindingEntry>& InEntry);
+
+protected:
+	virtual TSharedRef<SWidget> BuildRowWidget() override;
+	virtual const ANSICHAR* GetTableRowStyle() const override
+	{
+		return "BindingView.WidgetRow";
+	}
+
+
+private:
+	FBindingSource GetSelectedWidget() const;
+
+	void SetSelectedWidget(FBindingSource Source);
+
+	FReply AddBinding() const;
+};
+
+} // namespace UE::MVVM
diff --git a/Engine/Plugins/Runtime/ModelViewViewModel/Source/ModelViewViewModelEditor/Private/Widgets/SMVVMViewBindingListView.cpp b/Engine/Plugins/Runtime/ModelViewViewModel/Source/ModelViewViewModelEditor/Private/Widgets/SMVVMViewBindingListView.cpp
--- a/Engine/Plugins/Runtime/ModelViewViewModel/Source/ModelViewViewModelEditor/Private/Widgets/SMVVMViewBindingListView.cpp
+++ b/Engine/Plugins/Runtime/ModelViewViewModel/Source/ModelViewViewModelEditor/Private/Widgets/SMVVMViewBindingListView.cpp
@@ -2,16 +2,13 @@
 
 #include "Widgets/SMVVMViewBindingListView.h"
 
-#include "Bindings/MVVMBindingHelper.h"
 #include "Blueprint/WidgetTree.h"
 #include "BlueprintEditor.h"
-#include "Details/WidgetPropertyDragDropOp.h"
-#include "Dialog/SCustomDialog.h"
-#include "DragAndDrop/DecoratedDragDropOp.h"
 #include "Framework/MultiBox/MultiBoxBuilder.h"
 #include "Framework/Views/TableViewMetadata.h"
 #include "Kismet2/BlueprintEditorUtils.h"
 #include "Misc/MessageDialog.h"
+#include "ScopedTransaction.h"
 
 #include "MVVMBlueprintView.h"
 #include "MVVMBlueprintViewBinding.h"
@@ -20,26 +17,14 @@
 #include "MVVMEditorSubsystem.h"
 #include "MVVMDeveloperProjectSettings.h"
 #include "MVVMWidgetBlueprintExtension_View.h"
+#include "Types/MVVMBindingEntry.h"
 
-#include "ScopedTransaction.h"
-#include "SEnumCombo.h"
-#include "SSimpleButton.h" 
-#include "Styling/MVVMEditorStyle.h"
-#include "Styling/StyleColors.h"
-#include "ViewModelFieldDragDropOp.h"
-
-#include "Widgets/SMVVMEventParameter.h"
-#include "Widgets/SMVVMFieldEntry.h"
-#include "Widgets/SMVVMFieldSelector.h"
-#include "Widgets/SMVVMFunctionParameter.h"
-#include "Widgets/SMVVMSourceSelector.h"
+#include "Widgets/BindingEntry/SMVVMBindingRow.h"
+#include "Widgets/BindingEntry/SMVVMEventRow.h"
+#include "Widgets/BindingEntry/SMVVMFunctionParameterRow.h"
+#include "Widgets/BindingEntry/SMVVMGroupRow.h"
 #include "Widgets/SMVVMViewBindingPanel.h"
-#include "Widgets/Layout/SSpacer.h"
-#include "Widgets/Input/SCheckBox.h"
-#include "Widgets/Input/SComboBox.h"
-#include "Widgets/Input/SEditableTextBox.h"
-#include "Widgets/SMVVMViewModelPanel.h" // IWYU pragma: keep
-#include "Widgets/SMVVMViewModelBindingListWidget.h"
+#include "Widgets/SMVVMViewModelPanel.h"
 #include "Widgets/Views/STreeView.h"
 
 #define LOCTEXT_NAMESPACE "BindingListView"
@@ -47,235 +32,8 @@
 namespace UE::MVVM
 {
 
-/**
- * A wrapper around either a widget row or a binding row
- */
-struct FBindingEntry
-{
-	enum class ERowType
-	{
-		None,
-		Group,
-		Binding,
-		BindingParameter,
-		Event,
-		EventParameter,
-	};
-
-	FMVVMBlueprintViewBinding* GetBinding(UMVVMBlueprintView* View) const
-	{
-		return View->GetBinding(BindingId);
-	}
-
-	const FMVVMBlueprintViewBinding* GetBinding(const UMVVMBlueprintView* View) const
-	{
-		return View->GetBinding(BindingId);
-	}
-
-	ERowType GetRowType() const
-	{
-		return RowType;
-	}
-
-	FName GetGroupName() const
-	{
-		return Name;
-	}
-
-	FGuid GetGroupAsViewModel() const
-	{
-		return BindingId;
-	}
-
-	bool IsGroupWidget() const
-	{
-		return bGroupIsWidget;
-	}
-
-	void SetGroup(FName WidgetName)
-	{
-		check(RowType == ERowType::None);
-		RowType = ERowType::Group;
-		Name = WidgetName;
-		bGroupIsWidget = true;
-	}
-
-	void SetGroup(FName ViewModelName, FGuid ViewModelId)
-	{
-		if (!ViewModelId.IsValid())
-		{
-			SetGroup(ViewModelName);
-		}
-		else
-		{
-			check(RowType == ERowType::None);
-			RowType = ERowType::Group;
-			Name = ViewModelName;
-			BindingId = ViewModelId;
-			bGroupIsWidget = false;
-		}
-	}
-
-	FGuid GetBindingId() const
-	{
-		return BindingId;
-	}
-
-	void SetBindingId(FGuid Id)
-	{
-		check(RowType == ERowType::None);
-		RowType = ERowType::Binding;
-		BindingId = Id;
-	}
-
-	FName GetBindingParameterName() const
-	{
-		return Name;
-	}
-
-	void SetBindingParameter(FGuid Id, FName ParameterName)
-	{
-		check(RowType == ERowType::None);
-		RowType = ERowType::BindingParameter;
-		BindingId = Id;
-		Name = ParameterName;
-	}
-
-	UMVVMBlueprintViewEvent* GetEvent() const
-	{
-		return Event.Get();
-	}
-
-	void SetEvent(UMVVMBlueprintViewEvent* InEvent)
-	{
-		check(RowType == ERowType::None);
-		RowType = ERowType::Event;
-		Event = InEvent;
-	}
-
-	FName GetEventParameterName() const
-	{
-		return Name;
-	}
-
-	void SetEventParameter(UMVVMBlueprintViewEvent* InEvent, FName ParameterName)
-	{
-		check(RowType == ERowType::None);
-		RowType = ERowType::EventParameter;
-		Event = InEvent;
-		Name = ParameterName;
-	}
-
-	TConstArrayView<TSharedPtr<FBindingEntry>> GetAllChildren() const
-	{
-		return AllChildren;
-	}
-
-	TConstArrayView<TSharedPtr<FBindingEntry>> GetFilteredChildren() const
-	{
-		return bUseFilteredChildren ? FilteredChildren : AllChildren;
-	}
-
-	void AddChild(TSharedPtr<FBindingEntry> Child)
-	{
-		AllChildren.Add(Child);
-	}
-
-	void AddFilteredChild(TSharedPtr<FBindingEntry> Child)
-	{
-		FilteredChildren.Add(Child);
-		bUseFilteredChildren = true;
-	}
-
-	void ResetChildren()
-	{
-		AllChildren.Reset();
-		FilteredChildren.Reset();
-		bUseFilteredChildren = false;
-	}
-
-	void SetUseFilteredChildList()
-	{
-		bUseFilteredChildren = true;
-	}
-
-	bool operator==(const FBindingEntry& Other) const
-	{
-		return RowType == Other.RowType
-			&& Name == Other.Name
-			&& BindingId == Other.BindingId
-			&& Event == Other.Event;
-	}
-
-	FString GetSearchNameString(UMVVMBlueprintView* View, UWidgetBlueprint* WidgetBP)
-	{
-		FString RowToString;
-		FString FunctionKeywords;
-
-		switch (RowType)
-		{
-		case UE::MVVM::FBindingEntry::ERowType::Group:
-		case UE::MVVM::FBindingEntry::ERowType::BindingParameter:
-		case UE::MVVM::FBindingEntry::ERowType::EventParameter:
-			RowToString = Name.ToString();
-			break;
-		case UE::MVVM::FBindingEntry::ERowType::Binding:
-			{
-				FMVVMBlueprintViewBinding* BindingInRow = GetBinding(View);
-				check(BindingInRow);
-				RowToString.Append(BindingInRow->GetSearchableString(WidgetBP));
-			}
-			break;
-		case UE::MVVM::FBindingEntry::ERowType::Event:
-			if (UMVVMBlueprintViewEvent* EventInRow = GetEvent())
-			{
-				RowToString.Append(EventInRow->GetSearchableString());
-			}
-			break;
-		default:
-			break;
-		}
-
-		RowToString.ReplaceInline(TEXT(" "), TEXT(""));
-
-		return RowToString;
-	}
-private:
-	ERowType RowType = ERowType::None;
-	FName Name;
-	FGuid BindingId;
-	TWeakObjectPtr<UMVVMBlueprintViewEvent> Event;
-	TArray<TSharedPtr<FBindingEntry>> AllChildren;
-	TArray<TSharedPtr<FBindingEntry>> FilteredChildren;
-	bool bGroupIsWidget = false;
-	bool bUseFilteredChildren = false;
-};
-
 namespace Private
 {
-	TArray<FName>* GetBindingModeNames()
-	{
-		static TArray<FName> BindingModeNames;
-
-		if (BindingModeNames.IsEmpty())
-		{
-			UEnum* ModeEnum = StaticEnum<EMVVMBindingMode>();
-
-			BindingModeNames.Reserve(ModeEnum->NumEnums());
-
-			for (int32 BindingIndex = 0; BindingIndex < ModeEnum->NumEnums() - 1; ++BindingIndex)
-			{
-				const bool bIsHidden = ModeEnum->HasMetaData(TEXT("Hidden"), BindingIndex);
-				if (!bIsHidden)
-				{
-					BindingModeNames.Add(ModeEnum->GetNameByIndex(BindingIndex));
-				}
-			}
-		}
-
-		return &BindingModeNames;
-	}
-
 	void ExpandAll(const TSharedPtr<STreeView<TSharedPtr<FBindingEntry>>>& TreeView, const TSharedPtr<FBindingEntry>& Entry)
 	{
 		TreeView->SetItemExpansion(Entry, true);
@@ -400,1533 +158,12 @@ namespace Private
 	}
 } // namespace
 
-/**
- * 
- */
-class SGroupRow : public STableRow<TSharedPtr<FBindingEntry>>
-{
-public:
-	SLATE_BEGIN_ARGS(SGroupRow) {}
-	SLATE_END_ARGS()
-
-	void Construct(const FArguments& Args, const TSharedRef<STableViewBase>& OwnerTableView, const TSharedPtr<FBindingEntry>& InEntry, UWidgetBlueprint* InWidgetBlueprint)
-	{
-		Entry = InEntry;
-		WidgetBlueprintWeak = InWidgetBlueprint;
-
-		STableRow<TSharedPtr<FBindingEntry>>::Construct(
-			STableRow<TSharedPtr<FBindingEntry>>::FArguments()
-			.Padding(2.0f)
-			.Style(FMVVMEditorStyle::Get(), "BindingView.WidgetRow")
-			[
-				SNew(SHorizontalBox)
-				+ SHorizontalBox::Slot()
-				.Padding(2, 1)
-				.HAlign(HAlign_Left)
-				.VAlign(VAlign_Fill)
-				.AutoWidth()
-				[
-					SNew(SBox)
-					.MinDesiredWidth(150)
-					[
-						SNew(SBindingContextSelector, InWidgetBlueprint)
-						.ShowClear(false)
-						.AutoRefresh(true)
-						.ViewModels(false)
-						.SelectedBindingSource(this, &SGroupRow::GetSelectedWidget)
-						.OnSelectionChanged(this, &SGroupRow::SetSelectedWidget)
-					]
-				]
-				+ SHorizontalBox::Slot()
-				.HAlign(HAlign_Left)
-				.VAlign(VAlign_Center)
-				.AutoWidth()
-				[
-					SNew(SSimpleButton)
-					.Icon(FAppStyle::Get().GetBrush("Icons.Plus"))
-					.IsEnabled_Lambda([this]() { return !Entry->GetGroupName().IsNone(); })
-					.OnClicked(this, &SGroupRow::AddBinding)
-				] 
-			],
-			OwnerTableView
-		);
-
-		TSharedPtr<SWidget> ChildContent = ChildSlot.DetachWidget();
-		ChildSlot
-		[
-			/* Add a single pixel top and bottom border for this widget. */
-			SNew(SBorder)
-			.BorderImage(FAppStyle::Get().GetBrush("Brushes.Recessed"))
-			.Padding(0, 1)
-			[
-				/* Restore the border that we're meant to have that reacts to selection/hover/etc. */
-				SNew(SBorder)
-				.BorderImage(this, &SGroupRow::GetBorderImage)
-				.Padding(0)
-				[
-					ChildContent.ToSharedRef()
-				]
-			]
-		];
-	}
-
-private:
-	FBindingSource GetSelectedWidget() const
-	{
-		return Entry->IsGroupWidget()
-			? FBindingSource::CreateForWidget(WidgetBlueprintWeak.Get(), Entry->GetGroupName())
-			: FBindingSource::CreateForViewModel(WidgetBlueprintWeak.Get(), Entry->GetGroupAsViewModel());
-	}
-
-	void SetSelectedWidget(FBindingSource Source)
-	{
-		if (UWidgetBlueprint* WidgetBlueprint = WidgetBlueprintWeak.Get())
-		{
-			UMVVMEditorSubsystem* EditorSubsystem = GEditor->GetEditorSubsystem<UMVVMEditorSubsystem>();
-			UMVVMBlueprintView* View = EditorSubsystem->GetView(WidgetBlueprint);
-
-			for (const TSharedPtr<FBindingEntry>& ChildEntry : Entry->GetAllChildren())
-			{
-				if (ChildEntry->GetRowType() == FBindingEntry::ERowType::Binding)
-				{
-					if (FMVVMBlueprintViewBinding* Binding = ChildEntry->GetBinding(View))
-					{
-						FMVVMBlueprintPropertyPath CurrentPath = Binding->DestinationPath;
-						Source.SetSourceTo(CurrentPath);
-
-						EditorSubsystem->SetDestinationPathForBinding(WidgetBlueprint, *Binding, CurrentPath);
-					}
-				}
-
-				if (ChildEntry->GetRowType() == FBindingEntry::ERowType::Event)
-				{
-					if (UMVVMBlueprintViewEvent* Event = ChildEntry->GetEvent())
-					{
-						FMVVMBlueprintPropertyPath CurrentPath = Event->GetEventPath();
-						Source.SetSourceTo(CurrentPath);
-
-						EditorSubsystem->SetEventPath(Event, CurrentPath);
-					}
-				}
-			}
-		}
-	}
-
-	FReply AddBinding() const
-	{
-		if (UWidgetBlueprint* WidgetBlueprint = WidgetBlueprintWeak.Get())
-		{
-			UMVVMEditorSubsystem* EditorSubsystem = GEditor->GetEditorSubsystem<UMVVMEditorSubsystem>();
-			FMVVMBlueprintViewBinding& Binding = EditorSubsystem->AddBinding(WidgetBlueprint);
-			FMVVMBlueprintPropertyPath Path;
-			if (Entry->IsGroupWidget())
-			{
-				if (Entry->GetGroupName() == WidgetBlueprint->GetFName())
-				{
-					Path.SetSelfContext();
-				}
-				else
-				{
-					Path.SetWidgetName(Entry->GetGroupName());
-				}
-			}
-			else
-			{
-				Path.SetViewModelId(Entry->GetGroupAsViewModel());
-			}
-			
-			EditorSubsystem->SetDestinationPathForBinding(WidgetBlueprint, Binding, Path);
-		}
-
-		return FReply::Handled();
-	}
-
-private:
-	TSharedPtr<FBindingEntry> Entry;
-	TWeakObjectPtr<UWidgetBlueprint> WidgetBlueprintWeak;
-};
-
-/**
- * 
- */
-class SBindingRow : public STableRow<TSharedPtr<FBindingEntry>>
-{
-public:
-	SLATE_BEGIN_ARGS(SBindingRow) {}
-	SLATE_END_ARGS()
-
-	void Construct(const FArguments& Args, const TSharedRef<STableViewBase>& OwnerTableView, const TSharedPtr<FBindingEntry>& InEntry, UWidgetBlueprint* InWidgetBlueprint)
-	{
-		Entry = InEntry;
-		WidgetBlueprintWeak = InWidgetBlueprint;
-
-		static IConsoleVariable* CVarDefaultExecutionMode = IConsoleManager::Get().FindConsoleVariable(TEXT("MVVM.DefaultExecutionMode"));
-		ensure(CVarDefaultExecutionMode);
-		DefaultExecutionMode = CVarDefaultExecutionMode ? (EMVVMExecutionMode)CVarDefaultExecutionMode->GetInt() : EMVVMExecutionMode::DelayedWhenSharedElseImmediate;
-
-		
-		FMVVMBlueprintViewBinding* ViewBinding = GetThisViewBinding();
-
-		STableRow<TSharedPtr<FBindingEntry>>::Construct(
-			STableRow<TSharedPtr<FBindingEntry>>::FArguments()
-			.ShowWires(true)
-			.Style(FMVVMEditorStyle::Get(), "BindingView.BindingRow")
-			[
-				SNew(SBorder)
-				.BorderImage(FAppStyle::Get().GetBrush("PlainBorder"))
-			    .Padding(0)
-			    .BorderBackgroundColor(this, &SBindingRow::GetErrorBorderColor)
-			     [
-					SNew(SBox)
-					.HeightOverride(30)
-					[
-						SNew(SHorizontalBox)
-
-						+ SHorizontalBox::Slot()
-						.Padding(2, 0)
-						.VAlign(VAlign_Center)
-						.HAlign(HAlign_Left) 
-						.AutoWidth()
-						[
-							SNew(SCheckBox)
-							.IsChecked(this, &SBindingRow::IsBindingCompiled)
-							.OnCheckStateChanged(this, &SBindingRow::OnIsBindingCompileChanged)
-						]
-
-						+ SHorizontalBox::Slot()
-						.Padding(4.f, 0.f)
-						.VAlign(VAlign_Center)
-						.HAlign(HAlign_Left)
-						.AutoWidth()
-						[
-							SNew(SBox)
-							.MinDesiredWidth(150)
-							[
-								SNew(SFieldSelector, InWidgetBlueprint)
-								.OnGetPropertyPath(this, &SBindingRow::GetSelectedPropertyPath, false)
-								.OnGetConversionFunction(this, &SBindingRow::GetSelectedConversionFunction, false)
-								.OnFieldSelectionChanged(this, &SBindingRow::HandleFieldSelectionChanged, false)
-								.OnGetSelectionContext(this, &SBindingRow::GetSelectedSelectionContext, false)
-								.OnDrop(this, &SBindingRow::HandleFieldSelectorDrop, false)
-								.OnDragEnter(this, &SBindingRow::HandleFieldSelectorDragEnter, false)
-								.ShowContext(false)
-							]
-						]
-
-						+ SHorizontalBox::Slot()
-						.Padding(4.f, 0.f)
-						.VAlign(VAlign_Center)
-						.HAlign(HAlign_Left) 
-						.AutoWidth()
-						[
-							SNew(SComboBox<FName>)
-							.OptionsSource(Private::GetBindingModeNames())
-							.InitiallySelectedItem(StaticEnum<EMVVMBindingMode>()->GetNameByValue((int64) ViewBinding->BindingType))
-							.OnSelectionChanged(this, &SBindingRow::OnBindingModeSelectionChanged)
-							.OnGenerateWidget(this, &SBindingRow::GenerateBindingModeWidget)
-							.ToolTipText(this, &SBindingRow::GetCurrentBindingModeLabel)
-							.Content()
-							[
-								SNew(SBox)
-								.HAlign(HAlign_Center)
-								.VAlign(VAlign_Center)
-								.WidthOverride(16)
-								.HeightOverride(16)
-								[
-									SNew(SImage)
-									.Image(this, &SBindingRow::GetCurrentBindingModeBrush)
-								]
-							]
-						]
-				
-						+ SHorizontalBox::Slot()
-						.Padding(4.f, 0.f)
-						.VAlign(VAlign_Center)
-						.HAlign(HAlign_Left)
-						.AutoWidth()
-						[
-							SNew(SBox)
-							.MinDesiredWidth(150)
-							[
-								SNew(SFieldSelector, InWidgetBlueprint)
-								.OnGetPropertyPath(this, &SBindingRow::GetSelectedPropertyPath, true)
-								.OnGetConversionFunction(this, &SBindingRow::GetSelectedConversionFunction, true)
-								.OnFieldSelectionChanged(this, &SBindingRow::HandleFieldSelectionChanged, true)
-								.OnGetSelectionContext(this, &SBindingRow::GetSelectedSelectionContext, true)
-								.OnDrop(this, &SBindingRow::HandleFieldSelectorDrop, true)
-								.OnDragEnter(this, &SBindingRow::HandleFieldSelectorDragEnter, true)
-							]
-						]
-
-						+ SHorizontalBox::Slot()
-						[
-							SNew(SSpacer)
-						]
-
-						+ SHorizontalBox::Slot()
-						.Padding(2, 1)
-						.VAlign(VAlign_Fill)
-						.HAlign(HAlign_Right)
-						.AutoWidth()
-						[
-							SNew(SCheckBox)
-							.IsChecked(this, &SBindingRow::IsExecutionModeOverrideChecked)
-							.OnCheckStateChanged(this, &SBindingRow::OnExecutionModeOverrideChanged)
-						]
-
-						+ SHorizontalBox::Slot()
-						.Padding(2, 1)
-						.VAlign(VAlign_Fill)
-						.HAlign(HAlign_Right)
-						.AutoWidth()
-						[
-							SNew(SComboButton)
-							.ContentPadding(FMargin(4.f, 0.f))
-							.OnGetMenuContent(this, &SBindingRow::OnGetExecutionModeMenuContent)
-							.IsEnabled(this, &SBindingRow::IsExecutionModeOverridden)
-							.ButtonContent()
-							[
-								SNew(STextBlock)
-								.Text(this, &SBindingRow::GetExecutioModeValue)
-								.ToolTipText(this, &SBindingRow::GetExecutioModeValueToolTip)
-							]
-						]
-
-						+ SHorizontalBox::Slot()
-						.Padding(2, 0)
-						.VAlign(VAlign_Center)
-						.HAlign(HAlign_Left)
-						.AutoWidth()
-						[
-							SNew(SSimpleButton)
-							.Icon(FAppStyle::Get().GetBrush("Icons.Error"))
-							.Visibility(this, &SBindingRow::GetErrorButtonVisibility)
-							.ToolTipText(this, &SBindingRow::GetErrorButtonToolTip)
-							.OnClicked(this, &SBindingRow::OnErrorButtonClicked)
-						]
-					]
-				]
-			],
-			OwnerTableView
-		);
-
-		TSharedPtr<SWidget> ChildContent = ChildSlot.DetachWidget();
-		ChildSlot
-		[
-			/* Add a single pixel top and bottom border for this widget. */
-			SNew(SBorder)
-			.BorderImage(FAppStyle::Get().GetBrush("Brushes.Recessed"))
-			.Padding(0, 2, 0, 1)
-			[
-				/* Restore the border that we're meant to have that reacts to selection/hover/etc. */
-				SNew(SBorder)
-				.BorderImage(this, &SBindingRow::GetBorderImage)
-				.Padding(0)
-				[
-					ChildContent.ToSharedRef()
-				]
-			]
-		];
-	}
-
-	FMVVMBlueprintViewBinding* GetThisViewBinding() const
-	{
-		UMVVMEditorSubsystem* EditorSubsystem = GEditor->GetEditorSubsystem<UMVVMEditorSubsystem>();
-		if (UMVVMBlueprintView* BlueprintViewPtr = EditorSubsystem->GetView(WidgetBlueprintWeak.Get()))
-		{
-			FMVVMBlueprintViewBinding* ViewBinding = Entry->GetBinding(BlueprintViewPtr);
-			return ViewBinding;
-		}
-		return nullptr;
-	}
-
-private:
-
-	FSlateColor GetErrorBorderColor() const
-	{
-		if (const FMVVMBlueprintViewBinding* ViewModelBinding = GetThisViewBinding())
-		{
-			UMVVMEditorSubsystem* EditorSubsystem = GEditor->GetEditorSubsystem<UMVVMEditorSubsystem>();
-			UMVVMBlueprintView* BlueprintViewPtr = EditorSubsystem->GetView(WidgetBlueprintWeak.Get());
-			if (BlueprintViewPtr->HasBindingMessage(ViewModelBinding->BindingId, EBindingMessageType::Error))
-			{
-				return FStyleColors::Error;
-			}
-			else if (BlueprintViewPtr->HasBindingMessage(ViewModelBinding->BindingId, EBindingMessageType::Warning))
-			{
-				return FStyleColors::Warning;
-			}
-		}
-		return FStyleColors::Transparent;
-	}
-
-	ECheckBoxState IsBindingEnabled() const
-	{
-		if (const FMVVMBlueprintViewBinding* ViewModelBinding = GetThisViewBinding())
-		{
-			return ViewModelBinding->bEnabled ? ECheckBoxState::Checked : ECheckBoxState::Unchecked;
-		}
-		return ECheckBoxState::Undetermined;
-	}
-
-	ECheckBoxState IsBindingCompiled() const
-	{
-		if (const FMVVMBlueprintViewBinding* ViewModelBinding = GetThisViewBinding())
-		{
-			return ViewModelBinding->bCompile ? ECheckBoxState::Checked : ECheckBoxState::Unchecked;
-		}
-		return ECheckBoxState::Undetermined;
-	}
-
-	EVisibility GetErrorButtonVisibility() const
-	{
-		if (const FMVVMBlueprintViewBinding* ViewModelBinding = GetThisViewBinding())
-		{
-			UMVVMEditorSubsystem* EditorSubsystem = GEditor->GetEditorSubsystem<UMVVMEditorSubsystem>();
-			UMVVMBlueprintView* BlueprintViewPtr = EditorSubsystem->GetView(WidgetBlueprintWeak.Get());
-			bool HasBindingError = BlueprintViewPtr->HasBindingMessage(GetThisViewBinding()->BindingId, EBindingMessageType::Error) || BlueprintViewPtr->HasBindingMessage(GetThisViewBinding()->BindingId, EBindingMessageType::Warning);
-			return HasBindingError ? EVisibility::Visible : EVisibility::Hidden;
-		}
-		return EVisibility::Collapsed;
-	}
-
-	FText GetErrorButtonToolTip() const
-	{
-		// Get error messages of this binding
-		if (const FMVVMBlueprintViewBinding* ViewModelBinding = GetThisViewBinding())
-		{
-			UMVVMEditorSubsystem* EditorSubsystem = GEditor->GetEditorSubsystem<UMVVMEditorSubsystem>();
-			UMVVMBlueprintView* BlueprintViewPtr = EditorSubsystem->GetView(WidgetBlueprintWeak.Get());
-			TArray<FText> BindingErrorList = BlueprintViewPtr->GetBindingMessages(ViewModelBinding->BindingId, EBindingMessageType::Error);
-			TArray<FText> BindingWarningList = BlueprintViewPtr->GetBindingMessages(ViewModelBinding->BindingId, EBindingMessageType::Warning);
-			BindingErrorList.Append(BindingWarningList);
-
-			static const FText NewLineText = FText::FromString(TEXT("\n"));
-			FText HintText = LOCTEXT("ErrorButtonText", "Errors and Warnings: (Click to show in a separate window)");
-			FText ErrorsText = FText::Join(NewLineText, BindingErrorList);
-			return FText::Join(NewLineText, HintText, ErrorsText);
-		}
-		return FText();
-	}
-
-	FReply OnErrorButtonClicked()
-	{
-		ErrorDialog.Reset();
-		ErrorItems.Reset();
-
-		const UWidgetBlueprint* WidgetBlueprint = WidgetBlueprintWeak.Get();
-		UMVVMEditorSubsystem* EditorSubsystem = GEditor->GetEditorSubsystem<UMVVMEditorSubsystem>();
-		if (const UMVVMBlueprintView* View = EditorSubsystem->GetView(WidgetBlueprint))
-		{
-			if (const FMVVMBlueprintViewBinding* ViewBinding = Entry->GetBinding(View))
-			{
-				for (const FText& ErrorText : View->GetBindingMessages(ViewBinding->BindingId, EBindingMessageType::Error))
-				{
-					ErrorItems.Add(MakeShared<FText>(ErrorText));
-				}
-
-				for (const FText& WarningText : View->GetBindingMessages(ViewBinding->BindingId, EBindingMessageType::Warning))
-				{
-					ErrorItems.Add(MakeShared<FText>(WarningText));
-				}
-
-				const FText BindingDisplayName = FText::FromString(ViewBinding->GetDisplayNameString(WidgetBlueprint));
-				ErrorDialog = SNew(SCustomDialog)
-					.Title(FText::Format(LOCTEXT("Compilation Errors and Warnings", "Compilation Errors and Warnings for {0}"), BindingDisplayName))
-					.Buttons({
-						SCustomDialog::FButton(LOCTEXT("OK", "OK"))
-					})
-					.Content()
-					[
-						SNew(SListView<TSharedPtr<FText>>)
-						.ListItemsSource(&ErrorItems)
-						.OnGenerateRow(this, &SBindingRow::OnGenerateErrorRow)
-					];
-
-				ErrorDialog->Show();
-			}
-		}
-
-		return FReply::Handled();
-	}
-
-	TSharedRef<ITableRow> OnGenerateErrorRow(TSharedPtr<FText> Text, const TSharedRef<STableViewBase>& TableView) const
-	{
-		return SNew(STableRow<TSharedPtr<FText>>, TableView)
-			.Content()
-			[
-				SNew(SEditableTextBox)
-				.BackgroundColor(FStyleColors::Background)
-				.IsReadOnly(true)
-				.Text(*Text.Get())
-			];
-	}
-
-	TArray<FBindingSource> GetAvailableViewModels() const
-	{
-		UMVVMEditorSubsystem* EditorSubsystem = GEditor->GetEditorSubsystem<UMVVMEditorSubsystem>();
-		return EditorSubsystem->GetAllViewModels(WidgetBlueprintWeak.Get());
-	}
-
-	FMVVMBlueprintPropertyPath GetSelectedPropertyPath(bool bSource) const
-	{
-		if (FMVVMBlueprintViewBinding* ViewBinding = GetThisViewBinding())
-		{
-			return bSource ? ViewBinding->SourcePath : ViewBinding->DestinationPath;
-		}
-		return FMVVMBlueprintPropertyPath();
-	}
-
-	const UFunction* GetSelectedConversionFunction(bool bSourceToDest) const
-	{
-		if (FMVVMBlueprintViewBinding* ViewBinding = GetThisViewBinding())
-		{
-			UMVVMEditorSubsystem* EditorSubsystem = GEditor->GetEditorSubsystem<UMVVMEditorSubsystem>();
-			return EditorSubsystem->GetConversionFunction(WidgetBlueprintWeak.Get(), *ViewBinding, bSourceToDest);
-		}
-
-		return nullptr;
-	}
-
-	void HandleFieldSelectionChanged(FMVVMBlueprintPropertyPath SelectedField, const UFunction* Function, bool bSource)
-	{
-		UWidgetBlueprint* WidgetBlueprintPtr = WidgetBlueprintWeak.Get();
-		if (FMVVMBlueprintViewBinding* ViewBinding = GetThisViewBinding())
-		{
-			UMVVMEditorSubsystem* Subsystem = GEditor->GetEditorSubsystem<UMVVMEditorSubsystem>();
-			if (bSource)
-			{
-				Subsystem->SetSourceToDestinationConversionFunction(WidgetBlueprintPtr, *ViewBinding, Function);
-				if (ViewBinding->SourcePath != SelectedField)
-				{
-					Subsystem->SetSourcePathForBinding(WidgetBlueprintPtr, *ViewBinding, SelectedField);
-				}
-			}
-			else
-			{
-				Subsystem->SetDestinationToSourceConversionFunction(WidgetBlueprintPtr, *ViewBinding, Function);
-				if (ViewBinding->DestinationPath != SelectedField)
-				{
-					Subsystem->SetDestinationPathForBinding(WidgetBlueprintPtr, *ViewBinding, SelectedField);
-				}
-			}
-		}
-	}
-
-	FFieldSelectionContext GetSelectedSelectionContext(bool bSource) const
-	{
-		FFieldSelectionContext Result;
-
-		const UWidgetBlueprint* WidgetBlueprintPtr = WidgetBlueprintWeak.Get();
-		if (WidgetBlueprintPtr == nullptr)
-		{
-			return Result;
-		}
-		
-		if (const FMVVMBlueprintViewBinding* ViewBinding = GetThisViewBinding())
-		{
-			Result.BindingMode = ViewBinding->BindingType;
-
-			{
-				TArray<FMVVMConstFieldVariant> Fields = bSource ? ViewBinding->DestinationPath.GetFields(WidgetBlueprintPtr->SkeletonGeneratedClass) : ViewBinding->SourcePath.GetFields(WidgetBlueprintPtr->SkeletonGeneratedClass);
-				if (Fields.Num() > 0)
-				{
-					FMVVMConstFieldVariant LastField = Fields.Last();
-					if (LastField.IsProperty())
-					{
-						Result.AssignableTo = LastField.GetProperty();
-					}
-					else if (LastField.IsFunction())
-					{
-						if (const UFunction* Function = LastField.GetFunction())
-						{
-							Result.AssignableTo = BindingHelper::GetReturnProperty(Function);
-						}
-					}
-				}
-			}
-
-			if (!bSource && !ViewBinding->DestinationPath.GetWidgetName().IsNone())
-			{
-				Result.FixedBindingSource = FBindingSource::CreateForWidget(WidgetBlueprintPtr, ViewBinding->DestinationPath.GetWidgetName());
-			}
-
-			Result.bAllowWidgets = true;
-			Result.bAllowViewModels = bSource;
-			Result.bAllowConversionFunctions = false;
-			bool bIsReadingValue = (IsForwardBinding(ViewBinding->BindingType) && bSource)
-				|| (IsBackwardBinding(ViewBinding->BindingType) && !bSource);
-			if (!(IsBackwardBinding(ViewBinding->BindingType) && IsForwardBinding(ViewBinding->BindingType)))
-			{
-				Result.bAllowConversionFunctions = bIsReadingValue;
-			}
-
-			Result.bReadable = bIsReadingValue;
-			Result.bWritable = (IsForwardBinding(ViewBinding->BindingType) && !bSource)
-				|| (IsBackwardBinding(ViewBinding->BindingType) && bSource);
-		}
-		return Result;
-	}
-
-	FReply HandleFieldSelectorDrop(const FGeometry& MyGeometry, const FDragDropEvent& DragDropEvent, bool bSource)
-	{
-		TSharedPtr<FDecoratedDragDropOp> DragDropOp = DragDropEvent.GetOperationAs<FDecoratedDragDropOp>();
-		if (!DragDropOp.IsValid())
-		{
-			return FReply::Unhandled();
-		}
-
-		// Accept all drag-drop operations that are widget properties, but only accept view model fields when we are dropping into the Source box.
-		if (!DragDropOp->IsOfType<FWidgetPropertyDragDropOp>() && (!DragDropOp->IsOfType<FViewModelFieldDragDropOp>() || !bSource))
-		{
-			return FReply::Unhandled();
-		}
-
-		if (FMVVMBlueprintViewBinding* ViewBinding = GetThisViewBinding())
-		{
-			if (UWidgetBlueprint* WidgetBlueprintPtr = WidgetBlueprintWeak.Get())
-			{
-				TSharedPtr<FViewModelFieldDragDropOp> ViewModelFieldDragDropOp = DragDropEvent.GetOperationAs<FViewModelFieldDragDropOp>();
-				TSharedPtr<FWidgetPropertyDragDropOp> WidgetPropertyDragDropOp = DragDropEvent.GetOperationAs<FWidgetPropertyDragDropOp>();
-				bool bIsViewModelProperty = false;
-
-				if (ViewModelFieldDragDropOp)
-				{
-					bIsViewModelProperty = true;
-				}
-
-				UWidgetBlueprint* DragDropWidgetBP = bIsViewModelProperty ? ViewModelFieldDragDropOp->WidgetBP.Get() : WidgetPropertyDragDropOp->WidgetBP.Get();
-
-				if (WidgetBlueprintPtr == DragDropWidgetBP)
-				{
-					TArray<FFieldVariant> FieldPath = bIsViewModelProperty ? ViewModelFieldDragDropOp->DraggedField : WidgetPropertyDragDropOp->DraggedPropertyPath;
-					FMVVMBlueprintPropertyPath PropertyPath;
-
-					PropertyPath.ResetPropertyPath();
-					for (const FFieldVariant& Field : FieldPath)
-					{
-						PropertyPath.AppendPropertyPath(WidgetBlueprintPtr, FMVVMConstFieldVariant(Field));
-					}
-
-					UMVVMEditorSubsystem* Subsystem = GEditor->GetEditorSubsystem<UMVVMEditorSubsystem>();
-					if (bIsViewModelProperty)
-					{
-						if (ViewModelFieldDragDropOp->ViewModelId.IsValid())
-						{
-							PropertyPath.SetViewModelId(ViewModelFieldDragDropOp->ViewModelId);
-						}
-						else
-						{
-							return FReply::Unhandled();
-						}
-					}
-					else
-					{
-						if (UWidget* OwnerWidgetPtr = WidgetPropertyDragDropOp->OwnerWidget.Get())
-						{
-							// Search for the widget by its name in the widget tree
-							// If the widget is not found, we know it is the root preview widget so we use the blueprint name.
-							if (WidgetBlueprintPtr->WidgetTree->FindWidget(OwnerWidgetPtr->GetFName()) && WidgetBlueprintPtr->GetFName() != OwnerWidgetPtr->GetFName())
-							{
-								PropertyPath.SetWidgetName(OwnerWidgetPtr->GetFName());
-							}
-							else
-							{
-								PropertyPath.SetSelfContext();
-							}
-						}
-					}
-
-					if (bSource)
-					{
-						Subsystem->SetSourceToDestinationConversionFunction(WidgetBlueprintPtr, *ViewBinding, nullptr);
-						Subsystem->SetSourcePathForBinding(WidgetBlueprintPtr, *ViewBinding, PropertyPath);
-					}
-					else
-					{
-						Subsystem->SetDestinationToSourceConversionFunction(WidgetBlueprintPtr, *ViewBinding, nullptr);
-						Subsystem->SetDestinationPathForBinding(WidgetBlueprintPtr, *ViewBinding, PropertyPath);
-					}
-					return FReply::Handled();
-				}
-			}
-		}		
-		return FReply::Unhandled();
-	}
-
-	void HandleFieldSelectorDragEnter(const FGeometry& MyGeometry, const FDragDropEvent& DragDropEvent, bool bSource)
-	{
-		TSharedPtr<FDecoratedDragDropOp> DragDropOp = DragDropEvent.GetOperationAs<FDecoratedDragDropOp>();
-		if (DragDropOp.IsValid())
-		{
-			// Accept all drag-drop operations that are widget properties, but only accept view model fields when we are dropping into the Source box.
-			if (DragDropOp->IsOfType<FWidgetPropertyDragDropOp>() || (DragDropOp->IsOfType<FViewModelFieldDragDropOp>() && bSource))
-			{
-				if (UWidgetBlueprint* WidgetBlueprintPtr = WidgetBlueprintWeak.Get())
-				{
-					TSharedPtr<FViewModelFieldDragDropOp> ViewModelFieldDragDropOp = DragDropEvent.GetOperationAs<FViewModelFieldDragDropOp>();
-					TSharedPtr<FWidgetPropertyDragDropOp> WidgetPropertyDragDropOp = DragDropEvent.GetOperationAs<FWidgetPropertyDragDropOp>();
-					bool IsViewModelProperty = false;
-
-					if (ViewModelFieldDragDropOp)
-					{
-						IsViewModelProperty = true;
-					}
-
-					UWidgetBlueprint* DragDropWidgetBP = IsViewModelProperty ? ViewModelFieldDragDropOp->WidgetBP.Get() : WidgetPropertyDragDropOp->WidgetBP.Get();
-
-					if (DragDropWidgetBP && DragDropWidgetBP == WidgetBlueprintPtr)
-					{
-						DragDropOp->CurrentIconBrush = FAppStyle::GetBrush(TEXT("Graph.ConnectorFeedback.OK"));
-						return;
-					}
-				}
-			}	
-			else
-			{
-				DragDropOp->CurrentIconBrush = FAppStyle::GetBrush(TEXT("Graph.ConnectorFeedback.Error"));
-			}
-		}
-	}
-
-	ECheckBoxState IsExecutionModeOverrideChecked() const
-	{
-		if (FMVVMBlueprintViewBinding* ViewBinding = GetThisViewBinding())
-		{
-			return ViewBinding->bOverrideExecutionMode ? ECheckBoxState::Checked : ECheckBoxState::Unchecked;
-		}
-		return ECheckBoxState::Undetermined;
-	}
-
-	void OnExecutionModeOverrideChanged(ECheckBoxState NewState)
-	{
-		if (FMVVMBlueprintViewBinding* ViewBinding = GetThisViewBinding())
-		{
-			if (NewState == ECheckBoxState::Checked)
-			{
-				UMVVMEditorSubsystem* Subsystem = GEditor->GetEditorSubsystem<UMVVMEditorSubsystem>();
-				Subsystem->OverrideExecutionModeForBinding(WidgetBlueprintWeak.Get(), *ViewBinding, DefaultExecutionMode);
-			}
-			else
-			{
-				UMVVMEditorSubsystem* Subsystem = GEditor->GetEditorSubsystem<UMVVMEditorSubsystem>();
-				Subsystem->ResetExecutionModeForBinding(WidgetBlueprintWeak.Get(), *ViewBinding);
-			}
-		}
-	}
-
-	void OnExecutionModeSelectionChanged(EMVVMExecutionMode Value)
-	{
-		if (FMVVMBlueprintViewBinding* ViewBinding = GetThisViewBinding())
-		{
-			UMVVMEditorSubsystem* Subsystem = GEditor->GetEditorSubsystem<UMVVMEditorSubsystem>();
-			Subsystem->OverrideExecutionModeForBinding(WidgetBlueprintWeak.Get(), *ViewBinding, Value);
-		}
-	}
-
-	TSharedRef<SWidget> OnGetExecutionModeMenuContent()
-	{
-		const bool bCloseAfterSelection = true;
-		FMenuBuilder MenuBuilder(bCloseAfterSelection, nullptr, nullptr, true);
-
-		UEnum* Enum = StaticEnum<EMVVMExecutionMode>();
-		for (int32 Index = 0; Index < Enum->NumEnums() - 1; Index++)
-		{
-			EMVVMExecutionMode Mode = static_cast<EMVVMExecutionMode>(Enum->GetValueByIndex(Index));
-			if (!GetDefault<UMVVMDeveloperProjectSettings>()->IsExecutionModeAllowed(Mode))
-			{
-				continue;
-			}
-
-			MenuBuilder.AddMenuEntry(
-				Enum->GetDisplayNameTextByIndex(Index),
-				Enum->GetToolTipTextByIndex(Index),
-				FSlateIcon(),
-				FUIAction
-				(
-					FExecuteAction::CreateLambda([this, Mode]()
-					{
-						OnExecutionModeSelectionChanged(Mode);
-					})
-				)
-			);
-		}
-
-		return MenuBuilder.MakeWidget();
-	}
-
-	FText GetExecutioModeValue() const
-	{
-		EMVVMExecutionMode ExecutionMode = DefaultExecutionMode;
-		if (FMVVMBlueprintViewBinding* ViewBinding = GetThisViewBinding())
-		{
-			if (ViewBinding->bOverrideExecutionMode)
-			{
-				ExecutionMode = ViewBinding->OverrideExecutionMode;
-			}
-		}
-		return StaticEnum<EMVVMExecutionMode>()->GetDisplayNameTextByValue(static_cast<int64>(ExecutionMode));
-	}
-	
-	FText GetExecutioModeValueToolTip() const
-	{
-		EMVVMExecutionMode ExecutionMode = DefaultExecutionMode;
-		if (FMVVMBlueprintViewBinding* ViewBinding = GetThisViewBinding())
-		{
-			if (ViewBinding->bOverrideExecutionMode)
-			{
-				ExecutionMode = ViewBinding->OverrideExecutionMode;
-			}
-		}
-		return StaticEnum<EMVVMExecutionMode>()->GetToolTipTextByIndex(static_cast<int64>(ExecutionMode));
-	}
-	
-	bool IsExecutionModeOverridden() const
-	{
-		if (FMVVMBlueprintViewBinding* ViewBinding = GetThisViewBinding())
-		{
-			return ViewBinding->bOverrideExecutionMode;
-		}
-		return false;
-	}
-
-	void OnIsBindingEnableChanged(ECheckBoxState NewState)
-	{
-		if (NewState == ECheckBoxState::Undetermined)
-		{
-			return;
-		}
-
-		if (FMVVMBlueprintViewBinding* ViewBinding = GetThisViewBinding())
-		{
-			UMVVMEditorSubsystem* Subsystem = GEditor->GetEditorSubsystem<UMVVMEditorSubsystem>();
-			Subsystem->SetEnabledForBinding(WidgetBlueprintWeak.Get(), *ViewBinding, NewState == ECheckBoxState::Checked);
-		}
-	}
-
-	void OnIsBindingCompileChanged(ECheckBoxState NewState)
-	{
-		if (NewState == ECheckBoxState::Undetermined)
-		{
-			return;
-		}
-
-		if (FMVVMBlueprintViewBinding* ViewBinding = GetThisViewBinding())
-		{
-			UMVVMEditorSubsystem* Subsystem = GEditor->GetEditorSubsystem<UMVVMEditorSubsystem>();
-			Subsystem->SetCompileForBinding(WidgetBlueprintWeak.Get(), *ViewBinding, NewState == ECheckBoxState::Checked);
-		}
-	}
-
-	const FSlateBrush* GetBindingModeBrush(EMVVMBindingMode BindingMode) const
-	{
-		switch (BindingMode)
-		{
-		case EMVVMBindingMode::OneTimeToDestination:
-			return FMVVMEditorStyle::Get().GetBrush("BindingMode.OneTimeOneWay");
-		case EMVVMBindingMode::OneWayToDestination:
-			return FMVVMEditorStyle::Get().GetBrush("BindingMode.OneWayToSource");
-		case EMVVMBindingMode::OneWayToSource:
-			return FMVVMEditorStyle::Get().GetBrush("BindingMode.OneWay");
-		case EMVVMBindingMode::OneTimeToSource:
-			return nullptr;
-		case EMVVMBindingMode::TwoWay:
-			return FMVVMEditorStyle::Get().GetBrush("BindingMode.TwoWay");
-		default:
-			return nullptr;
-		}
-	}
-
-	const FSlateBrush* GetCurrentBindingModeBrush() const
-	{
-		if (FMVVMBlueprintViewBinding* ViewBinding = GetThisViewBinding())
-		{
-			return GetBindingModeBrush(ViewBinding->BindingType);
-		}
-		return nullptr;
-	}
-
-	FText GetCurrentBindingModeLabel() const
-	{
-		if (FMVVMBlueprintViewBinding* ViewBinding = GetThisViewBinding())
-		{
-			return GetBindingModeLabel(ViewBinding->BindingType);
-		}
-		return FText::GetEmpty();
-	}
-
-	FText GetBindingModeLabel(EMVVMBindingMode BindingMode) const
-	{
-		static FText OneTimeToDestinationLabel = LOCTEXT("OneTimeToDestinationLabel", "One Time To Widget");
-		static FText OneWayToDestinationLabel = LOCTEXT("OneWayToDestinationLabel", "One Way To Widget");
-		static FText OneWayToSourceLabel = LOCTEXT("OneWayToSourceLabel", "One Way To View Model");
-		static FText OneTimeToSourceLabel = LOCTEXT("OneTimeToSourceLabel", "One Time To View Model");
-		static FText TwoWayLabel = LOCTEXT("TwoWayLabel", "Two Way");
-
-		switch (BindingMode)
-		{
-		case EMVVMBindingMode::OneTimeToDestination:
-			return OneTimeToDestinationLabel;
-		case EMVVMBindingMode::OneWayToDestination:
-			return OneWayToDestinationLabel;
-		case EMVVMBindingMode::OneWayToSource:
-			return OneWayToSourceLabel;
-		case EMVVMBindingMode::OneTimeToSource:
-			return OneTimeToSourceLabel;
-		case EMVVMBindingMode::TwoWay:
-			return TwoWayLabel;
-		default:
-			return FText::GetEmpty();
-		}
-	}
-
-	TSharedRef<SWidget> GenerateBindingModeWidget(FName ValueName) const
-	{
-		const UEnum* ModeEnum = StaticEnum<EMVVMBindingMode>();
-		int32 Index = ModeEnum->GetIndexByName(ValueName);
-		EMVVMBindingMode MVVMBindingMode = EMVVMBindingMode(Index);
-		return SNew(SHorizontalBox)
-			+SHorizontalBox::Slot()
-			.AutoWidth()
-			[
-				SNew(SBox)
-				.HAlign(HAlign_Center)
-				.VAlign(VAlign_Center)
-				.WidthOverride(16)
-				.HeightOverride(16)
-				[
-					SNew(SImage)
-					.Image(GetBindingModeBrush(MVVMBindingMode))
-				]
-			]
-			+SHorizontalBox::Slot()
-			.AutoWidth()
-			.Padding(2.0, 0, 0, 0)
-			[
-				SNew(STextBlock)
-				.Text(GetBindingModeLabel(MVVMBindingMode))
-				.ToolTipText(ModeEnum->GetToolTipTextByIndex(Index))
-			];
-	}
-
-	void OnBindingModeSelectionChanged(FName ValueName, ESelectInfo::Type)
-	{
-		if (FMVVMBlueprintViewBinding* ViewBinding = GetThisViewBinding())
-		{
-			const UEnum* ModeEnum = StaticEnum<EMVVMBindingMode>();
-			EMVVMBindingMode NewMode = (EMVVMBindingMode) ModeEnum->GetValueByName(ValueName);
-
-			UMVVMEditorSubsystem* Subsystem = GEditor->GetEditorSubsystem<UMVVMEditorSubsystem>();
-			Subsystem->SetBindingTypeForBinding(WidgetBlueprintWeak.Get(), *ViewBinding, NewMode);
-		}
-	}
-
-private:
-	TSharedPtr<FBindingEntry> Entry;
-	TWeakObjectPtr<UWidgetBlueprint> WidgetBlueprintWeak;
-	TSharedPtr<SCustomDialog> ErrorDialog;
-	TArray<TSharedPtr<FText>> ErrorItems;
-	EMVVMExecutionMode DefaultExecutionMode = EMVVMExecutionMode::DelayedWhenSharedElseImmediate;
-};
-
-/**
-*
-*/
-class SEventRow : public STableRow<TSharedPtr<FBindingEntry>>
-{
-public:
-	SLATE_BEGIN_ARGS(SEventRow) {}
-	SLATE_END_ARGS()
-
-	void Construct(const FArguments& Args, const TSharedRef<STableViewBase>& OwnerTableView, const TSharedPtr<FBindingEntry>& InEntry, UWidgetBlueprint* InWidgetBlueprint)
-	{
-		Entry = InEntry;
-		ViewEvent = Entry->GetEvent();
-		WidgetBlueprintWeak = InWidgetBlueprint;
-
-		STableRow<TSharedPtr<FBindingEntry>>::Construct(
-			STableRow<TSharedPtr<FBindingEntry>>::FArguments()
-			.ShowWires(true)
-			.Style(FMVVMEditorStyle::Get(), "BindingView.BindingRow")
-			[
-				SNew(SBorder)
-				.BorderImage(FAppStyle::Get().GetBrush("PlainBorder"))
-				.Padding(0.0f)
-				.BorderBackgroundColor(this, &SEventRow::GetErrorBorderColor)
-				[
-					SNew(SBox)
-					.HeightOverride(30)
-					[
-						SNew(SHorizontalBox)
-						+ SHorizontalBox::Slot()
-						.Padding(2.0f, 0.0f)
-						.VAlign(VAlign_Center)
-						.HAlign(HAlign_Left)
-						.AutoWidth()
-						[
-							SNew(SCheckBox)
-							.IsChecked(this, &SEventRow::IsEventCompiled)
-							.OnCheckStateChanged(this, &SEventRow::OnIsEventCompileChanged)
-						]
-
-						+ SHorizontalBox::Slot()
-						.Padding(4.0f, 0.0f)
-						.VAlign(VAlign_Center)
-						.HAlign(HAlign_Left)
-						.AutoWidth()
-						[
-							SNew(SBox)
-							.MinDesiredWidth(150.0f)
-							[
-								SNew(SFieldSelector, InWidgetBlueprint)
-								.OnGetPropertyPath(this, &SEventRow::GetSelectedPropertyPath, true)
-								.OnGetConversionFunction(this, &SEventRow::GetSelectedConversionFunction, true)
-								.OnFieldSelectionChanged(this, &SEventRow::HandleFieldSelectionChanged, true)
-								.OnGetSelectionContext(this, &SEventRow::GetSelectedSelectionContext, true)
-								.OnDrop(this, &SEventRow::HandleFieldSelectorDrop, true)
-								.OnDragEnter(this, &SEventRow::HandleFieldSelectorDragEnter, true)
-								.ShowContext(false)
-							]
-						]
-
-						+ SHorizontalBox::Slot()
-						.Padding(4.0f, 0.0f)
-						.VAlign(VAlign_Fill)
-						.HAlign(HAlign_Left)
-						.AutoWidth()
-						[
-							SNew(SBox)
-							.HAlign(HAlign_Center)
-							.VAlign(VAlign_Center)
-							.WidthOverride(16.0f)
-							.HeightOverride(16.0f)
-							[
-								SNew(SImage)
-								.Image(FMVVMEditorStyle::Get().GetBrush("BindingMode.OneWay"))
-							]
-						]
-
-						+ SHorizontalBox::Slot()
-						.Padding(4.0f, 0.0f, 2.0f, 0.0f)
-						.VAlign(VAlign_Center)
-						.HAlign(HAlign_Left)
-						.AutoWidth()
-						[
-							SNew(SBox)
-							.MinDesiredWidth(150.0f)
-							[
-								SNew(SFieldSelector, InWidgetBlueprint)
-								.OnGetPropertyPath(this, &SEventRow::GetSelectedPropertyPath, false)
-								.OnGetConversionFunction(this, &SEventRow::GetSelectedConversionFunction, false)
-								.OnFieldSelectionChanged(this, &SEventRow::HandleFieldSelectionChanged, false)
-								.OnGetSelectionContext(this, &SEventRow::GetSelectedSelectionContext, false)
-								.OnDrop(this, &SEventRow::HandleFieldSelectorDrop, false)
-								.OnDragEnter(this, &SEventRow::HandleFieldSelectorDragEnter, false)
-							]
-						]
-
-						+ SHorizontalBox::Slot()
-						[
-							SNew(SSpacer)
-						]
-
-						+SHorizontalBox::Slot()
-						.Padding(2.0f, 0.0f)
-						.VAlign(VAlign_Center)
-						.HAlign(HAlign_Left)
-						.AutoWidth()
-						[
-							SNew(SSimpleButton)
-							.Icon(FAppStyle::Get().GetBrush("Icons.Error"))
-							.Visibility(this, &SEventRow::GetErrorButtonVisibility)
-							.ToolTipText(this, &SEventRow::GetErrorButtonToolTip)
-							.OnClicked(this, &SEventRow::OnErrorButtonClicked)
-						]
-					]
-				]
-			],
-			OwnerTableView);
-
-		TSharedPtr<SWidget> ChildContent = ChildSlot.DetachWidget();
-		ChildSlot
-		[
-			/* Add a single pixel top and bottom border for this widget. */
-			SNew(SBorder)
-			.BorderImage(FAppStyle::Get().GetBrush("Brushes.Recessed"))
-			.Padding(0.0f, 2.0f, 0.0f, 1.0f)
-			[
-			/* Restore the border that we're meant to have that reacts to selection/hover/etc. */
-				SNew(SBorder)
-				.BorderImage(this, &SEventRow::GetBorderImage)
-				.Padding(0.0f)
-				[
-					ChildContent.ToSharedRef()
-				]
-			]
-		];
-	}
-
-private:
-	FSlateColor GetErrorBorderColor() const
-	{
-		if (const UMVVMBlueprintViewEvent* Event = ViewEvent.Get())
-		{
-			if (Event->HasCompilationMessage(UMVVMBlueprintViewEvent::EMessageType::Error))
-			{
-				return FStyleColors::Error;
-			}
-			else if (Event->HasCompilationMessage(UMVVMBlueprintViewEvent::EMessageType::Warning))
-			{
-				return FStyleColors::Warning;
-			}
-		}
-		return FStyleColors::Transparent;
-	}
-
-	EVisibility GetErrorButtonVisibility() const
-	{
-		if (const UMVVMBlueprintViewEvent* Event = ViewEvent.Get())
-		{
-			bool bHasBindingError = Event->HasCompilationMessage(UMVVMBlueprintViewEvent::EMessageType::Error);
-			bool bHasBindingWarning = Event->HasCompilationMessage(UMVVMBlueprintViewEvent::EMessageType::Warning);
-			return bHasBindingError || bHasBindingWarning ? EVisibility::Visible : EVisibility::Hidden;
-		}
-		return EVisibility::Collapsed;
-	}
-
-	FText GetErrorButtonToolTip() const
-	{
-		if (const UMVVMBlueprintViewEvent* Event = ViewEvent.Get())
-		{
-			TArray<FText> BindingErrorList = Event->GetCompilationMessages(UMVVMBlueprintViewEvent::EMessageType::Error);
-			TArray<FText> BindingWarningList = Event->GetCompilationMessages(UMVVMBlueprintViewEvent::EMessageType::Warning);
-			BindingErrorList.Append(BindingWarningList);
-
-			static const FText NewLineText = FText::FromString(TEXT("\n"));
-			FText HintText = LOCTEXT("ErrorButtonText", "Errors and Warnings: (Click to show in a separate window)");
-			FText ErrorsText = FText::Join(NewLineText, BindingErrorList);
-			return FText::Join(NewLineText, HintText, ErrorsText);
-		}
-		return FText();
-	}
-
-	FReply OnErrorButtonClicked()
-	{
-		ErrorDialog.Reset();
-		ErrorItems.Reset();
-
-		if (const UMVVMBlueprintViewEvent* Event = ViewEvent.Get())
-		{
-			for (const FText& ErrorText : Event->GetCompilationMessages(UMVVMBlueprintViewEvent::EMessageType::Error))
-			{
-				ErrorItems.Add(MakeShared<FText>(ErrorText));
-			}
-
-			for (const FText& WarningText : Event->GetCompilationMessages(UMVVMBlueprintViewEvent::EMessageType::Warning))
-			{
-				ErrorItems.Add(MakeShared<FText>(WarningText));
-			}
-
-			const FText BindingDisplayName = Event->GetDisplayName(true);
-			ErrorDialog = SNew(SCustomDialog)
-				.Title(FText::Format(LOCTEXT("Compilation Errors and Warnings", "Compilation Errors and Warnings for {0}"), BindingDisplayName))
-				.Buttons({
-					SCustomDialog::FButton(LOCTEXT("OK", "OK"))
-					})
-				.Content()
-				[
-					SNew(SListView<TSharedPtr<FText>>)
-					.ListItemsSource(&ErrorItems)
-					.OnGenerateRow(this, &SEventRow::OnGenerateErrorRow)
-				];
-
-			ErrorDialog->Show();
-		}
-
-		return FReply::Handled();
-	}
-
-	TSharedRef<ITableRow> OnGenerateErrorRow(TSharedPtr<FText> Text, const TSharedRef<STableViewBase>& TableView) const
-	{
-		return SNew(STableRow<TSharedPtr<FText>>, TableView)
-			.Content()
-			[
-				SNew(SEditableTextBox)
-				.BackgroundColor(FStyleColors::Background)
-				.IsReadOnly(true)
-				.Text(*Text.Get())
-			];
-	}
-
-	ECheckBoxState IsEventEnabled() const
-	{
-		if (const UMVVMBlueprintViewEvent* Event = ViewEvent.Get())
-		{
-			return Event->bEnabled ? ECheckBoxState::Checked : ECheckBoxState::Unchecked;
-		}
-		return ECheckBoxState::Undetermined;
-	}
-
-	void OnIsEventEnableChanged(ECheckBoxState NewState)
-	{
-		if (NewState == ECheckBoxState::Undetermined)
-		{
-			return;
-		}
-
-		if (UMVVMBlueprintViewEvent* Event = ViewEvent.Get())
-		{
-			UMVVMEditorSubsystem* Subsystem = GEditor->GetEditorSubsystem<UMVVMEditorSubsystem>();
-			Subsystem->SetEnabledForEvent(Event, NewState == ECheckBoxState::Checked);
-		}
-	}
-
-	ECheckBoxState IsEventCompiled() const
-	{
-		if (const UMVVMBlueprintViewEvent* Event = ViewEvent.Get())
-		{
-			return Event->bCompile ? ECheckBoxState::Checked : ECheckBoxState::Unchecked;
-		}
-		return ECheckBoxState::Undetermined;
-	}
-
-	void OnIsEventCompileChanged(ECheckBoxState NewState)
-	{
-		if (NewState == ECheckBoxState::Undetermined)
-		{
-			return;
-		}
-
-		if (UMVVMBlueprintViewEvent* Event = ViewEvent.Get())
-		{
-			UMVVMEditorSubsystem* Subsystem = GEditor->GetEditorSubsystem<UMVVMEditorSubsystem>();
-			Subsystem->SetCompileForEvent(Event, NewState == ECheckBoxState::Checked);
-		}
-	}
-
-	FMVVMBlueprintPropertyPath GetSelectedPropertyPath(bool bEvent) const
-	{
-		if (const UMVVMBlueprintViewEvent* Event = ViewEvent.Get())
-		{
-			return bEvent ? Event->GetEventPath() : Event->GetDestinationPath();
-		}
-		return FMVVMBlueprintPropertyPath();
-	}
-
-	const UFunction* GetSelectedConversionFunction(bool bSourceToDest) const
-	{
-		if (const UMVVMBlueprintViewEvent* Event = ViewEvent.Get())
-		{
-		}
-
-		return nullptr;
-	}
-
-	void HandleFieldSelectionChanged(FMVVMBlueprintPropertyPath SelectedField, const UFunction* Function, bool bEvent)
-	{
-		if (UMVVMBlueprintViewEvent* Event = ViewEvent.Get())
-		{
-			UMVVMEditorSubsystem* Subsystem = GEditor->GetEditorSubsystem<UMVVMEditorSubsystem>();
-			if (bEvent)
-			{
-				if (Event->GetEventPath() != SelectedField)
-				{
-					Subsystem->SetEventPath(Event, SelectedField);
-				}
-			}
-			else
-			{
-				if (Event->GetDestinationPath() != SelectedField)
-				{
-					Subsystem->SetEventDestinationPath(Event, SelectedField);
-				}
-			}
-		}
-	}
-
-	FFieldSelectionContext GetSelectedSelectionContext(bool bEvent) const
-	{
-		FFieldSelectionContext Result;
-		const UWidgetBlueprint* WidgetBlueprintPtr = WidgetBlueprintWeak.Get();
-		if (WidgetBlueprintPtr == nullptr)
-		{
-			return Result;
-		}
-
-		if (UMVVMBlueprintViewEvent* Event = ViewEvent.Get())
-		{
-			Result.BindingMode = EMVVMBindingMode::OneTimeToDestination;
-			if (bEvent && !Event->GetEventPath().GetWidgetName().IsNone())
-			{
-				Result.FixedBindingSource = FBindingSource::CreateForWidget(WidgetBlueprintPtr, Event->GetEventPath().GetWidgetName());
-			}
-
-			Result.bAllowWidgets = true;
-			Result.bAllowViewModels = !bEvent;
-			Result.bAllowConversionFunctions = false;
-			Result.bReadable = bEvent;
-			Result.bWritable = !bEvent;
-		}
-		return Result;
-	}
-
-	FReply HandleFieldSelectorDrop(const FGeometry& MyGeometry, const FDragDropEvent& DragDropEvent, bool bEvent)
-	{
-		TSharedPtr<FDecoratedDragDropOp> DragDropOp = DragDropEvent.GetOperationAs<FDecoratedDragDropOp>();
-		if (!DragDropOp.IsValid())
-		{
-			return FReply::Unhandled();
-		}
-
-		// Accept all drag-drop operations that are widget properties, but only accept view model fields when we are dropping into the Source box.
-		if (!DragDropOp->IsOfType<FWidgetPropertyDragDropOp>() && (!DragDropOp->IsOfType<FViewModelFieldDragDropOp>() || !bEvent))
-		{
-			return FReply::Unhandled();
-		}
-
-		UMVVMBlueprintViewEvent* Event = ViewEvent.Get();
-		UWidgetBlueprint* WidgetBlueprintPtr = WidgetBlueprintWeak.Get();
-		if (Event == nullptr || WidgetBlueprintPtr == nullptr)
-		{
-			return FReply::Unhandled();
-		}
-
-		TSharedPtr<FViewModelFieldDragDropOp> ViewModelFieldDragDropOp = DragDropEvent.GetOperationAs<FViewModelFieldDragDropOp>();
-		TSharedPtr<FWidgetPropertyDragDropOp> WidgetPropertyDragDropOp = DragDropEvent.GetOperationAs<FWidgetPropertyDragDropOp>();
-		bool bIsViewModelProperty = ViewModelFieldDragDropOp != nullptr;
-
-		UWidgetBlueprint* DragDropWidgetBP = bIsViewModelProperty ? ViewModelFieldDragDropOp->WidgetBP.Get() : WidgetPropertyDragDropOp->WidgetBP.Get();
-		if (WidgetBlueprintPtr != DragDropWidgetBP)
-		{
-			return FReply::Unhandled();
-		}
-
-		TArray<FFieldVariant> FieldPath = bIsViewModelProperty ? ViewModelFieldDragDropOp->DraggedField : WidgetPropertyDragDropOp->DraggedPropertyPath;
-		FMVVMBlueprintPropertyPath PropertyPath;
-		for (const FFieldVariant& Field : FieldPath)
-		{
-			PropertyPath.AppendPropertyPath(WidgetBlueprintPtr, FMVVMConstFieldVariant(Field));
-		}
-
-		UMVVMEditorSubsystem* Subsystem = GEditor->GetEditorSubsystem<UMVVMEditorSubsystem>();
-		if (bIsViewModelProperty)
-		{
-			if (ViewModelFieldDragDropOp->ViewModelId.IsValid())
-			{
-				PropertyPath.SetViewModelId(ViewModelFieldDragDropOp->ViewModelId);
-			}
-			else
-			{
-				return FReply::Unhandled();
-			}
-		}
-		else
-		{
-			if (UWidget* OwnerWidgetPtr = WidgetPropertyDragDropOp->OwnerWidget.Get())
-			{
-				// Search for the widget by its name in the widget tree
-				// If the widget is not found, we know it is the root preview widget so we use the blueprint name.
-				if (WidgetBlueprintPtr->WidgetTree->FindWidget(OwnerWidgetPtr->GetFName()) && WidgetBlueprintPtr->GetFName() != OwnerWidgetPtr->GetFName())
-				{
-					PropertyPath.SetWidgetName(OwnerWidgetPtr->GetFName());
-				}
-				else
-				{
-					PropertyPath.SetSelfContext();
-				}
-			}
-		}
-
-		if (bEvent)
-		{
-			Subsystem->SetEventPath(Event, PropertyPath);
-		}
-		else
-		{
-			Subsystem->SetEventDestinationPath(Event, PropertyPath);
-		}
-		return FReply::Handled();
-	}
-
-	void HandleFieldSelectorDragEnter(const FGeometry& MyGeometry, const FDragDropEvent& DragDropEvent, bool bEvent)
-	{
-		TSharedPtr<FDecoratedDragDropOp> DragDropOp = DragDropEvent.GetOperationAs<FDecoratedDragDropOp>();
-		if (DragDropOp.IsValid())
-		{
-			// Accept all drag-drop operations that are widget properties, but only accept view model fields when we are dropping into the Source box.
-			if (DragDropOp->IsOfType<FWidgetPropertyDragDropOp>() || (DragDropOp->IsOfType<FViewModelFieldDragDropOp>() && !bEvent))
-			{
-				if (UWidgetBlueprint* WidgetBlueprintPtr = WidgetBlueprintWeak.Get())
-				{
-					TSharedPtr<FViewModelFieldDragDropOp> ViewModelFieldDragDropOp = DragDropEvent.GetOperationAs<FViewModelFieldDragDropOp>();
-					TSharedPtr<FWidgetPropertyDragDropOp> WidgetPropertyDragDropOp = DragDropEvent.GetOperationAs<FWidgetPropertyDragDropOp>();
-					bool IsViewModelProperty = false;
-
-					if (ViewModelFieldDragDropOp)
-					{
-						IsViewModelProperty = true;
-					}
-
-					UWidgetBlueprint* DragDropWidgetBP = IsViewModelProperty ? ViewModelFieldDragDropOp->WidgetBP.Get() : WidgetPropertyDragDropOp->WidgetBP.Get();
-
-					if (DragDropWidgetBP && DragDropWidgetBP == WidgetBlueprintPtr)
-					{
-						DragDropOp->CurrentIconBrush = FAppStyle::GetBrush(TEXT("Graph.ConnectorFeedback.OK"));
-						return;
-					}
-				}
-			}
-			else
-			{
-				DragDropOp->CurrentIconBrush = FAppStyle::GetBrush(TEXT("Graph.ConnectorFeedback.Error"));
-			}
-		}
-	}
-
-private:
-	TSharedPtr<FBindingEntry> Entry;
-	TWeakObjectPtr<UWidgetBlueprint> WidgetBlueprintWeak;
-	TWeakObjectPtr<UMVVMBlueprintViewEvent> ViewEvent;
-	TSharedPtr<SCustomDialog> ErrorDialog;
-	TArray<TSharedPtr<FText>> ErrorItems;
-};
-
-/**
- * 
- */
-class SFunctionParameterRow : public STableRow<TSharedPtr<FBindingEntry>>
-{
-	SLATE_BEGIN_ARGS(SFunctionParameterRow) {}
-	SLATE_END_ARGS()
-
-	void Construct(const FArguments& Args, const TSharedRef<STableViewBase>& OwnerTableView, const TSharedPtr<FBindingEntry>& InEntry, UWidgetBlueprint* InWidgetBlueprint)
-	{
-		Entry = InEntry;
-		check(Entry->GetRowType() == FBindingEntry::ERowType::BindingParameter || Entry->GetRowType() == FBindingEntry::ERowType::EventParameter);
-
-		WidgetBlueprint = InWidgetBlueprint;
-
-		UMVVMEditorSubsystem* EditorSubsystem = GEditor->GetEditorSubsystem<UMVVMEditorSubsystem>();
-
-		const FMVVMBlueprintPin* Pin = nullptr;
-		const UEdGraphPin* GraphPin = nullptr;
-		TSharedPtr<SWidget> ContentWidget;
-
-		if (Entry->GetRowType() == FBindingEntry::ERowType::BindingParameter)
-		{
-			bool bSimpleConversionFunction = false;
-			UMVVMBlueprintView* View = EditorSubsystem->GetView(InWidgetBlueprint);
-			FMVVMBlueprintViewBinding* Binding = Entry->GetBinding(View);
-			check(Binding);
-
-			const bool bSourceToDestination = UE::MVVM::IsForwardBinding(Binding->BindingType);
-			UMVVMBlueprintViewConversionFunction* ConversionFunction = Binding->Conversion.GetConversionFunction(bSourceToDestination);
-			check(ConversionFunction);
-
-			Pin = ConversionFunction->GetPins().FindByPredicate([Name = Entry->GetBindingParameterName()](const FMVVMBlueprintPin& Other){ return Other.GetName() == Name; });
-			GraphPin = ConversionFunction->GetOrCreateGraphPin(InWidgetBlueprint, Entry->GetBindingParameterName());
-
-			ContentWidget = SNew(SFunctionParameter, InWidgetBlueprint)
-				.BindingId(Binding->BindingId)
-				.ParameterName(Entry->GetBindingParameterName())
-				.SourceToDestination(bSourceToDestination)
-				.AllowDefault(!bSimpleConversionFunction);
-		}
-		else if (Entry->GetRowType() == FBindingEntry::ERowType::EventParameter)
-		{
-			UMVVMBlueprintViewEvent* ViewEvent = Entry->GetEvent();
-			check(ViewEvent);
-
-			Pin = ViewEvent->GetPins().FindByPredicate([Name = Entry->GetEventParameterName()](const FMVVMBlueprintPin& Other) { return Other.GetName() == Name; });
-			GraphPin = ViewEvent->GetOrCreateGraphPin(Entry->GetEventParameterName());
-
-			ContentWidget = SNew(SEventParameter, InWidgetBlueprint)
-				.Event(Entry->GetEvent())
-				.ParameterName(Entry->GetEventParameterName())
-				.AllowDefault(true);
-		}
-
-		FSlateColor PrimaryColor, SecondaryColor;
-		const FSlateBrush* PrimaryBrush = nullptr;
-		const FSlateBrush* SecondaryBrush = nullptr;
-		FText DisplayName, ToolTip;
-		bool bTextColorIsRed = false;
-		if (Pin == nullptr || GraphPin == nullptr || Pin->GetStatus() == EMVVMBlueprintPinStatus::Orphaned)
-		{
-			PrimaryBrush = FBlueprintEditorUtils::GetIconFromPin(FEdGraphPinType());
-			PrimaryColor = FLinearColor::Red;
-			DisplayName = GraphPin ? GraphPin->GetDisplayName() : (Pin ? FText::FromName(Pin->GetName()) : FText::GetEmpty());
-			bTextColorIsRed = true;
-		}
-		else
-		{
-			PrimaryBrush = FBlueprintEditor::GetVarIconAndColorFromPinType(GraphPin->PinType, PrimaryColor, SecondaryBrush, SecondaryColor);
-			DisplayName = GraphPin->GetDisplayName();
-			ToolTip = FText::FromString(GraphPin->PinToolTip);
-		}
-
-		TSharedPtr<STextBlock> DisplayNameTextBlock;
-		STableRow<TSharedPtr<FBindingEntry>>::Construct(
-			STableRow<TSharedPtr<FBindingEntry>>::FArguments()
-			.ShowWires(true)
-			.Style(FMVVMEditorStyle::Get(), "BindingView.ParameterRow")
-			[
-				SNew(SBox)
-				.HeightOverride(30)
-				.ToolTipText(ToolTip)
-				[
-					SNew(SHorizontalBox)
-
-					+ SHorizontalBox::Slot()
-					.HAlign(HAlign_Left)
-					.VAlign(VAlign_Center)
-					.Padding(0, 0, 8, 0)
-					.AutoWidth()
-					[
-						SNew(SHorizontalBox)
-						+ SHorizontalBox::Slot()
-						.VAlign(VAlign_Center)
-						.Padding(0, 0, 8, 0)
-						.AutoWidth()
-						[
-							SNew(SImage)
-							.Image(PrimaryBrush)
-							.ColorAndOpacity(PrimaryColor)
-							.DesiredSizeOverride(FVector2D(16, 16))
-						]
-						+ SHorizontalBox::Slot()
-						[
-							SAssignNew(DisplayNameTextBlock, STextBlock)
-							.Text(DisplayName)
-						]
-					]
-					+ SHorizontalBox::Slot()
-					.HAlign(HAlign_Left)
-					.VAlign(VAlign_Center)
-					.AutoWidth()
-					[
-						ContentWidget.ToSharedRef()
-					]
-				]
-			], 
-			OwnerTableView
-		);
-
-		if (bTextColorIsRed)
-		{
-			DisplayNameTextBlock->SetColorAndOpacity(PrimaryColor);
-		}
-	}
-
-private:
-	TWeakObjectPtr<UWidgetBlueprint> WidgetBlueprint;
-	TSharedPtr<FBindingEntry> Entry;
-};
 
-void SBindingsList::Construct(const FArguments& InArgs, TSharedPtr<SBindingsPanel> Owner, UMVVMWidgetBlueprintExtension_View* InMVVMExtension)
+void SBindingsList::Construct(const FArguments& InArgs, TSharedPtr<SBindingsPanel> Owner, TSharedPtr<FWidgetBlueprintEditor> InBlueprintEditor, UMVVMWidgetBlueprintExtension_View* InMVVMExtension)
 {
 	BindingPanel = Owner;
 	MVVMExtension = InMVVMExtension;
+	WeakBlueprintEditor = InBlueprintEditor;
 	check(InMVVMExtension);
 	check(InMVVMExtension->GetBlueprintView());
 
@@ -2245,23 +482,23 @@ TSharedRef<ITableRow> SBindingsList::GenerateEntryRow(TSharedPtr<FBindingEntry>
 		{
 			case FBindingEntry::ERowType::Group:
 			{
-				Row = SNew(SGroupRow, OwnerTable, Entry, MVVMExtensionPtr->GetWidgetBlueprint());
+				Row = SNew(UE::MVVM::BindingEntry::SGroupRow, OwnerTable, WeakBlueprintEditor.Pin(), MVVMExtensionPtr->GetWidgetBlueprint(), Entry);
 				break;
 			}
 			case FBindingEntry::ERowType::Binding:
 			{
-				Row = SNew(SBindingRow, OwnerTable, Entry, MVVMExtensionPtr->GetWidgetBlueprint());
+				Row = SNew(UE::MVVM::BindingEntry::SBindingRow, OwnerTable, WeakBlueprintEditor.Pin(), MVVMExtensionPtr->GetWidgetBlueprint(), Entry);
 				break;
 			}
 			case FBindingEntry::ERowType::BindingParameter:
 			case FBindingEntry::ERowType::EventParameter:
 			{
-				Row = SNew(SFunctionParameterRow, OwnerTable, Entry, MVVMExtensionPtr->GetWidgetBlueprint());
+				Row = SNew(UE::MVVM::BindingEntry::SFunctionParameterRow, OwnerTable, WeakBlueprintEditor.Pin(), MVVMExtensionPtr->GetWidgetBlueprint(), Entry);
 				break;
 			}
 			case FBindingEntry::ERowType::Event:
 			{
-				Row = SNew(SEventRow, OwnerTable, Entry, MVVMExtensionPtr->GetWidgetBlueprint());
+				Row = SNew(UE::MVVM::BindingEntry::SEventRow, OwnerTable, WeakBlueprintEditor.Pin(), MVVMExtensionPtr->GetWidgetBlueprint(), Entry);
 				break;
 			}
 		}
diff --git a/Engine/Plugins/Runtime/ModelViewViewModel/Source/ModelViewViewModelEditor/Private/Widgets/SMVVMViewBindingListView.h b/Engine/Plugins/Runtime/ModelViewViewModel/Source/ModelViewViewModelEditor/Private/Widgets/SMVVMViewBindingListView.h
--- a/Engine/Plugins/Runtime/ModelViewViewModel/Source/ModelViewViewModelEditor/Private/Widgets/SMVVMViewBindingListView.h
+++ b/Engine/Plugins/Runtime/ModelViewViewModel/Source/ModelViewViewModelEditor/Private/Widgets/SMVVMViewBindingListView.h
@@ -12,6 +12,7 @@ template <typename ItemType> class STreeView;
 class STableViewBase;
 class UMVVMWidgetBlueprintExtension_View;
 class UMVVMBlueprintViewEvent;
+class FWidgetBlueprintEditor;
 
 namespace UE::MVVM
 {
@@ -25,7 +26,7 @@ public:
 	SLATE_BEGIN_ARGS(SBindingsList) {}
 	SLATE_END_ARGS()
 
-	void Construct(const FArguments& InArgs, TSharedPtr<SBindingsPanel> Owner, UMVVMWidgetBlueprintExtension_View* MVVMExtension);
+	void Construct(const FArguments& InArgs, TSharedPtr<SBindingsPanel> Owner, TSharedPtr<FWidgetBlueprintEditor> BlueprintEditor, UMVVMWidgetBlueprintExtension_View* MVVMExtension);
 	~SBindingsList();
 
 	void Refresh();
@@ -59,6 +60,7 @@ private:
 	TArray<TSharedPtr<FBindingEntry>> AllRootGroups;
 	TArray<TSharedPtr<FBindingEntry>> FilteredRootGroups;
 	TWeakObjectPtr<UMVVMWidgetBlueprintExtension_View> MVVMExtension;
+	TWeakPtr<FWidgetBlueprintEditor> WeakBlueprintEditor;
 	mutable bool bSelectionChangedGuard = false;
 	FText FilterText;
 };
diff --git a/Engine/Plugins/Runtime/ModelViewViewModel/Source/ModelViewViewModelEditor/Private/Widgets/SMVVMViewBindingPanel.cpp b/Engine/Plugins/Runtime/ModelViewViewModel/Source/ModelViewViewModelEditor/Private/Widgets/SMVVMViewBindingPanel.cpp
--- a/Engine/Plugins/Runtime/ModelViewViewModel/Source/ModelViewViewModelEditor/Private/Widgets/SMVVMViewBindingPanel.cpp
+++ b/Engine/Plugins/Runtime/ModelViewViewModel/Source/ModelViewViewModelEditor/Private/Widgets/SMVVMViewBindingPanel.cpp
@@ -488,7 +488,7 @@ TSharedRef<SWidget> SBindingsPanel::GenerateEditViewWidget()
 	BindingsList = nullptr;
 	if (MVVMExtension.Get())
 	{
-		BindingsList = SNew(SBindingsList, StaticCastSharedRef<SBindingsPanel>(AsShared()), MVVMExtension.Get());
+		BindingsList = SNew(SBindingsList, StaticCastSharedRef<SBindingsPanel>(AsShared()), WeakBlueprintEditor.Pin(), MVVMExtension.Get());
 	}
 
 	DetailContainer = SNew(SBorder)
