MVVM: Sort conversion functions alphabetically

diff --git a/Engine/Plugins/Runtime/ModelViewViewModel/Source/ModelViewViewModelEditor/Private/Widgets/SMVVMFieldSelectorMenu.cpp b/Engine/Plugins/Runtime/ModelViewViewModel/Source/ModelViewViewModelEditor/Private/Widgets/SMVVMFieldSelectorMenu.cpp
--- a/Engine/Plugins/Runtime/ModelViewViewModel/Source/ModelViewViewModelEditor/Private/Widgets/SMVVMFieldSelectorMenu.cpp
+++ b/Engine/Plugins/Runtime/ModelViewViewModel/Source/ModelViewViewModelEditor/Private/Widgets/SMVVMFieldSelectorMenu.cpp
@@ -1153,6 +1153,8 @@ void SFieldSelectorMenu::HandleConversionFunctionCategorySelected(TSharedPtr<FCo
 			AddConversionFunctionChildrenRecursive(Item, ConversionFunctions);
 		}
 
+		Algo::SortBy(ConversionFunctions, &UFunction::GetDisplayNameText, FText::FSortPredicate());
+
 		FilterConversionFunctions();
 	}
 
