MVVM: Update after 30041295 review to the viewmodel customization

diff --git a/Engine/Plugins/Runtime/ModelViewViewModel/Source/ModelViewViewModelEditor/Private/Customizations/MVVMBlueprintViewModelContextCustomization.cpp b/Engine/Plugins/Runtime/ModelViewViewModel/Source/ModelViewViewModelEditor/Private/Customizations/MVVMBlueprintViewModelContextCustomization.cpp
--- a/Engine/Plugins/Runtime/ModelViewViewModel/Source/ModelViewViewModelEditor/Private/Customizations/MVVMBlueprintViewModelContextCustomization.cpp
+++ b/Engine/Plugins/Runtime/ModelViewViewModel/Source/ModelViewViewModelEditor/Private/Customizations/MVVMBlueprintViewModelContextCustomization.cpp
@@ -418,16 +418,9 @@ void FBlueprintViewModelContextDetailCustomization::HandleCreationTypeChanged()
 	{
 		if (FMVVMBlueprintViewModelContext* ContextPtr = Private::GetViewModelContext(ContextHandle.ToSharedRef()))
 		{
-			if ((EMVVMBlueprintViewModelContextCreationType)NewValue == EMVVMBlueprintViewModelContextCreationType::Manual)
-			{
-				ContextPtr->bOptional = true;
-				ContextPtr->bCreateSetterFunction = true;
-			}
-			else
-			{
-				ContextPtr->bOptional = false;
-				ContextPtr->bCreateSetterFunction = false;
-			}
+			const bool bIsManual = (EMVVMBlueprintViewModelContextCreationType)NewValue == EMVVMBlueprintViewModelContextCreationType::Manual;
+			ContextPtr->bOptional = bIsManual;
+			ContextPtr->bCreateSetterFunction = bIsManual;
 		}
 	}
 }
