MVVM: Change how the Binding subsystem is ticked. SlateApplication PreTick can be triggered from multiples time per engine loop.

diff --git a/Engine/Plugins/Runtime/ModelViewViewModel/Source/ModelViewViewModel/Private/View/MVVMBindingSubsystem.cpp b/Engine/Plugins/Runtime/ModelViewViewModel/Source/ModelViewViewModel/Private/View/MVVMBindingSubsystem.cpp
--- a/Engine/Plugins/Runtime/ModelViewViewModel/Source/ModelViewViewModel/Private/View/MVVMBindingSubsystem.cpp
+++ b/Engine/Plugins/Runtime/ModelViewViewModel/Source/ModelViewViewModel/Private/View/MVVMBindingSubsystem.cpp
@@ -13,25 +13,6 @@
 
 DECLARE_CYCLE_STAT(TEXT("MVVM Bindings"), STAT_MVVMBindingTick, STATGROUP_Slate);
 
-void UMVVMBindingSubsystem::Initialize(FSubsystemCollectionBase& Collection)
-{
-	Super::Initialize(Collection);
-
-	if (FSlateApplication::IsInitialized())
-	{
-		FSlateApplication::Get().OnPreTick().AddUObject(this, &UMVVMBindingSubsystem::HandlePreTick);
-	}
-}
-
-void UMVVMBindingSubsystem::Deinitialize()
-{
-	if (FSlateApplication::IsInitialized())
-	{
-		FSlateApplication::Get().OnPreTick().RemoveAll(this);
-	}
-	Super::Deinitialize();
-}
-
 namespace UE::MVVM::Private
 {
 struct FViewAndBinding
@@ -58,7 +39,7 @@ struct FViewAndBinding
 };
 }
 
-void UMVVMBindingSubsystem::HandlePreTick(float DeltaTime)
+void UMVVMBindingSubsystem::Tick(float DeltaTime)
 {
 	SCOPE_CYCLE_COUNTER(STAT_MVVMBindingTick);
 
@@ -142,6 +123,11 @@ void UMVVMBindingSubsystem::HandlePreTick(float DeltaTime)
 	}
 }
 
+TStatId UMVVMBindingSubsystem::GetStatId() const
+{
+	RETURN_QUICK_DECLARE_CYCLE_STAT(UMVVMBindingSubsystem, STATGROUP_Tickables);
+}
+
 void UMVVMBindingSubsystem::AddViewWithTickBinding(const UMVVMView* InView)
 {
 	check(!ViewsWithTickBindings.Contains(InView));
diff --git a/Engine/Plugins/Runtime/ModelViewViewModel/Source/ModelViewViewModel/Private/View/MVVMBindingSubsystem.h b/Engine/Plugins/Runtime/ModelViewViewModel/Source/ModelViewViewModel/Private/View/MVVMBindingSubsystem.h
--- a/Engine/Plugins/Runtime/ModelViewViewModel/Source/ModelViewViewModel/Private/View/MVVMBindingSubsystem.h
+++ b/Engine/Plugins/Runtime/ModelViewViewModel/Source/ModelViewViewModel/Private/View/MVVMBindingSubsystem.h
@@ -3,6 +3,7 @@
 #pragma once
 
 #include "Subsystems/EngineSubsystem.h"
+#include "Tickable.h"
 #include "UObject/ObjectKey.h"
 #include "View/MVVMViewClass.h"
 
@@ -12,14 +13,11 @@ class UMVVMView;
 
 /** */
 UCLASS(NotBlueprintable, Hidden)
-class UMVVMBindingSubsystem : public UEngineSubsystem
+class UMVVMBindingSubsystem : public UEngineSubsystem, public FTickableGameObject
 {
 	GENERATED_BODY()
 
 public:
-	virtual void Initialize(FSubsystemCollectionBase& Collection) override;
-	virtual void Deinitialize() override;
-
 	void AddViewWithTickBinding(const UMVVMView* View);
 	void RemoveViewWithTickBinding(const UMVVMView* View);
 
@@ -27,9 +25,24 @@ public:
 	void RemoveDelayedBindings(const UMVVMView* View);
 	void RemoveDelayedBindings(const UMVVMView* View, FMVVMViewClass_SourceKey SourceKey);
 
-private:
-	void HandlePreTick(float DeltaTIme);
+	//~ FTickableGameObject
+	virtual bool IsTickableWhenPaused() const override
+	{
+		return true;
+	}
+	virtual bool IsTickableInEditor() const override
+	{
+		return true;
+	}
+	virtual ETickableTickType GetTickableTickType() const override
+	{
+		return ETickableTickType::Always;
+	}
+
+	virtual void Tick(float DeltaTime) override;
+	virtual TStatId GetStatId() const override;
 
+private:
 	using FDelayedBindingList = TArray<FMVVMViewClass_BindingKey, TInlineAllocator<8>>;
 	using FDelayedBindingMap = TMap<TObjectKey<const UMVVMView>, FDelayedBindingList>;
 	FDelayedBindingMap DelayedBindings;
