MVVM: Add developer option to copy the generated graph from a binding. Move the common for tthe ENtryRow to the MVVMRowHelper.

diff --git a/Engine/Plugins/Runtime/ModelViewViewModel/Source/ModelViewViewModelBlueprint/Public/MVVMDeveloperProjectSettings.h b/Engine/Plugins/Runtime/ModelViewViewModel/Source/ModelViewViewModelBlueprint/Public/MVVMDeveloperProjectSettings.h
--- a/Engine/Plugins/Runtime/ModelViewViewModel/Source/ModelViewViewModelBlueprint/Public/MVVMDeveloperProjectSettings.h
+++ b/Engine/Plugins/Runtime/ModelViewViewModel/Source/ModelViewViewModelBlueprint/Public/MVVMDeveloperProjectSettings.h
@@ -126,6 +126,10 @@ public:
 	UPROPERTY(EditAnywhere, config, Category = "Viewmodel")
 	bool bShowViewSettings = true;
 
+	/** For the binding list widget, allow the user to generate a copy of the binding/event graph. */
+	UPROPERTY(EditAnywhere, config, Category = "Viewmodel")
+	bool bShowDeveloperGenerateGraphSettings = true;
+
 	/** When a conversion function requires a wrapper graph, add and save the generated graph to the blueprint. */
 	UPROPERTY(EditAnywhere, config, Category = "Viewmodel")
 	bool bAllowConversionFunctionGeneratedGraphInEditor = false;
diff --git a/Engine/Plugins/Runtime/ModelViewViewModel/Source/ModelViewViewModelEditor/Private/Framework/MVVMRowHelper.cpp b/Engine/Plugins/Runtime/ModelViewViewModel/Source/ModelViewViewModelEditor/Private/Framework/MVVMRowHelper.cpp
new file mode 100644
--- /dev/null
+++ b/Engine/Plugins/Runtime/ModelViewViewModel/Source/ModelViewViewModelEditor/Private/Framework/MVVMRowHelper.cpp
@@ -0,0 +1,575 @@
+// Copyright Epic Games, Inc. All Rights Reserved.
+
+#include "Framework/MVVMRowHelper.h"
+
+#include "Blueprint/WidgetTree.h"
+#include "BlueprintEditor.h"
+#include "Details/WidgetPropertyDragDropOp.h"
+#include "DragAndDrop/DecoratedDragDropOp.h"
+#include "Editor/EditorEngine.h"
+#include "EdGraphSchema_K2.h"
+#include "Kismet2/BlueprintEditorUtils.h"
+#include "Misc/Guid.h"
+#include "ScopedTransaction.h"
+#include "Widgets/ViewModelFieldDragDropOp.h"
+
+#include "Styling/AppStyle.h"
+
+#include "MVVMEditorSubsystem.h"
+#include "MVVMBlueprintView.h"
+#include "MVVMBlueprintViewBinding.h"
+#include "MVVMBlueprintViewConversionFunction.h"
+#include "MVVMBlueprintViewEvent.h"
+#include "MVVMWidgetBlueprintExtension_View.h"
+#include "Misc/MessageDialog.h"
+#include "Types/MVVMBindingEntry.h"
+#include "Types/MVVMBindingMode.h"
+
+extern UNREALED_API UEditorEngine* GEditor;
+
+#define LOCTEXT_NAMESPACE "BindingListView_Helper"
+
+namespace UE::MVVM::BindingEntry
+{
+
+void FRowHelper::GatherAllChildBindings(UMVVMBlueprintView* BlueprintView, const TConstArrayView<TSharedPtr<FBindingEntry>> Entries, TArray<const FMVVMBlueprintViewBinding*>& OutBindings, TArray<UMVVMBlueprintViewEvent*>& OutEvents)
+{
+	for (const TSharedPtr<FBindingEntry>& Entry : Entries)
+	{
+		if (Entry->GetRowType() == FBindingEntry::ERowType::Binding)
+		{
+			const FMVVMBlueprintViewBinding* Binding = Entry->GetBinding(BlueprintView);
+			if (Binding != nullptr)
+			{
+				OutBindings.AddUnique(Binding);
+			}
+		}
+
+		if (Entry->GetRowType() == FBindingEntry::ERowType::Event)
+		{
+			UMVVMBlueprintViewEvent* Event = Entry->GetEvent();
+			if (Event)
+			{
+				OutEvents.AddUnique(Event);
+			}
+		}
+		GatherAllChildBindings(BlueprintView, Entry->GetAllChildren(), OutBindings, OutEvents);
+	}
+}
+
+void FRowHelper::DeleteEntries(const UWidgetBlueprint* WidgetBlueprint, UMVVMBlueprintView* BlueprintView, TArrayView<const TSharedPtr<FBindingEntry>> Selection)
+{
+	if (WidgetBlueprint && BlueprintView)
+	{
+		TArray<const FMVVMBlueprintViewBinding*> BindingsToRemove;
+		TArray<UMVVMBlueprintViewEvent*> EventsToRemove;
+		GatherAllChildBindings(BlueprintView, Selection, BindingsToRemove, EventsToRemove);
+
+		if (BindingsToRemove.Num() == 0 && EventsToRemove.Num() == 0)
+		{
+			return;
+		}
+
+		TArray<FText> BindingDisplayNames;
+		for (const FMVVMBlueprintViewBinding* Binding : BindingsToRemove)
+		{
+			BindingDisplayNames.Add(FText::FromString(Binding->GetDisplayNameString(WidgetBlueprint)));
+		}
+		for (const UMVVMBlueprintViewEvent* Event : EventsToRemove)
+		{
+			BindingDisplayNames.Add(Event->GetDisplayName(true));
+		}
+
+		const FText Message = FText::Format(BindingDisplayNames.Num() == 1 ?
+			LOCTEXT("ConfirmDeleteSingle", "Are you sure that you want to delete this binding/event?\n\n{1}") :
+			LOCTEXT("ConfirmDeleteMultiple", "Are you sure that you want to delete these {0} bindings/events?\n\n{1}"),
+			BindingDisplayNames.Num(),
+			FText::Join(FText::FromString("\n"), BindingDisplayNames));
+
+		const FText Title = LOCTEXT("DeleteBindings", "Delete Bindings?");
+		if (FMessageDialog::Open(EAppMsgType::YesNo, EAppReturnType::Yes, Message, Title) == EAppReturnType::Yes)
+		{
+			FScopedTransaction Transaction(LOCTEXT("DeleteBindingsTransaction", "Delete Bindings"));
+			BlueprintView->Modify();
+
+			for (int32 BindingIndex = BindingsToRemove.Num() - 1; BindingIndex > -1; BindingIndex--)
+			{
+				const FMVVMBlueprintViewBinding* Binding = BindingsToRemove[BindingIndex];
+				BlueprintView->RemoveBinding(Binding);
+			}
+			for (UMVVMBlueprintViewEvent* Event : EventsToRemove)
+			{
+				BlueprintView->RemoveEvent(Event);
+			}
+		}
+	}
+}
+
+void FRowHelper::DuplicateBlueprintGraph(FBlueprintEditor* BlueprintEditor, UWidgetBlueprint* WidgetBlueprint, UMVVMBlueprintView* BlueprintView, TArrayView<const TSharedPtr<FBindingEntry>> Entries)
+{
+	auto DuplicateGraph = [WidgetBlueprint, BlueprintEditor](UEdGraph* Graph)
+		{
+			if (Graph)
+			{
+				const FScopedTransaction Transaction(LOCTEXT("DuplicateGraph", "Duplicate Graph"));
+				WidgetBlueprint->Modify();
+
+				UEdGraph* DuplicatedGraph = GetDefault<UEdGraphSchema_K2>()->DuplicateGraph(Graph);
+				check(DuplicatedGraph);
+
+				for (UEdGraphNode* EdGraphNode : DuplicatedGraph->Nodes)
+				{
+					if (EdGraphNode)
+					{
+						EdGraphNode->CreateNewGuid();
+					}
+				}
+
+				DuplicatedGraph->Modify();
+				WidgetBlueprint->FunctionGraphs.Add(DuplicatedGraph);
+
+				FBlueprintEditorUtils::MarkBlueprintAsStructurallyModified(WidgetBlueprint);
+				if (BlueprintEditor)
+				{
+					BlueprintEditor->OpenDocument(DuplicatedGraph, FDocumentTracker::OpenNewDocument);
+				}
+			}
+		};
+
+	const UMVVMEditorSubsystem* EditorSubsystem = GEditor->GetEditorSubsystem<UMVVMEditorSubsystem>();
+	for (const TSharedPtr<FBindingEntry>& Entry : Entries)
+	{
+		if (Entry->GetRowType() == FBindingEntry::ERowType::Binding)
+		{
+			if (FMVVMBlueprintViewBinding* ViewBinding = Entry->GetBinding(BlueprintView))
+			{
+				UMVVMBlueprintViewConversionFunction* ConversionFunctionA = ViewBinding->Conversion.GetConversionFunction(true);
+				UMVVMBlueprintViewConversionFunction* ConversionFunctionB = ViewBinding->Conversion.GetConversionFunction(false);
+				if (ConversionFunctionA)
+				{
+					DuplicateGraph(ConversionFunctionA->GetWrapperGraph());
+				}
+				if (ConversionFunctionB)
+				{
+					DuplicateGraph(ConversionFunctionB->GetWrapperGraph());
+				}
+			}
+		}
+		else if (Entry->GetRowType() == FBindingEntry::ERowType::Event)
+		{
+			if (UMVVMBlueprintViewEvent* Event = Entry->GetEvent())
+			{
+				DuplicateGraph(Event->GetWrapperGraph());
+			}
+		}
+	}
+}
+
+TOptional<FMVVMBlueprintPropertyPath> FRowHelper::DropFieldSelector(const UWidgetBlueprint* WidgetBlueprint, const FDragDropEvent& DragDropEvent)
+{
+	if (WidgetBlueprint == nullptr)
+	{
+		return TOptional<FMVVMBlueprintPropertyPath>();
+	}
+
+	TSharedPtr<FDecoratedDragDropOp> DragDropOp = DragDropEvent.GetOperationAs<FDecoratedDragDropOp>();
+	if (!DragDropOp.IsValid())
+	{
+		return TOptional<FMVVMBlueprintPropertyPath>();
+	}
+
+	// Accept all drag-drop operations that are widget properties, but only accept view model fields when we are dropping into the Source box.
+	if (!DragDropOp->IsOfType<FWidgetPropertyDragDropOp>() && (!DragDropOp->IsOfType<FViewModelFieldDragDropOp>()))
+	{
+		return TOptional<FMVVMBlueprintPropertyPath>();
+	}
+
+	TSharedPtr<FViewModelFieldDragDropOp> ViewModelFieldDragDropOp = DragDropEvent.GetOperationAs<FViewModelFieldDragDropOp>();
+	TSharedPtr<FWidgetPropertyDragDropOp> WidgetPropertyDragDropOp = DragDropEvent.GetOperationAs<FWidgetPropertyDragDropOp>();
+	bool bIsViewModelProperty = ViewModelFieldDragDropOp != nullptr;
+
+	UWidgetBlueprint* DragDropWidgetBP = bIsViewModelProperty ? ViewModelFieldDragDropOp->WidgetBP.Get() : WidgetPropertyDragDropOp->WidgetBP.Get();
+	if (WidgetBlueprint != DragDropWidgetBP)
+	{
+		return TOptional<FMVVMBlueprintPropertyPath>();
+	}
+
+	TArray<FFieldVariant> FieldPath = bIsViewModelProperty ? ViewModelFieldDragDropOp->DraggedField : WidgetPropertyDragDropOp->DraggedPropertyPath;
+	FMVVMBlueprintPropertyPath PropertyPath;
+	for (const FFieldVariant& Field : FieldPath)
+	{
+		PropertyPath.AppendPropertyPath(WidgetBlueprint, FMVVMConstFieldVariant(Field));
+	}
+
+	UMVVMEditorSubsystem* Subsystem = GEditor->GetEditorSubsystem<UMVVMEditorSubsystem>();
+	if (bIsViewModelProperty)
+	{
+		if (ViewModelFieldDragDropOp->ViewModelId.IsValid())
+		{
+			PropertyPath.SetViewModelId(ViewModelFieldDragDropOp->ViewModelId);
+		}
+		else
+		{
+			return TOptional<FMVVMBlueprintPropertyPath>();
+		}
+	}
+	else
+	{
+		if (UWidget* OwnerWidgetPtr = WidgetPropertyDragDropOp->OwnerWidget.Get())
+		{
+			// Search for the widget by its name in the widget tree
+			// If the widget is not found, we know it is the root preview widget so we use the blueprint name.
+			if (WidgetBlueprint->WidgetTree->FindWidget(OwnerWidgetPtr->GetFName()) && WidgetBlueprint->GetFName() != OwnerWidgetPtr->GetFName())
+			{
+				PropertyPath.SetWidgetName(OwnerWidgetPtr->GetFName());
+			}
+			else
+			{
+				PropertyPath.SetSelfContext();
+			}
+		}
+	}
+
+	return PropertyPath;
+}
+
+void FRowHelper::DragEnterFieldSelector(const UWidgetBlueprint* WidgetBlueprint, const FDragDropEvent& DragDropEvent)
+{
+	TSharedPtr<FDecoratedDragDropOp> DragDropOp = DragDropEvent.GetOperationAs<FDecoratedDragDropOp>();
+	if (DragDropOp.IsValid())
+	{
+		// Accept all drag-drop operations that are widget properties, but only accept view model fields when we are dropping into the Source box.
+		if (DragDropOp->IsOfType<FWidgetPropertyDragDropOp>() || DragDropOp->IsOfType<FViewModelFieldDragDropOp>())
+		{
+			TSharedPtr<FViewModelFieldDragDropOp> ViewModelFieldDragDropOp = DragDropEvent.GetOperationAs<FViewModelFieldDragDropOp>();
+			TSharedPtr<FWidgetPropertyDragDropOp> WidgetPropertyDragDropOp = DragDropEvent.GetOperationAs<FWidgetPropertyDragDropOp>();
+			bool IsViewModelProperty = false;
+
+			if (ViewModelFieldDragDropOp)
+			{
+				IsViewModelProperty = true;
+			}
+
+			UWidgetBlueprint* DragDropWidgetBP = IsViewModelProperty ? ViewModelFieldDragDropOp->WidgetBP.Get() : WidgetPropertyDragDropOp->WidgetBP.Get();
+			if (DragDropWidgetBP && DragDropWidgetBP == WidgetBlueprint)
+			{
+				DragDropOp->CurrentIconBrush = FAppStyle::GetBrush(TEXT("Graph.ConnectorFeedback.OK"));
+			}
+		}
+		else
+		{
+			DragDropOp->CurrentIconBrush = FAppStyle::GetBrush(TEXT("Graph.ConnectorFeedback.Error"));
+		}
+	}
+}
+
+namespace Private
+{
+
+TArray<TSharedPtr<FBindingEntry>> WeakToSharedPtr(const TSharedPtr<TArray<TWeakPtr<FBindingEntry>>>& CopiedEntries)
+{
+	TArray<TSharedPtr<FBindingEntry>> Result;
+	Result.Reserve(CopiedEntries->Num());
+	for (TWeakPtr<FBindingEntry> Entry : *CopiedEntries)
+	{
+		if (TSharedPtr<FBindingEntry> Pin = Entry.Pin())
+		{
+			Result.Add(Pin);
+		}
+	}
+	return Result;
+}
+
+void HandleDeleteEntries(UWidgetBlueprint* WidgetBlueprint, UMVVMBlueprintView* View, TSharedPtr<TArray<TWeakPtr<FBindingEntry>>> Entries)
+{
+	FRowHelper::DeleteEntries(WidgetBlueprint, View, WeakToSharedPtr(Entries));
+}
+
+void HandleResetSelectedPin(UWidgetBlueprint* WidgetBlueprint, UMVVMBlueprintView* View, TSharedPtr<TArray<TWeakPtr<FBindingEntry>>> Entries)
+{
+	const UMVVMEditorSubsystem* EditorSubsystem = GEditor->GetEditorSubsystem<UMVVMEditorSubsystem>();
+	for (const TSharedPtr<FBindingEntry>& Entry : WeakToSharedPtr(Entries))
+	{
+		if (Entry->GetRowType() == FBindingEntry::ERowType::EventParameter)
+		{
+			EditorSubsystem->ResetPinToDefaultValue(WidgetBlueprint, Entry->GetEvent(), Entry->GetEventParameterName());
+		}
+		else if (Entry->GetRowType() == FBindingEntry::ERowType::BindingParameter)
+		{
+			if (FMVVMBlueprintViewBinding* Binding = Entry->GetBinding(View))
+			{
+				const bool bSourceToDestination = UE::MVVM::IsForwardBinding(Binding->BindingType);
+				EditorSubsystem->ResetPinToDefaultValue(WidgetBlueprint, *Binding, Entry->GetBindingParameterName(), bSourceToDestination);
+			}
+		}
+	}
+}
+
+void HandleBreakSelectedPin(UWidgetBlueprint* WidgetBlueprint, UMVVMBlueprintView* View, TSharedPtr<TArray<TWeakPtr<FBindingEntry>>> Entries)
+{
+	const UMVVMEditorSubsystem* EditorSubsystem = GEditor->GetEditorSubsystem<UMVVMEditorSubsystem>();
+	for (const TSharedPtr<FBindingEntry>& Entry : WeakToSharedPtr(Entries))
+	{
+		if (Entry->GetRowType() == FBindingEntry::ERowType::EventParameter)
+		{
+			EditorSubsystem->SplitPin(WidgetBlueprint, Entry->GetEvent(), Entry->GetEventParameterName());
+		}
+		else if (Entry->GetRowType() == FBindingEntry::ERowType::BindingParameter)
+		{
+			if (FMVVMBlueprintViewBinding* Binding = Entry->GetBinding(View))
+			{
+				const bool bSourceToDestination = UE::MVVM::IsForwardBinding(Binding->BindingType);
+				EditorSubsystem->SplitPin(WidgetBlueprint, *Binding, Entry->GetBindingParameterName(), bSourceToDestination);
+			}
+		}
+	}
+}
+
+void HandleRecombineSelectedPin(UWidgetBlueprint* WidgetBlueprint, UMVVMBlueprintView* View, TSharedPtr<TArray<TWeakPtr<FBindingEntry>>> Entries)
+{
+	const UMVVMEditorSubsystem* EditorSubsystem = GEditor->GetEditorSubsystem<UMVVMEditorSubsystem>();
+	for (const TSharedPtr<FBindingEntry>& Entry : WeakToSharedPtr(Entries))
+	{
+		if (Entry->GetRowType() == FBindingEntry::ERowType::EventParameter)
+		{
+			EditorSubsystem->RecombinePin(WidgetBlueprint, Entry->GetEvent(), Entry->GetEventParameterName());
+		}
+		else if (Entry->GetRowType() == FBindingEntry::ERowType::BindingParameter)
+		{
+			if (FMVVMBlueprintViewBinding* Binding = Entry->GetBinding(View))
+			{
+				const bool bSourceToDestination = UE::MVVM::IsForwardBinding(Binding->BindingType);
+				EditorSubsystem->RecombinePin(WidgetBlueprint, *Binding, Entry->GetBindingParameterName(), bSourceToDestination);
+			}
+		}
+	}
+}
+
+void HandleResetOrphanedSelectedPin(UWidgetBlueprint* WidgetBlueprint, UMVVMBlueprintView* View, TSharedPtr<TArray<TWeakPtr<FBindingEntry>>> Entries)
+{
+	const UMVVMEditorSubsystem* EditorSubsystem = GEditor->GetEditorSubsystem<UMVVMEditorSubsystem>();
+	for (const TSharedPtr<FBindingEntry>& Entry : WeakToSharedPtr(Entries))
+	{
+		if (Entry->GetRowType() == FBindingEntry::ERowType::EventParameter)
+		{
+			EditorSubsystem->ResetOrphanedPin(WidgetBlueprint, Entry->GetEvent(), Entry->GetEventParameterName());
+		}
+		else if (Entry->GetRowType() == FBindingEntry::ERowType::BindingParameter)
+		{
+			if (FMVVMBlueprintViewBinding* Binding = Entry->GetBinding(View))
+			{
+				const bool bSourceToDestination = UE::MVVM::IsForwardBinding(Binding->BindingType);
+				EditorSubsystem->ResetOrphanedPin(WidgetBlueprint, *Binding, Entry->GetBindingParameterName(), bSourceToDestination);
+			}
+		}
+	}
+}
+
+} //namespace
+
+FMenuBuilder FRowHelper::CreateContextMenu(UWidgetBlueprint* WidgetBlueprint, UMVVMBlueprintView* View, TConstArrayView<TSharedPtr<FBindingEntry>> Entries)
+{
+	const bool bShouldCloseWindowAfterMenuSelection = true;
+	FMenuBuilder MenuBuilder(bShouldCloseWindowAfterMenuSelection, nullptr);
+
+	if (WidgetBlueprint && View && Entries.Num() > 0)
+	{
+		TSharedPtr<TArray<TWeakPtr<FBindingEntry>>> CopiedEntries;
+		CopiedEntries = MakeShared<TArray<TWeakPtr<FBindingEntry>>>();
+		CopiedEntries->Reserve(Entries.Num());
+		for (const TSharedPtr<FBindingEntry>& Entry : Entries)
+		{
+			CopiedEntries->Add(Entry);
+		}
+
+		const UMVVMEditorSubsystem* EditorSubsystem = GEditor->GetEditorSubsystem<UMVVMEditorSubsystem>();
+		{
+			bool bCanRemoveEntry = true;
+			for (const TSharedPtr<FBindingEntry>& Entry : Entries)
+			{
+				switch (Entry->GetRowType())
+				{
+				case FBindingEntry::ERowType::Group:
+				case FBindingEntry::ERowType::Binding:
+				case FBindingEntry::ERowType::Event:
+					break;
+				default:
+					bCanRemoveEntry = false;
+					break;
+				}
+			}
+
+			FUIAction RemoveAction;
+			RemoveAction.ExecuteAction = FExecuteAction::CreateStatic(&Private::HandleDeleteEntries, WidgetBlueprint, View, CopiedEntries);
+			RemoveAction.CanExecuteAction = FCanExecuteAction::CreateLambda([bCanRemoveEntry]() { return bCanRemoveEntry; });
+			MenuBuilder.AddMenuEntry(LOCTEXT("RemoveBinding", "Remove"),
+				LOCTEXT("RemoveBindingTooltip", "Remove bindings or events."),
+				FSlateIcon(FAppStyle::GetAppStyleSetName(), "Icons.Delete"),
+				RemoveAction);
+		}
+
+		{
+			bool bCanSplitPin = true;
+			bool bCanRecombinePin = true;
+			bool bCanRecombinePinVisible = true;
+			bool bCanResetPin = true;
+			bool bCanResetPinVisible = true;
+			bool bCanResetOrphanedPin = true;
+			auto SetAll = [&](bool Value)
+				{
+					bCanSplitPin = Value;
+					bCanRecombinePin = Value;
+					bCanRecombinePinVisible = Value;
+					bCanResetPin = Value;
+					bCanResetPinVisible = Value;
+					bCanResetOrphanedPin = Value;
+				};
+			auto AllFalse = [&]()
+				{
+					return !bCanRecombinePin && !bCanSplitPin && !bCanRecombinePinVisible && bCanResetPin && !bCanResetPinVisible && !bCanResetOrphanedPin;
+				};
+
+			for (const TSharedPtr<FBindingEntry>& Entry : Entries)
+			{
+				if (Entry->GetRowType() == FBindingEntry::ERowType::EventParameter)
+				{
+					if (!EditorSubsystem->CanSplitPin(WidgetBlueprint, Entry->GetEvent(), Entry->GetEventParameterName()))
+					{
+						bCanSplitPin = false;
+					}
+					if (!EditorSubsystem->CanRecombinePin(WidgetBlueprint, Entry->GetEvent(), Entry->GetEventParameterName()))
+					{
+						bCanRecombinePin = false;
+					}
+					if (!EditorSubsystem->CanResetPinToDefaultValue(WidgetBlueprint, Entry->GetEvent(), Entry->GetEventParameterName()))
+					{
+						bCanResetPin = false;
+					}
+					if (!EditorSubsystem->CanResetOrphanedPin(WidgetBlueprint, Entry->GetEvent(), Entry->GetEventParameterName()))
+					{
+						bCanResetOrphanedPin = false;
+					}
+
+					UMVVMBlueprintViewEvent* ViewEvent = Entry->GetEvent();
+					UEdGraphPin* GraphPin = ViewEvent ? ViewEvent->GetOrCreateGraphPin(Entry->GetEventParameterName()) : nullptr;
+					if (GraphPin == nullptr)
+					{
+						bCanRecombinePinVisible = false;
+						bCanResetPinVisible = false;
+					}
+					else
+					{
+						if (GraphPin->ParentPin == nullptr)
+						{
+							bCanRecombinePinVisible = false;
+						}
+						if (!FMVVMBlueprintPin::IsInputPin(GraphPin) || GetDefault<UEdGraphSchema_K2>()->ShouldHidePinDefaultValue(GraphPin))
+						{
+							bCanResetPinVisible = false;
+						}
+					}
+				}
+				else if (Entry->GetRowType() == FBindingEntry::ERowType::BindingParameter)
+				{
+					const FMVVMBlueprintViewBinding* Binding = Entry->GetBinding(View);
+					if (Binding)
+					{
+						const bool bSourceToDestination = UE::MVVM::IsForwardBinding(Binding->BindingType);
+						if (!EditorSubsystem->CanSplitPin(WidgetBlueprint, *Binding, Entry->GetBindingParameterName(), bSourceToDestination))
+						{
+							bCanSplitPin = false;
+						}
+						if (!EditorSubsystem->CanRecombinePin(WidgetBlueprint, *Binding, Entry->GetBindingParameterName(), bSourceToDestination))
+						{
+							bCanRecombinePin = false;
+						}
+						if (!EditorSubsystem->CanResetPinToDefaultValue(WidgetBlueprint, *Binding, Entry->GetBindingParameterName(), bSourceToDestination))
+						{
+							bCanResetPin = false;
+						}
+						if (!EditorSubsystem->CanResetOrphanedPin(WidgetBlueprint, *Binding, Entry->GetBindingParameterName(), bSourceToDestination))
+						{
+							bCanResetOrphanedPin = false;
+						}
+
+						UEdGraphPin* GraphPin = EditorSubsystem->GetConversionFunctionArgumentPin(WidgetBlueprint, *Binding, Entry->GetBindingParameterName(), bSourceToDestination);
+						if (GraphPin == nullptr)
+						{
+							bCanRecombinePinVisible = false;
+							bCanResetPinVisible = false;
+						}
+						else
+						{
+							if (GraphPin->ParentPin == nullptr)
+							{
+								bCanRecombinePinVisible = false;
+							}
+							if (!FMVVMBlueprintPin::IsInputPin(GraphPin) || GetDefault<UEdGraphSchema_K2>()->ShouldHidePinDefaultValue(GraphPin))
+							{
+								bCanResetPinVisible = false;
+							}
+						}
+					}
+					else
+					{
+						SetAll(false);
+					}
+				}
+				else
+				{
+					SetAll(false);
+				}
+
+				if (AllFalse())
+				{
+					break;
+				}
+			}
+
+			if (bCanResetPinVisible)
+			{
+				FUIAction ResetPinAction;
+				ResetPinAction.ExecuteAction = FExecuteAction::CreateStatic(&Private::HandleResetSelectedPin, WidgetBlueprint, View, CopiedEntries);
+				ResetPinAction.CanExecuteAction = FCanExecuteAction::CreateLambda([bCanResetPin]() { return bCanResetPin; });
+				MenuBuilder.AddMenuEntry(LOCTEXT("ResetPin", "Reset to Default Value"),
+					LOCTEXT("ResetPinTooltip", "Reset value of this pin to the default"),
+					FSlateIcon(),
+					ResetPinAction);
+			}
+			if (bCanSplitPin)
+			{
+				FUIAction SplitPinAction;
+				SplitPinAction.ExecuteAction = FExecuteAction::CreateStatic(&Private::HandleBreakSelectedPin, WidgetBlueprint, View, CopiedEntries);
+				MenuBuilder.AddMenuEntry(LOCTEXT("BreakPin", "Split Struct Pin"),
+					LOCTEXT("BreakPinTooltip", "Breaks a struct pin in to a separate pin per element."),
+					FSlateIcon(),
+					SplitPinAction);
+			}
+			if (bCanRecombinePinVisible)
+			{
+				FUIAction RecombinePinAction;
+				RecombinePinAction.ExecuteAction = FExecuteAction::CreateStatic(&Private::HandleRecombineSelectedPin, WidgetBlueprint, View, CopiedEntries);
+				RecombinePinAction.CanExecuteAction = FCanExecuteAction::CreateLambda([bCanRecombinePin]() { return bCanRecombinePin; });
+				MenuBuilder.AddMenuEntry(LOCTEXT("RecombinePin", "Recombine Struct Pin"),
+					LOCTEXT("RecombinePinTooltip", "Takes struct pins that have been broken in to composite elements and combines them back to a single struct pin."),
+					FSlateIcon(),
+					RecombinePinAction);
+			}
+			if (bCanResetOrphanedPin)
+			{
+				FUIAction ResetOrphanedPinAction;
+				ResetOrphanedPinAction.ExecuteAction = FExecuteAction::CreateStatic(&Private::HandleResetOrphanedSelectedPin, WidgetBlueprint, View, CopiedEntries);
+				MenuBuilder.AddMenuEntry(LOCTEXT("ResetOrphanedPin", "Remove the Orphaned Struct Pin"),
+					LOCTEXT("ResetOrphanedPinTooltip", "Removes pins that used to exist but do not exist anymore."),
+					FSlateIcon(),
+					ResetOrphanedPinAction);
+			}
+		}
+	}
+
+	return MenuBuilder;
+}
+
+} // namespace UE::MVVM::BindingEntry
+
+#undef LOCTEXT_NAMESPACE
diff --git a/Engine/Plugins/Runtime/ModelViewViewModel/Source/ModelViewViewModelEditor/Private/Framework/MVVMRowHelper.h b/Engine/Plugins/Runtime/ModelViewViewModel/Source/ModelViewViewModelEditor/Private/Framework/MVVMRowHelper.h
new file mode 100644
--- /dev/null
+++ b/Engine/Plugins/Runtime/ModelViewViewModel/Source/ModelViewViewModelEditor/Private/Framework/MVVMRowHelper.h
@@ -0,0 +1,34 @@
+// Copyright Epic Games, Inc. All Rights Reserved.
+
+#pragma once
+
+#include "CoreMinimal.h"
+
+#include "Framework/MultiBox/MultiBoxBuilder.h"
+#include "MVVMPropertyPath.h"
+
+class FBlueprintEditor;
+class FDragDropEvent;
+class UWidgetBlueprint;
+class UMVVMBlueprintView;
+struct FMVVMBlueprintViewBinding;
+class UMVVMBlueprintViewEvent;
+namespace UE::MVVM { struct FBindingEntry; }
+
+namespace UE::MVVM::BindingEntry
+{
+
+struct FRowHelper
+{
+	static void GatherAllChildBindings(UMVVMBlueprintView* BlueprintView, const TConstArrayView<TSharedPtr<FBindingEntry>> Entries, TArray<const FMVVMBlueprintViewBinding*>& OutBindings, TArray<UMVVMBlueprintViewEvent*>& OutEvents);
+
+	static void DeleteEntries(const UWidgetBlueprint* WidgetBlueprint, UMVVMBlueprintView* BlueprintView, TArrayView<const TSharedPtr<FBindingEntry>> Selection);
+	static void DuplicateBlueprintGraph(FBlueprintEditor* Editor, UWidgetBlueprint* WidgetBlueprint, UMVVMBlueprintView* BlueprintView, TArrayView<const TSharedPtr<FBindingEntry>> Selection);
+
+	static TOptional<FMVVMBlueprintPropertyPath> DropFieldSelector(const UWidgetBlueprint* WidgetBlueprint, const FDragDropEvent& DragDropEvent);
+	static void DragEnterFieldSelector(const UWidgetBlueprint* WidgetBlueprint, const FDragDropEvent& DragDropEvent);
+
+	static FMenuBuilder CreateContextMenu(UWidgetBlueprint* WidgetBlueprint, UMVVMBlueprintView* View, TConstArrayView<TSharedPtr<FBindingEntry>> Entries);
+};
+
+} // namespace UE::MVVM
diff --git a/Engine/Plugins/Runtime/ModelViewViewModel/Source/ModelViewViewModelEditor/Private/Styling/MVVMEditorStyle.cpp b/Engine/Plugins/Runtime/ModelViewViewModel/Source/ModelViewViewModelEditor/Private/Styling/MVVMEditorStyle.cpp
--- a/Engine/Plugins/Runtime/ModelViewViewModel/Source/ModelViewViewModelEditor/Private/Styling/MVVMEditorStyle.cpp
+++ b/Engine/Plugins/Runtime/ModelViewViewModel/Source/ModelViewViewModelEditor/Private/Styling/MVVMEditorStyle.cpp
@@ -30,6 +30,7 @@ FMVVMEditorStyle::FMVVMEditorStyle()
 	if (ensure(MVVMPlugin))
 	{
 		SetContentRoot(MVVMPlugin->GetContentDir() / TEXT("Editor"));
+		SetCoreContentRoot(FPaths::EngineContentDir() / TEXT("Slate"));
 	}
 
 	// Class Icons
@@ -64,6 +65,18 @@ FMVVMEditorStyle::FMVVMEditorStyle()
 		.SetOddRowBackgroundHoveredBrush(*FAppStyle::Get().GetBrush("Brushes.Header"))
 	);
 
+	FButtonStyle NoStyleComboButtonButtonStyle = FButtonStyle()
+		.SetNormal(FSlateBoxBrush(RootToContentDir("Common/ButtonHoverHint.png"), FMargin(4.0f / 16.0f), FLinearColor(1.0f, 1.0f, 1.0f, 0.15f)))
+		.SetHovered(FSlateBoxBrush(RootToContentDir("Common/ButtonHoverHint.png"), FMargin(4.0f / 16.0f), FLinearColor(1.0f, 1.0f, 1.0f, 0.25f)))
+		.SetPressed(FSlateBoxBrush(RootToContentDir("Common/ButtonHoverHint.png"), FMargin(4.0f / 16.0f), FLinearColor(1.0f, 1.0f, 1.0f, 0.30f)))
+		.SetNormalPadding(FMargin(0.0f, 0.0f, 0.0f, 1.0f))
+		.SetPressedPadding(FMargin(0.0f, 1.0f, 0.0f, 0.0f));
+	Set("NoStyleComboButton", FComboButtonStyle()
+		.SetButtonStyle(NoStyleComboButtonButtonStyle)
+		.SetDownArrowImage(FSlateNoResource())
+		.SetMenuBorderBrush(FSlateNoResource())
+		.SetMenuBorderPadding(FMargin(0.0f)));
+
 	// ViewModelSelectionWidget Icons
 	{
 		Set("ViewModelSelection.AddIcon", new IMAGE_BRUSH("Slate/MVVMViewModelSelection_Add_16x", Icon16x16));
@@ -90,6 +103,8 @@ FMVVMEditorStyle::FMVVMEditorStyle()
 	Set("ConversionFunction.DestToSource", new IMAGE_BRUSH_SVG("Slate/ConversionFunction_DestToSource", Icon16x16));
 	Set("ConversionFunction.SourceToDest", new IMAGE_BRUSH_SVG("Slate/ConversionFunction_SourceToDest", Icon16x16));
 
+	Set("Icon.Ellipsis", new CORE_IMAGE_BRUSH_SVG("Starship/Common/ellipsis-vertical-narrow", FVector2D(6.0, 24.0)));
+
 	Set("FieldSelector.ComboButton", 
 		FComboButtonStyle(FAppStyle::Get().GetWidgetStyle<FComboButtonStyle>("ComboButton"))
 		.SetDownArrowImage(*FCoreStyle::Get().GetBrush("Icons.Edit")));
diff --git a/Engine/Plugins/Runtime/ModelViewViewModel/Source/ModelViewViewModelEditor/Private/Widgets/BindingEntry/SMVVMBindingRow.cpp b/Engine/Plugins/Runtime/ModelViewViewModel/Source/ModelViewViewModelEditor/Private/Widgets/BindingEntry/SMVVMBindingRow.cpp
--- a/Engine/Plugins/Runtime/ModelViewViewModel/Source/ModelViewViewModelEditor/Private/Widgets/BindingEntry/SMVVMBindingRow.cpp
+++ b/Engine/Plugins/Runtime/ModelViewViewModel/Source/ModelViewViewModelEditor/Private/Widgets/BindingEntry/SMVVMBindingRow.cpp
@@ -4,14 +4,11 @@
 #include "Widgets/BindingEntry/SMVVMBindingRow.h"
 
 #include "Bindings/MVVMBindingHelper.h"
-#include "Blueprint/WidgetTree.h"
+#include "Framework/MVVMRowHelper.h"
 #include "MVVMBlueprintViewBinding.h"
-#include "MVVMBlueprintViewEvent.h"
 #include "MVVMDeveloperProjectSettings.h"
 #include "MVVMPropertyPath.h"
 
-#include "Details/WidgetPropertyDragDropOp.h"
-#include "DragAndDrop/DecoratedDragDropOp.h"
 #include "Styling/AppStyle.h"
 #include "Styling/MVVMEditorStyle.h"
 
@@ -24,7 +21,6 @@
 #include "Widgets/SBoxPanel.h"
 #include "Widgets/SMVVMFieldSelector.h"
 #include "Widgets/SMVVMSourceSelector.h"
-#include "Widgets/ViewModelFieldDragDropOp.h"
 #include "SSimpleButton.h"
 
 #define LOCTEXT_NAMESPACE "BindingListView_BindingRow"
@@ -223,6 +219,25 @@ TSharedRef<SWidget> SBindingRow::BuildRowWidget()
 				.ToolTipText(this, &SBindingRow::GetErrorButtonToolTip)
 				.OnClicked(this, &SBindingRow::OnErrorButtonClicked)
 			]
+
+			+ SHorizontalBox::Slot()
+			.Padding(2.0f, 0.0f)
+			.VAlign(VAlign_Center)
+			.HAlign(HAlign_Left)
+			.AutoWidth()
+			[
+				SNew(SComboButton)
+				.ContentPadding(FMargin(2.f, 0.0f))
+				.ComboButtonStyle(&FMVVMEditorStyle::Get().GetWidgetStyle<FComboButtonStyle>("NoStyleComboButton"))
+				.HasDownArrow(false)
+				.OnGetMenuContent(this, &SBindingRow::HandleContextMenu)
+				.ButtonContent()
+				[
+					SNew(SImage)
+					.Image(FMVVMEditorStyle::Get().GetBrush("Icon.Ellipsis"))
+					.DesiredSizeOverride(FVector2D(6.0, 24.0))
+				]
+			]
 		]
 	];
 }
@@ -457,125 +472,36 @@ FFieldSelectionContext SBindingRow::GetSelectedSelectionContext(bool bSource) co
 
 FReply SBindingRow::HandleFieldSelectorDrop(const FGeometry& MyGeometry, const FDragDropEvent& DragDropEvent, bool bSource)
 {
-	TSharedPtr<FDecoratedDragDropOp> DragDropOp = DragDropEvent.GetOperationAs<FDecoratedDragDropOp>();
-	if (!DragDropOp.IsValid())
+	FMVVMBlueprintViewBinding* ViewBinding = GetThisViewBinding();
+	UWidgetBlueprint* WidgetBlueprint = GetBlueprint();
+	if (ViewBinding == nullptr)
 	{
 		return FReply::Unhandled();
 	}
 
-	// Accept all drag-drop operations that are widget properties, but only accept view model fields when we are dropping into the Source box.
-	if (!DragDropOp->IsOfType<FWidgetPropertyDragDropOp>() && (!DragDropOp->IsOfType<FViewModelFieldDragDropOp>() || !bSource))
+	TOptional<FMVVMBlueprintPropertyPath> PropertyPath = BindingEntry::FRowHelper::DropFieldSelector(WidgetBlueprint, DragDropEvent);
+	if (!PropertyPath.IsSet())
 	{
 		return FReply::Unhandled();
 	}
 
-	FMVVMBlueprintViewBinding* ViewBinding = GetThisViewBinding();
-	UWidgetBlueprint* WidgetBlueprint = GetBlueprint();
-	if (ViewBinding == nullptr || WidgetBlueprint == nullptr)
+	UMVVMEditorSubsystem* Subsystem = GetEditorSubsystem();
+	if (bSource)
 	{
-		return FReply::Unhandled();
+		Subsystem->SetSourceToDestinationConversionFunction(WidgetBlueprint, *ViewBinding, nullptr);
+		Subsystem->SetSourcePathForBinding(WidgetBlueprint, *ViewBinding, PropertyPath.GetValue());
 	}
-
-	TSharedPtr<FViewModelFieldDragDropOp> ViewModelFieldDragDropOp = DragDropEvent.GetOperationAs<FViewModelFieldDragDropOp>();
-	TSharedPtr<FWidgetPropertyDragDropOp> WidgetPropertyDragDropOp = DragDropEvent.GetOperationAs<FWidgetPropertyDragDropOp>();
-	bool bIsViewModelProperty = false;
-
-	if (ViewModelFieldDragDropOp)
+	else
 	{
-		bIsViewModelProperty = true;
+		Subsystem->SetDestinationToSourceConversionFunction(WidgetBlueprint, *ViewBinding, nullptr);
+		Subsystem->SetDestinationPathForBinding(WidgetBlueprint, *ViewBinding, PropertyPath.GetValue());
 	}
-
-	UWidgetBlueprint* DragDropWidgetBP = bIsViewModelProperty ? ViewModelFieldDragDropOp->WidgetBP.Get() : WidgetPropertyDragDropOp->WidgetBP.Get();
-
-	if (WidgetBlueprint == DragDropWidgetBP)
-	{
-		TArray<FFieldVariant> FieldPath = bIsViewModelProperty ? ViewModelFieldDragDropOp->DraggedField : WidgetPropertyDragDropOp->DraggedPropertyPath;
-		FMVVMBlueprintPropertyPath PropertyPath;
-
-		PropertyPath.ResetPropertyPath();
-		for (const FFieldVariant& Field : FieldPath)
-		{
-			PropertyPath.AppendPropertyPath(WidgetBlueprint, FMVVMConstFieldVariant(Field));
-		}
-
-		UMVVMEditorSubsystem* Subsystem = GetEditorSubsystem();
-		if (bIsViewModelProperty)
-		{
-			if (ViewModelFieldDragDropOp->ViewModelId.IsValid())
-			{
-				PropertyPath.SetViewModelId(ViewModelFieldDragDropOp->ViewModelId);
-			}
-			else
-			{
-				return FReply::Unhandled();
-			}
-		}
-		else
-		{
-			if (UWidget* OwnerWidgetPtr = WidgetPropertyDragDropOp->OwnerWidget.Get())
-			{
-				// Search for the widget by its name in the widget tree
-				// If the widget is not found, we know it is the root preview widget so we use the blueprint name.
-				if (WidgetBlueprint->WidgetTree->FindWidget(OwnerWidgetPtr->GetFName()) && WidgetBlueprint->GetFName() != OwnerWidgetPtr->GetFName())
-				{
-					PropertyPath.SetWidgetName(OwnerWidgetPtr->GetFName());
-				}
-				else
-				{
-					PropertyPath.SetSelfContext();
-				}
-			}
-		}
-
-		if (bSource)
-		{
-			Subsystem->SetSourceToDestinationConversionFunction(WidgetBlueprint, *ViewBinding, nullptr);
-			Subsystem->SetSourcePathForBinding(WidgetBlueprint, *ViewBinding, PropertyPath);
-		}
-		else
-		{
-			Subsystem->SetDestinationToSourceConversionFunction(WidgetBlueprint, *ViewBinding, nullptr);
-			Subsystem->SetDestinationPathForBinding(WidgetBlueprint, *ViewBinding, PropertyPath);
-		}
-		return FReply::Handled();
-	}
-
-	return FReply::Unhandled();
+	return FReply::Handled();
 }
 
 void SBindingRow::HandleFieldSelectorDragEnter(const FGeometry& MyGeometry, const FDragDropEvent& DragDropEvent, bool bSource)
 {
-	TSharedPtr<FDecoratedDragDropOp> DragDropOp = DragDropEvent.GetOperationAs<FDecoratedDragDropOp>();
-	if (DragDropOp.IsValid())
-	{
-		// Accept all drag-drop operations that are widget properties, but only accept view model fields when we are dropping into the Source box.
-		if (DragDropOp->IsOfType<FWidgetPropertyDragDropOp>() || (DragDropOp->IsOfType<FViewModelFieldDragDropOp>() && bSource))
-		{
-			if (UWidgetBlueprint* WidgetBlueprint = GetBlueprint())
-			{
-				TSharedPtr<FViewModelFieldDragDropOp> ViewModelFieldDragDropOp = DragDropEvent.GetOperationAs<FViewModelFieldDragDropOp>();
-				TSharedPtr<FWidgetPropertyDragDropOp> WidgetPropertyDragDropOp = DragDropEvent.GetOperationAs<FWidgetPropertyDragDropOp>();
-				bool IsViewModelProperty = false;
-
-				if (ViewModelFieldDragDropOp)
-				{
-					IsViewModelProperty = true;
-				}
-
-				UWidgetBlueprint* DragDropWidgetBP = IsViewModelProperty ? ViewModelFieldDragDropOp->WidgetBP.Get() : WidgetPropertyDragDropOp->WidgetBP.Get();
-
-				if (DragDropWidgetBP && DragDropWidgetBP == WidgetBlueprint)
-				{
-					DragDropOp->CurrentIconBrush = FAppStyle::GetBrush(TEXT("Graph.ConnectorFeedback.OK"));
-					return;
-				}
-			}
-		}
-		else
-		{
-			DragDropOp->CurrentIconBrush = FAppStyle::GetBrush(TEXT("Graph.ConnectorFeedback.Error"));
-		}
-	}
+	BindingEntry::FRowHelper::DragEnterFieldSelector(GetBlueprint(), DragDropEvent);
 }
 
 ECheckBoxState SBindingRow::IsExecutionModeOverrideChecked() const
@@ -809,6 +735,43 @@ void SBindingRow::OnBindingModeSelectionChanged(FName ValueName, ESelectInfo::Ty
 	}
 }
 
+TSharedRef<SWidget> SBindingRow::HandleContextMenu() const
+{
+	TSharedPtr<FBindingEntry> Entry = GetEntry();
+	FMenuBuilder MenuBuilder = BindingEntry::FRowHelper::CreateContextMenu(GetBlueprint(), GetBlueprintView(), MakeArrayView(&Entry, 1));
+
+	{
+		MenuBuilder.BeginSection("Developer", LOCTEXT("Developer", "Developer"));
+
+		FMVVMBlueprintViewBinding* ViewBinding = GetThisViewBinding();
+		if (GetDefault<UMVVMDeveloperProjectSettings>()->bShowDeveloperGenerateGraphSettings)
+		{
+			bool bCanDuplicateGraph = ViewBinding && (ViewBinding->Conversion.GetConversionFunction(true) != nullptr || ViewBinding->Conversion.GetConversionFunction(false) != nullptr);
+			//CanDuplicateGraph(GraphToDuplicate)
+
+			FUIAction DuplicateAction;
+			DuplicateAction.ExecuteAction = FExecuteAction::CreateSP(this, &SBindingRow::HandleDuplicateGraph);
+			DuplicateAction.CanExecuteAction = FCanExecuteAction::CreateLambda([bCanDuplicateGraph]() { return bCanDuplicateGraph; });
+			MenuBuilder.AddMenuEntry(LOCTEXT("DuplicateGraph", "Copy binding graph")
+				, LOCTEXT("DuplicateGraphTooltip", "Add a copy of the binding graph to Blueprint."
+					" The copied graph will not bu used."
+					" The graph is always generated but may not be visible to the user.")
+				, FSlateIcon(FAppStyle::GetAppStyleSetName(), "Icons.Duplicate")
+				, DuplicateAction);
+		}
+
+		MenuBuilder.EndSection();
+	}
+
+	return MenuBuilder.MakeWidget();
+}
+
+void SBindingRow::HandleDuplicateGraph() const
+{
+	TSharedPtr<FBindingEntry> Entry = GetEntry();
+	BindingEntry::FRowHelper::DuplicateBlueprintGraph(GetBlueprintEditor().Get(), GetBlueprint(), GetBlueprintView(), MakeArrayView(&Entry, 1));
+}
+
 } // namespace
 
 #undef LOCTEXT_NAMESPACE
diff --git a/Engine/Plugins/Runtime/ModelViewViewModel/Source/ModelViewViewModelEditor/Private/Widgets/BindingEntry/SMVVMBindingRow.h b/Engine/Plugins/Runtime/ModelViewViewModel/Source/ModelViewViewModelEditor/Private/Widgets/BindingEntry/SMVVMBindingRow.h
--- a/Engine/Plugins/Runtime/ModelViewViewModel/Source/ModelViewViewModelEditor/Private/Widgets/BindingEntry/SMVVMBindingRow.h
+++ b/Engine/Plugins/Runtime/ModelViewViewModel/Source/ModelViewViewModelEditor/Private/Widgets/BindingEntry/SMVVMBindingRow.h
@@ -93,6 +93,9 @@ private:
 
 	void OnBindingModeSelectionChanged(FName ValueName, ESelectInfo::Type);
 
+	TSharedRef<SWidget> HandleContextMenu() const;
+	void HandleDuplicateGraph() const;
+
 private:
 	TArray<TSharedPtr<FText>> ErrorItems;
 	EMVVMExecutionMode DefaultExecutionMode = EMVVMExecutionMode::DelayedWhenSharedElseImmediate;
diff --git a/Engine/Plugins/Runtime/ModelViewViewModel/Source/ModelViewViewModelEditor/Private/Widgets/BindingEntry/SMVVMEventRow.cpp b/Engine/Plugins/Runtime/ModelViewViewModel/Source/ModelViewViewModelEditor/Private/Widgets/BindingEntry/SMVVMEventRow.cpp
--- a/Engine/Plugins/Runtime/ModelViewViewModel/Source/ModelViewViewModelEditor/Private/Widgets/BindingEntry/SMVVMEventRow.cpp
+++ b/Engine/Plugins/Runtime/ModelViewViewModel/Source/ModelViewViewModelEditor/Private/Widgets/BindingEntry/SMVVMEventRow.cpp
@@ -4,14 +4,11 @@
 #include "Widgets/BindingEntry/SMVVMEventRow.h"
 
 #include "Bindings/MVVMBindingHelper.h"
-#include "Blueprint/WidgetTree.h"
-#include "MVVMBlueprintViewBinding.h"
+#include "Framework/MVVMRowHelper.h"
 #include "MVVMBlueprintViewEvent.h"
 #include "MVVMDeveloperProjectSettings.h"
 #include "MVVMPropertyPath.h"
 
-#include "Details/WidgetPropertyDragDropOp.h"
-#include "DragAndDrop/DecoratedDragDropOp.h"
 #include "Styling/AppStyle.h"
 #include "Styling/MVVMEditorStyle.h"
 
@@ -24,7 +21,6 @@
 #include "Widgets/SBoxPanel.h"
 #include "Widgets/SMVVMFieldSelector.h"
 #include "Widgets/SMVVMSourceSelector.h"
-#include "Widgets/ViewModelFieldDragDropOp.h"
 #include "SSimpleButton.h"
 
 #define LOCTEXT_NAMESPACE "BindingListView_EventRow"
@@ -63,7 +59,7 @@ TSharedRef<SWidget> SEventRow::BuildRowWidget()
 	.BorderBackgroundColor(this, &SEventRow::GetErrorBorderColor)
 	[
 		SNew(SBox)
-		.HeightOverride(30)
+		.HeightOverride(30.0f)
 		[
 			SNew(SHorizontalBox)
 			+ SHorizontalBox::Slot()
@@ -77,6 +73,19 @@ TSharedRef<SWidget> SEventRow::BuildRowWidget()
 				.OnCheckStateChanged(this, &SEventRow::OnIsEventCompileChanged)
 			]
 
+			+ SHorizontalBox::Slot()
+			.Padding(0.0f)
+			.VAlign(VAlign_Center)
+			.HAlign(HAlign_Left)
+			.AutoWidth()
+			[
+				SNew(SSimpleButton)
+					.Icon(FAppStyle::Get().GetBrush("Icons.Error"))
+					.Visibility(this, &SEventRow::GetErrorButtonVisibility)
+					.ToolTipText(this, &SEventRow::GetErrorButtonToolTip)
+					.OnClicked(this, &SEventRow::OnErrorButtonClicked)
+			]
+
 			+ SHorizontalBox::Slot()
 			.Padding(4.0f, 0.0f)
 			.VAlign(VAlign_Center)
@@ -144,11 +153,17 @@ TSharedRef<SWidget> SEventRow::BuildRowWidget()
 			.HAlign(HAlign_Left)
 			.AutoWidth()
 			[
-				SNew(SSimpleButton)
-				.Icon(FAppStyle::Get().GetBrush("Icons.Error"))
-				.Visibility(this, &SEventRow::GetErrorButtonVisibility)
-				.ToolTipText(this, &SEventRow::GetErrorButtonToolTip)
-				.OnClicked(this, &SEventRow::OnErrorButtonClicked)
+				SNew(SComboButton)
+				.ContentPadding(FMargin(2.f, 0.0f))
+				.ComboButtonStyle(&FMVVMEditorStyle::Get().GetWidgetStyle<FComboButtonStyle>("NoStyleComboButton"))
+				.HasDownArrow(false)
+				.OnGetMenuContent(this, &SEventRow::HandleContextMenu)
+				.ButtonContent()
+				[
+					SNew(SImage)
+					.Image(FMVVMEditorStyle::Get().GetBrush("Icon.Ellipsis"))
+					.DesiredSizeOverride(FVector2D(6.0, 24.0))
+				]
 			]
 		]
 	];
@@ -360,115 +375,69 @@ FFieldSelectionContext SEventRow::GetSelectedSelectionContext(bool bEvent) const
 
 FReply SEventRow::HandleFieldSelectorDrop(const FGeometry& MyGeometry, const FDragDropEvent& DragDropEvent, bool bEvent)
 {
-	TSharedPtr<FDecoratedDragDropOp> DragDropOp = DragDropEvent.GetOperationAs<FDecoratedDragDropOp>();
-	if (!DragDropOp.IsValid())
-	{
-		return FReply::Unhandled();
-	}
-
-	// Accept all drag-drop operations that are widget properties, but only accept view model fields when we are dropping into the Source box.
-	if (!DragDropOp->IsOfType<FWidgetPropertyDragDropOp>() && (!DragDropOp->IsOfType<FViewModelFieldDragDropOp>() || !bEvent))
-	{
-		return FReply::Unhandled();
-	}
-
 	UMVVMBlueprintViewEvent* Event = GetEvent();
-	UWidgetBlueprint* WidgetBlueprintPtr = GetBlueprint();
-	if (Event == nullptr || WidgetBlueprintPtr == nullptr)
+	if (Event == nullptr)
 	{
 		return FReply::Unhandled();
 	}
 
-	TSharedPtr<FViewModelFieldDragDropOp> ViewModelFieldDragDropOp = DragDropEvent.GetOperationAs<FViewModelFieldDragDropOp>();
-	TSharedPtr<FWidgetPropertyDragDropOp> WidgetPropertyDragDropOp = DragDropEvent.GetOperationAs<FWidgetPropertyDragDropOp>();
-	bool bIsViewModelProperty = ViewModelFieldDragDropOp != nullptr;
-
-	UWidgetBlueprint* DragDropWidgetBP = bIsViewModelProperty ? ViewModelFieldDragDropOp->WidgetBP.Get() : WidgetPropertyDragDropOp->WidgetBP.Get();
-	if (WidgetBlueprintPtr != DragDropWidgetBP)
+	TOptional<FMVVMBlueprintPropertyPath> PropertyPath = BindingEntry::FRowHelper::DropFieldSelector(GetBlueprint(), DragDropEvent);
+	if (!PropertyPath.IsSet())
 	{
 		return FReply::Unhandled();
 	}
 
-	TArray<FFieldVariant> FieldPath = bIsViewModelProperty ? ViewModelFieldDragDropOp->DraggedField : WidgetPropertyDragDropOp->DraggedPropertyPath;
-	FMVVMBlueprintPropertyPath PropertyPath;
-	for (const FFieldVariant& Field : FieldPath)
-	{
-		PropertyPath.AppendPropertyPath(WidgetBlueprintPtr, FMVVMConstFieldVariant(Field));
-	}
-
-	UMVVMEditorSubsystem* Subsystem = GetEditorSubsystem();
-	if (bIsViewModelProperty)
-	{
-		if (ViewModelFieldDragDropOp->ViewModelId.IsValid())
-		{
-			PropertyPath.SetViewModelId(ViewModelFieldDragDropOp->ViewModelId);
-		}
-		else
-		{
-			return FReply::Unhandled();
-		}
-	}
-	else
-	{
-		if (UWidget* OwnerWidgetPtr = WidgetPropertyDragDropOp->OwnerWidget.Get())
-		{
-			// Search for the widget by its name in the widget tree
-			// If the widget is not found, we know it is the root preview widget so we use the blueprint name.
-			if (WidgetBlueprintPtr->WidgetTree->FindWidget(OwnerWidgetPtr->GetFName()) && WidgetBlueprintPtr->GetFName() != OwnerWidgetPtr->GetFName())
-			{
-				PropertyPath.SetWidgetName(OwnerWidgetPtr->GetFName());
-			}
-			else
-			{
-				PropertyPath.SetSelfContext();
-			}
-		}
-	}
-
 	if (bEvent)
 	{
-		Subsystem->SetEventPath(Event, PropertyPath);
+		GetEditorSubsystem()->SetEventPath(Event, PropertyPath.GetValue());
 	}
 	else
 	{
-		Subsystem->SetEventDestinationPath(Event, PropertyPath);
+		GetEditorSubsystem()->SetEventDestinationPath(Event, PropertyPath.GetValue());
 	}
 	return FReply::Handled();
 }
 
 void SEventRow::HandleFieldSelectorDragEnter(const FGeometry& MyGeometry, const FDragDropEvent& DragDropEvent, bool bEvent)
 {
-	TSharedPtr<FDecoratedDragDropOp> DragDropOp = DragDropEvent.GetOperationAs<FDecoratedDragDropOp>();
-	if (DragDropOp.IsValid())
+	BindingEntry::FRowHelper::DragEnterFieldSelector(GetBlueprint(), DragDropEvent);
+}
+
+TSharedRef<SWidget> SEventRow::HandleContextMenu() const
+{
+	TSharedPtr<FBindingEntry> Entry = GetEntry();
+	FMenuBuilder MenuBuilder = BindingEntry::FRowHelper::CreateContextMenu(GetBlueprint(), GetBlueprintView(), MakeArrayView(&Entry, 1));
+	
 	{
-		// Accept all drag-drop operations that are widget properties, but only accept view model fields when we are dropping into the Source box.
-		if (DragDropOp->IsOfType<FWidgetPropertyDragDropOp>() || (DragDropOp->IsOfType<FViewModelFieldDragDropOp>() && !bEvent))
-		{
-			if (UWidgetBlueprint* WidgetBlueprintPtr = GetBlueprint())
-			{
-				TSharedPtr<FViewModelFieldDragDropOp> ViewModelFieldDragDropOp = DragDropEvent.GetOperationAs<FViewModelFieldDragDropOp>();
-				TSharedPtr<FWidgetPropertyDragDropOp> WidgetPropertyDragDropOp = DragDropEvent.GetOperationAs<FWidgetPropertyDragDropOp>();
-				bool IsViewModelProperty = false;
-
-				if (ViewModelFieldDragDropOp)
-				{
-					IsViewModelProperty = true;
-				}
-
-				UWidgetBlueprint* DragDropWidgetBP = IsViewModelProperty ? ViewModelFieldDragDropOp->WidgetBP.Get() : WidgetPropertyDragDropOp->WidgetBP.Get();
-
-				if (DragDropWidgetBP && DragDropWidgetBP == WidgetBlueprintPtr)
-				{
-					DragDropOp->CurrentIconBrush = FAppStyle::GetBrush(TEXT("Graph.ConnectorFeedback.OK"));
-					return;
-				}
-			}
-		}
-		else
+		MenuBuilder.BeginSection("Developer", LOCTEXT("Developer", "Developer"));
+
+		UMVVMBlueprintViewEvent* Event = GetEvent();
+		if (GetDefault<UMVVMDeveloperProjectSettings>()->bShowDeveloperGenerateGraphSettings)		
 		{
-			DragDropOp->CurrentIconBrush = FAppStyle::GetBrush(TEXT("Graph.ConnectorFeedback.Error"));
+			bool bCanDuplicateGraph = Event && Event->GetWrapperGraph();
+			//CanDuplicateGraph(GraphToDuplicate)
+
+			FUIAction DuplicateAction;
+			DuplicateAction.ExecuteAction = FExecuteAction::CreateSP(this, &SEventRow::HandleDuplicateGraph);
+			DuplicateAction.CanExecuteAction = FCanExecuteAction::CreateLambda([bCanDuplicateGraph]() { return bCanDuplicateGraph; });
+			MenuBuilder.AddMenuEntry(LOCTEXT("DuplicateGraph", "Copy event graph")
+				, LOCTEXT("DuplicateGraphTooltip", "Add a copy of the event graph to Blueprint."
+					" The copied graph will not bu used."
+					" The graph is always generated but may not be visible to the user.")
+				, FSlateIcon(FAppStyle::GetAppStyleSetName(), "Icons.Duplicate")
+				, DuplicateAction);
 		}
+
+		MenuBuilder.EndSection();
 	}
+
+	return MenuBuilder.MakeWidget();
+}
+
+void SEventRow::HandleDuplicateGraph() const
+{
+	TSharedPtr<FBindingEntry> Entry = GetEntry();
+	BindingEntry::FRowHelper::DuplicateBlueprintGraph(GetBlueprintEditor().Get(), GetBlueprint(), GetBlueprintView(), MakeArrayView(&Entry, 1));
 }
 
 } // namespace
diff --git a/Engine/Plugins/Runtime/ModelViewViewModel/Source/ModelViewViewModelEditor/Private/Widgets/BindingEntry/SMVVMEventRow.h b/Engine/Plugins/Runtime/ModelViewViewModel/Source/ModelViewViewModelEditor/Private/Widgets/BindingEntry/SMVVMEventRow.h
--- a/Engine/Plugins/Runtime/ModelViewViewModel/Source/ModelViewViewModelEditor/Private/Widgets/BindingEntry/SMVVMEventRow.h
+++ b/Engine/Plugins/Runtime/ModelViewViewModel/Source/ModelViewViewModelEditor/Private/Widgets/BindingEntry/SMVVMEventRow.h
@@ -63,6 +63,9 @@ private:
 
 	void HandleFieldSelectorDragEnter(const FGeometry& MyGeometry, const FDragDropEvent& DragDropEvent, bool bSource);
 
+	TSharedRef<SWidget> HandleContextMenu() const;
+	void HandleDuplicateGraph() const;
+
 private:
 	TArray<TSharedPtr<FText>> ErrorItems;
 };
diff --git a/Engine/Plugins/Runtime/ModelViewViewModel/Source/ModelViewViewModelEditor/Private/Widgets/SMVVMViewBindingListView.cpp b/Engine/Plugins/Runtime/ModelViewViewModel/Source/ModelViewViewModelEditor/Private/Widgets/SMVVMViewBindingListView.cpp
--- a/Engine/Plugins/Runtime/ModelViewViewModel/Source/ModelViewViewModelEditor/Private/Widgets/SMVVMViewBindingListView.cpp
+++ b/Engine/Plugins/Runtime/ModelViewViewModel/Source/ModelViewViewModelEditor/Private/Widgets/SMVVMViewBindingListView.cpp
@@ -10,6 +10,7 @@
 #include "Misc/MessageDialog.h"
 #include "ScopedTransaction.h"
 
+#include "Framework/MVVMRowHelper.h"
 #include "MVVMBlueprintView.h"
 #include "MVVMBlueprintViewBinding.h"
 #include "MVVMBlueprintViewConversionFunction.h"
@@ -521,409 +522,13 @@ void SBindingsList::ClearFilterText()
 	FilterText = FText::GetEmpty();
 }
 
-namespace Private
-{
-	void GatherAllChildBindings(UMVVMBlueprintView* BlueprintView, const TConstArrayView<TSharedPtr<FBindingEntry>>& Entries, TArray<const FMVVMBlueprintViewBinding*>& OutBindings, TArray<UMVVMBlueprintViewEvent*> & OutEvents)
-	{
-		for (const TSharedPtr<FBindingEntry>& Entry : Entries)
-		{
-			if (Entry->GetRowType() == FBindingEntry::ERowType::Binding)
-			{
-				const FMVVMBlueprintViewBinding* Binding = Entry->GetBinding(BlueprintView);
-				if (Binding != nullptr)
-				{
-					OutBindings.AddUnique(Binding);
-				}
-			}
-
-			if (Entry->GetRowType() == FBindingEntry::ERowType::Event)
-			{
-				UMVVMBlueprintViewEvent* Event = Entry->GetEvent();
-				if (Event)
-				{
-					OutEvents.AddUnique(Event);
-				}
-			}
-			GatherAllChildBindings(BlueprintView, Entry->GetAllChildren(), OutBindings, OutEvents);
-		}
-	}
-} //namespace
-
-void SBindingsList::HandleDeleteSelected()
-{
-	TArray<TSharedPtr<FBindingEntry>> Selection = TreeView->GetSelectedItems();
-	if (Selection.Num() == 0)
-	{
-		return;
-	}
-
-	if (UMVVMWidgetBlueprintExtension_View* MVVMExtensionPtr = MVVMExtension.Get())
-	{
-		if (UMVVMBlueprintView* BlueprintView = MVVMExtensionPtr->GetBlueprintView())
-		{
-			const UWidgetBlueprint* WidgetBlueprint = MVVMExtensionPtr->GetWidgetBlueprint();
-
-			TArray<const FMVVMBlueprintViewBinding*> BindingsToRemove;
-			TArray<UMVVMBlueprintViewEvent*> EventsToRemove;
-			Private::GatherAllChildBindings(BlueprintView, Selection, BindingsToRemove, EventsToRemove);
-
-			if (BindingsToRemove.Num() == 0 && EventsToRemove.Num() == 0)
-			{
-				return;
-			}
-
-			TArray<FText> BindingDisplayNames;
-			for (const FMVVMBlueprintViewBinding* Binding : BindingsToRemove)
-			{
-				BindingDisplayNames.Add(FText::FromString(Binding->GetDisplayNameString(WidgetBlueprint)));
-			}
-			for (const UMVVMBlueprintViewEvent* Event : EventsToRemove)
-			{
-				BindingDisplayNames.Add(Event->GetDisplayName(true));
-			}
-
-			const FText Message = FText::Format(BindingDisplayNames.Num() == 1 ?
-				LOCTEXT("ConfirmDeleteSingle", "Are you sure that you want to delete this binding?\n\n{1}") :
-				LOCTEXT("ConfirmDeleteMultiple", "Are you sure that you want to delete these {0} bindings?\n\n{1}"), 
-			BindingDisplayNames.Num(), 
-			FText::Join(FText::FromString("\n"), BindingDisplayNames));
-
-			const FText Title = LOCTEXT("DeleteBindings", "Delete Bindings?");
-
-			if (FMessageDialog::Open(EAppMsgType::YesNo, EAppReturnType::Yes, Message, Title) == EAppReturnType::Yes)
-			{
-				FScopedTransaction Transaction(LOCTEXT("DeleteBindingsTransaction", "Delete Bindings"));
-
-				if (TSharedPtr<SBindingsPanel> BindingPanelPtr = BindingPanel.Pin())
-				{
-					BindingPanelPtr->OnBindingListSelectionChanged(TConstArrayView<FMVVMBlueprintViewBinding*>());
-				}
-
-				BlueprintView->Modify();
-
-				for (int32 BindingIndex = BindingsToRemove.Num() - 1; BindingIndex > -1; BindingIndex--)
-				{
-					const FMVVMBlueprintViewBinding* Binding = BindingsToRemove[BindingIndex];
-					BlueprintView->RemoveBinding(Binding);
-				}
-				for (UMVVMBlueprintViewEvent* Event : EventsToRemove)
-				{
-					BlueprintView->RemoveEvent(Event);
-				}
-			}
-		}
-	}
-}
-
-void SBindingsList::HandleResetSelectedPin()
-{
-	if (UMVVMWidgetBlueprintExtension_View* MVVMExtensionPtr = MVVMExtension.Get())
-	{
-		UWidgetBlueprint* WidgetBlueprint = MVVMExtensionPtr->GetWidgetBlueprint();
-		check(WidgetBlueprint);
-		UMVVMEditorSubsystem* EditorSubsystem = GEditor->GetEditorSubsystem<UMVVMEditorSubsystem>();
-
-		TArray<TSharedPtr<FBindingEntry>> Selection = TreeView->GetSelectedItems();
-		for (const TSharedPtr<FBindingEntry>& Entry : Selection)
-		{
-			if (Entry->GetRowType() == FBindingEntry::ERowType::EventParameter)
-			{
-				EditorSubsystem->ResetPinToDefaultValue(WidgetBlueprint, Entry->GetEvent(), Entry->GetEventParameterName());
-			}
-			else if (Entry->GetRowType() == FBindingEntry::ERowType::BindingParameter)
-			{
-				if (FMVVMBlueprintViewBinding* Binding = Entry->GetBinding(MVVMExtensionPtr->GetBlueprintView()))
-				{
-					const bool bSourceToDestination = UE::MVVM::IsForwardBinding(Binding->BindingType);
-					EditorSubsystem->ResetPinToDefaultValue(WidgetBlueprint, *Binding, Entry->GetBindingParameterName(), bSourceToDestination);
-				}
-			}
-		}
-	}
-}
-
-void SBindingsList::HandleBreakSelectedPin()
-{
-	if (UMVVMWidgetBlueprintExtension_View* MVVMExtensionPtr = MVVMExtension.Get())
-	{
-		UWidgetBlueprint* WidgetBlueprint = MVVMExtensionPtr->GetWidgetBlueprint();
-		check(WidgetBlueprint);
-		UMVVMEditorSubsystem* EditorSubsystem = GEditor->GetEditorSubsystem<UMVVMEditorSubsystem>();
-
-		TArray<TSharedPtr<FBindingEntry>> Selection = TreeView->GetSelectedItems();
-		for (const TSharedPtr<FBindingEntry>& Entry : Selection)
-		{
-			if (Entry->GetRowType() == FBindingEntry::ERowType::EventParameter)
-			{
-				EditorSubsystem->SplitPin(WidgetBlueprint, Entry->GetEvent(), Entry->GetEventParameterName());
-			}
-			else if(Entry->GetRowType() == FBindingEntry::ERowType::BindingParameter)
-			{
-				if (FMVVMBlueprintViewBinding* Binding = Entry->GetBinding(MVVMExtensionPtr->GetBlueprintView()))
-				{
-					const bool bSourceToDestination = UE::MVVM::IsForwardBinding(Binding->BindingType);
-					EditorSubsystem->SplitPin(WidgetBlueprint, *Binding, Entry->GetBindingParameterName(), bSourceToDestination);
-				}
-			}
-		}
-	}
-}
-
-void SBindingsList::HandleRecombineSelectedPin()
-{
-	if (UMVVMWidgetBlueprintExtension_View* MVVMExtensionPtr = MVVMExtension.Get())
-	{
-		UWidgetBlueprint* WidgetBlueprint = MVVMExtensionPtr->GetWidgetBlueprint();
-		check(WidgetBlueprint);
-		UMVVMEditorSubsystem* EditorSubsystem = GEditor->GetEditorSubsystem<UMVVMEditorSubsystem>();
-
-		TArray<TSharedPtr<FBindingEntry>> Selection = TreeView->GetSelectedItems();
-		for (const TSharedPtr<FBindingEntry>& Entry : Selection)
-		{
-			if (Entry->GetRowType() == FBindingEntry::ERowType::EventParameter)
-			{
-				EditorSubsystem->RecombinePin(WidgetBlueprint, Entry->GetEvent(), Entry->GetEventParameterName());
-			}
-			else if (Entry->GetRowType() == FBindingEntry::ERowType::BindingParameter)
-			{
-				if (FMVVMBlueprintViewBinding* Binding = Entry->GetBinding(MVVMExtensionPtr->GetBlueprintView()))
-				{
-					const bool bSourceToDestination = UE::MVVM::IsForwardBinding(Binding->BindingType);
-					EditorSubsystem->RecombinePin(WidgetBlueprint, *Binding, Entry->GetBindingParameterName(), bSourceToDestination);
-				}
-			}
-		}
-	}
-}
-
-void SBindingsList::HandleResetOrphanedSelectedPin()
-{
-	if (UMVVMWidgetBlueprintExtension_View* MVVMExtensionPtr = MVVMExtension.Get())
-	{
-		UWidgetBlueprint* WidgetBlueprint = MVVMExtensionPtr->GetWidgetBlueprint();
-		check(WidgetBlueprint);
-		UMVVMEditorSubsystem* EditorSubsystem = GEditor->GetEditorSubsystem<UMVVMEditorSubsystem>();
-
-		TArray<TSharedPtr<FBindingEntry>> Selection = TreeView->GetSelectedItems();
-		for (const TSharedPtr<FBindingEntry>& Entry : Selection)
-		{
-			if (Entry->GetRowType() == FBindingEntry::ERowType::EventParameter)
-			{
-				EditorSubsystem->ResetOrphanedPin(WidgetBlueprint, Entry->GetEvent(), Entry->GetEventParameterName());
-			}
-			else if (Entry->GetRowType() == FBindingEntry::ERowType::BindingParameter)
-			{
-				if (FMVVMBlueprintViewBinding* Binding = Entry->GetBinding(MVVMExtensionPtr->GetBlueprintView()))
-				{
-					const bool bSourceToDestination = UE::MVVM::IsForwardBinding(Binding->BindingType);
-					EditorSubsystem->ResetOrphanedPin(WidgetBlueprint, *Binding, Entry->GetBindingParameterName(), bSourceToDestination);
-				}
-			}
-		}
-	}
-}
-
 TSharedPtr<SWidget> SBindingsList::OnSourceConstructContextMenu()
 {
-	const bool bShouldCloseWindowAfterMenuSelection = true;
-	FMenuBuilder MenuBuilder(bShouldCloseWindowAfterMenuSelection, nullptr);
-
-	TArray<TSharedPtr<FBindingEntry>> Selection = TreeView->GetSelectedItems();
-	if (Selection.Num() > 0)
-	{
-		const UMVVMWidgetBlueprintExtension_View* MVVMExtensionPtr = MVVMExtension.Get();
-		const UWidgetBlueprint* WidgetBlueprint = MVVMExtensionPtr ? MVVMExtensionPtr->GetWidgetBlueprint() : nullptr;
-		const UMVVMBlueprintView* View = MVVMExtensionPtr ? MVVMExtensionPtr->GetBlueprintView() : nullptr;
-		const UMVVMEditorSubsystem* EditorSubsystem = GEditor->GetEditorSubsystem<UMVVMEditorSubsystem>();
-
-		if (View && WidgetBlueprint)
-		{
-			bool bCanRemoveEntry = false;
-			bool bCanSplitPin = false;
-			bool bCanRecombinePin = false;
-			bool bCanRecombinePinVisible = false;
-			bool bCanResetPin = false;
-			bool bCanResetPinVisible = false;
-			bool bCanResetOrphanedPin = false;
-			auto SetAll = [&](bool Value)
-			{
-				bCanRemoveEntry = Value;
-				bCanSplitPin = Value;
-				bCanRecombinePin = Value;
-				bCanRecombinePinVisible = Value;
-				bCanResetPin = Value;
-				bCanResetPinVisible = Value;
-				bCanResetOrphanedPin = Value;
-			};
-			auto AllFalse = [&]()
-			{
-				return !bCanRemoveEntry && !bCanRecombinePin && !bCanSplitPin && !bCanRecombinePinVisible && bCanResetPin && !bCanResetPinVisible && !bCanResetOrphanedPin;
-			};
-
-			if (Selection.Num())
-			{
-				SetAll(true);
-				for (const TSharedPtr<FBindingEntry>& Entry : Selection)
-				{
-					switch(Entry->GetRowType())
-					{
-						case FBindingEntry::ERowType::Group:
-						case FBindingEntry::ERowType::Binding:
-						case FBindingEntry::ERowType::Event:
-							break;
-						default:
-							bCanRemoveEntry = false;
-							break;
-					}
-
-					if (Entry->GetRowType() == FBindingEntry::ERowType::EventParameter)
-					{
-						if (!EditorSubsystem->CanSplitPin(WidgetBlueprint, Entry->GetEvent(), Entry->GetEventParameterName()))
-						{
-							bCanSplitPin = false;
-						}
-						if (!EditorSubsystem->CanRecombinePin(WidgetBlueprint, Entry->GetEvent(), Entry->GetEventParameterName()))
-						{
-							bCanRecombinePin = false;
-						}
-						if (!EditorSubsystem->CanResetPinToDefaultValue(WidgetBlueprint, Entry->GetEvent(), Entry->GetEventParameterName()))
-						{
-							bCanResetPin = false;
-						}
-						if (!EditorSubsystem->CanResetOrphanedPin(WidgetBlueprint, Entry->GetEvent(), Entry->GetEventParameterName()))
-						{
-							bCanResetOrphanedPin = false;
-						}
-
-						UMVVMBlueprintViewEvent* ViewEvent = Entry->GetEvent();
-						UEdGraphPin* GraphPin = ViewEvent ? ViewEvent->GetOrCreateGraphPin(Entry->GetEventParameterName()) : nullptr;
-						if (GraphPin == nullptr)
-						{
-							bCanRecombinePinVisible = false;
-							bCanResetPinVisible = false;
-						}
-						else
-						{
-							if (GraphPin->ParentPin == nullptr)
-							{
-								bCanRecombinePinVisible = false;
-							}
-							if (!FMVVMBlueprintPin::IsInputPin(GraphPin) || GetDefault<UEdGraphSchema_K2>()->ShouldHidePinDefaultValue(GraphPin))
-							{
-								bCanResetPinVisible = false;
-							}
-						}
-					}
-					else if (Entry->GetRowType() == FBindingEntry::ERowType::BindingParameter)
-					{
-						const FMVVMBlueprintViewBinding* Binding = Entry->GetBinding(View);
-						if (Binding)
-						{
-							const bool bSourceToDestination = UE::MVVM::IsForwardBinding(Binding->BindingType);
-							if (!EditorSubsystem->CanSplitPin(WidgetBlueprint, *Binding, Entry->GetBindingParameterName(), bSourceToDestination))
-							{
-								bCanSplitPin = false;
-							}
-							if (!EditorSubsystem->CanRecombinePin(WidgetBlueprint, *Binding, Entry->GetBindingParameterName(), bSourceToDestination))
-							{
-								bCanRecombinePin = false;
-							}
-							if (!EditorSubsystem->CanResetPinToDefaultValue(WidgetBlueprint, *Binding, Entry->GetBindingParameterName(), bSourceToDestination))
-							{
-								bCanResetPin = false;
-							}
-							if (!EditorSubsystem->CanResetOrphanedPin(WidgetBlueprint, *Binding, Entry->GetBindingParameterName(), bSourceToDestination))
-							{
-								bCanResetOrphanedPin = false;
-							}
-
-							UEdGraphPin* GraphPin = EditorSubsystem->GetConversionFunctionArgumentPin(WidgetBlueprint, *Binding, Entry->GetBindingParameterName(), bSourceToDestination);
-							if (GraphPin == nullptr)
-							{
-								bCanRecombinePinVisible = false;
-								bCanResetPinVisible = false;
-							}
-							else
-							{
-								if (GraphPin->ParentPin == nullptr)
-								{
-									bCanRecombinePinVisible = false;
-								}
-								if (!FMVVMBlueprintPin::IsInputPin(GraphPin) || GetDefault<UEdGraphSchema_K2>()->ShouldHidePinDefaultValue(GraphPin))
-								{
-									bCanResetPinVisible = false;
-								}
-							}
-						}
-						else
-						{
-							SetAll(false);
-						}
-					}
-					else
-					{
-						SetAll(false);
-					}
-
-					if (AllFalse())
-					{
-						break;
-					}
-				}
-			}
-
-			{
-				FUIAction RemoveAction;
-				RemoveAction.ExecuteAction = FExecuteAction::CreateSP(this, &SBindingsList::HandleDeleteSelected);
-				RemoveAction.CanExecuteAction = FCanExecuteAction::CreateLambda([bCanRemoveEntry]() { return bCanRemoveEntry; });
-				MenuBuilder.AddMenuEntry(LOCTEXT("RemoveBinding", "Remove"),
-					LOCTEXT("RemoveBindingTooltip", "Remove bindings or events."),
-					FSlateIcon(FAppStyle::GetAppStyleSetName(), "Icons.Delete"),
-					RemoveAction);
-			}
-
-			if (bCanResetPinVisible)
-			{
-				FUIAction ResetPinAction;
-				ResetPinAction.ExecuteAction = FExecuteAction::CreateSP(this, &SBindingsList::HandleResetSelectedPin);
-				ResetPinAction.CanExecuteAction = FCanExecuteAction::CreateLambda([bCanResetPin]() { return bCanResetPin; });
-				MenuBuilder.AddMenuEntry(LOCTEXT("ResetPin", "Reset to Default Value"),
-					LOCTEXT("ResetPinTooltip", "Reset value of this pin to the default"),
-					FSlateIcon(),
-					ResetPinAction);
-			}
-			if (bCanSplitPin)
-			{
-				FUIAction SplitPinAction;
-				SplitPinAction.ExecuteAction = FExecuteAction::CreateSP(this, &SBindingsList::HandleBreakSelectedPin);
-				MenuBuilder.AddMenuEntry(LOCTEXT("BreakPin", "Split Struct Pin"),
-					LOCTEXT("BreakPinTooltip", "Breaks a struct pin in to a separate pin per element."),
-					FSlateIcon(),
-					SplitPinAction);
-			}
-			if (bCanRecombinePinVisible)
-			{
-				FUIAction RecombinePinAction;
-				RecombinePinAction.ExecuteAction = FExecuteAction::CreateSP(this, &SBindingsList::HandleRecombineSelectedPin);
-				RecombinePinAction.CanExecuteAction = FCanExecuteAction::CreateLambda([bCanRecombinePin](){ return bCanRecombinePin; });
-				MenuBuilder.AddMenuEntry(LOCTEXT("RecombinePin", "Recombine Struct Pin"),
-					LOCTEXT("RecombinePinTooltip", "Takes struct pins that have been broken in to composite elements and combines them back to a single struct pin."),
-					FSlateIcon(),
-					RecombinePinAction);
-			}
-			if (bCanResetOrphanedPin)
-			{
-				FUIAction ResetOrphanedPinAction;
-				ResetOrphanedPinAction.ExecuteAction = FExecuteAction::CreateSP(this, &SBindingsList::HandleResetOrphanedSelectedPin);
-				MenuBuilder.AddMenuEntry(LOCTEXT("ResetOrphanedPin", "Remove the Orphaned Struct Pin"),
-					LOCTEXT("ResetOrphanedPinTooltip", "Removes pins that used to exist but do not exist anymore."),
-					FSlateIcon(),
-					ResetOrphanedPinAction);
-			}
-		}
-	}
+	UMVVMWidgetBlueprintExtension_View* MVVMExtensionPtr = MVVMExtension.Get();
+	UWidgetBlueprint* WidgetBlueprint = MVVMExtensionPtr ? MVVMExtensionPtr->GetWidgetBlueprint() : nullptr;
+	UMVVMBlueprintView* View = MVVMExtensionPtr ? MVVMExtensionPtr->GetBlueprintView() : nullptr;
 
-	return MenuBuilder.MakeWidget();
+	return BindingEntry::FRowHelper::CreateContextMenu(WidgetBlueprint, View, TreeView->GetSelectedItems()).MakeWidget();
 }
 
 void SBindingsList::RequestNavigateToBinding(FGuid BindingId)
@@ -948,7 +553,14 @@ FReply SBindingsList::OnKeyDown(const FGeometry& MyGeometry, const FKeyEvent& In
 {
 	if (InKeyEvent.GetKey() == EKeys::Delete)
 	{
-		HandleDeleteSelected();
+		if (UMVVMWidgetBlueprintExtension_View* MVVMExtensionPtr = MVVMExtension.Get())
+		{
+			UMVVMBlueprintView* BlueprintView = MVVMExtensionPtr->GetBlueprintView();
+			const UWidgetBlueprint* WidgetBlueprint = MVVMExtensionPtr->GetWidgetBlueprint();
+			TArray<TSharedPtr<FBindingEntry>> Selection = TreeView->GetSelectedItems();
+
+			BindingEntry::FRowHelper::DeleteEntries(WidgetBlueprint, BlueprintView, Selection);
+		}
 		return FReply::Handled();
 	}
 
diff --git a/Engine/Plugins/Runtime/ModelViewViewModel/Source/ModelViewViewModelEditor/Private/Widgets/SMVVMViewBindingListView.h b/Engine/Plugins/Runtime/ModelViewViewModel/Source/ModelViewViewModelEditor/Private/Widgets/SMVVMViewBindingListView.h
--- a/Engine/Plugins/Runtime/ModelViewViewModel/Source/ModelViewViewModelEditor/Private/Widgets/SMVVMViewBindingListView.h
+++ b/Engine/Plugins/Runtime/ModelViewViewModel/Source/ModelViewViewModelEditor/Private/Widgets/SMVVMViewBindingListView.h
@@ -46,11 +46,6 @@ private:
 	void GetChildrenOfEntry(TSharedPtr<FBindingEntry> Entry, TArray<TSharedPtr<FBindingEntry>>& OutChildren) const;
 
 	virtual FReply OnKeyDown(const FGeometry& MyGeometry, const FKeyEvent& InKeyEvent) override;
-	void HandleDeleteSelected();
-	void HandleResetSelectedPin();
-	void HandleBreakSelectedPin();
-	void HandleRecombineSelectedPin();
-	void HandleResetOrphanedSelectedPin();
 
 	void ForceRefresh();
 
