Minor fix to MVVM ClassBinding initialization. This has no impact on execution since the value is always set after the object is created, but it's safer this way.

diff --git a/Engine/Plugins/Runtime/ModelViewViewModel/Source/ModelViewViewModel/Public/View/MVVMViewClass.h b/Engine/Plugins/Runtime/ModelViewViewModel/Source/ModelViewViewModel/Public/View/MVVMViewClass.h
--- a/Engine/Plugins/Runtime/ModelViewViewModel/Source/ModelViewViewModel/Public/View/MVVMViewClass.h
+++ b/Engine/Plugins/Runtime/ModelViewViewModel/Source/ModelViewViewModel/Public/View/MVVMViewClass.h
@@ -158,7 +158,7 @@ private:
 	EMVVMExecutionMode ExecutionMode = EMVVMExecutionMode::Immediate;
 
 	UPROPERTY(VisibleAnywhere, Category = "View")
-	uint64 SourceBitField;
+	uint64 SourceBitField = 0;
 
 #if WITH_EDITORONLY_DATA
 	UPROPERTY(meta = (IgnoreForMemberInitializationTest))
