MVVM: allow viewmodel values to be edited on nested userwidgets.
Expose viewmodels whose creation type is "Create Instance" in the details panel of nested userwidgets and allow editing them independently.

#virtualized

diff --git a/Engine/Plugins/Runtime/ModelViewViewModel/Source/ModelViewViewModel/Private/View/MVVMViewClass.cpp b/Engine/Plugins/Runtime/ModelViewViewModel/Source/ModelViewViewModel/Private/View/MVVMViewClass.cpp
--- a/Engine/Plugins/Runtime/ModelViewViewModel/Source/ModelViewViewModel/Private/View/MVVMViewClass.cpp
+++ b/Engine/Plugins/Runtime/ModelViewViewModel/Source/ModelViewViewModel/Private/View/MVVMViewClass.cpp
@@ -38,11 +38,33 @@ UObject* FMVVMViewClass_Source::GetOrCreateInstance(const UMVVMViewClass* InView
 	UObject* Result = nullptr;
 	const bool bOptional = IsOptional();
 
+	auto GetInstanceObjectFromProperty = [&]() -> bool
+	{
+		TValueOrError<UE::MVVM::FFieldContext, void> FieldPathResult = InViewClass->GetBindingLibrary().EvaluateFieldPath(InUserWidget, FieldPath);
+		if (FieldPathResult.HasValue())
+		{
+			TValueOrError<UObject*, void> ObjectResult = UE::MVVM::FieldPathHelper::EvaluateObjectProperty(FieldPathResult.GetValue());
+			if (ObjectResult.HasValue() && ObjectResult.GetValue() != nullptr)
+			{
+				Result = ObjectResult.GetValue();
+			}
+		}
+
+		return Result != nullptr || bOptional;
+	};
+
 	if ((Flags & (uint16)EFlags::TypeCreateInstance) != 0)
 	{
 		if (ExpectedSourceType.Get() != nullptr)
 		{
-			Result = NewObject<UObject>(InUserWidget, ExpectedSourceType.Get(), NAME_None, RF_Transient);
+			if (ensure(FieldPath.IsValid()))
+			{
+				if (!GetInstanceObjectFromProperty())
+				{
+					UE::MVVM::FMessageLog Log(InUserWidget);
+					Log.Error(FText::Format(LOCTEXT("CreateInstanceInvalidStaticInstance", "The source '{0}' was not found for the static view model at initialization."), FText::FromName(PropertyName)));
+				}
+			}
 		}
 		else if (!bOptional)
 		{
@@ -106,20 +128,10 @@ UObject* FMVVMViewClass_Source::GetOrCreateInstance(const UMVVMViewClass* InView
 	}
 	else if (FieldPath.IsValid())
 	{
-		TValueOrError<UE::MVVM::FFieldContext, void> FieldPathResult = InViewClass->GetBindingLibrary().EvaluateFieldPath(InUserWidget, FieldPath);
-		if (FieldPathResult.HasValue())
-		{
-			TValueOrError<UObject*, void> ObjectResult = UE::MVVM::FieldPathHelper::EvaluateObjectProperty(FieldPathResult.GetValue());
-			if (ObjectResult.HasValue() && ObjectResult.GetValue() != nullptr)
-			{
-				Result = ObjectResult.GetValue();
-			}
-		}
-
-		if (Result == nullptr && !bOptional)
+		if (!GetInstanceObjectFromProperty())
 		{
 			UE::MVVM::FMessageLog Log(InUserWidget);
-			Log.Error(FText::Format(LOCTEXT("CreateInstanceInvalidBiding", "The source '{0}' was evaluated to be invalid at initialization."), FText::FromName(PropertyName)));
+			Log.Error(FText::Format(LOCTEXT("CreateInstanceInvalidBinding", "The source '{0}' was evaluated to be invalid at initialization."), FText::FromName(PropertyName)));
 		}
 	}
 
@@ -234,7 +246,7 @@ FString FMVVMViewClass_Source::ToString(const UMVVMViewClass* ViewClass, FToStri
 		StringBuilder << TEXT("\n    CreationType: ");
 		if ((Flags & (uint8)EFlags::TypeCreateInstance) != 0)
 		{
-			StringBuilder << TEXT("CreateInsance");
+			StringBuilder << TEXT("CreateInstance");
 		}
 		else if (Resolver)
 		{
diff --git a/Engine/Plugins/Runtime/ModelViewViewModel/Source/ModelViewViewModelBlueprint/Private/MVVMViewBlueprintCompiler.cpp b/Engine/Plugins/Runtime/ModelViewViewModel/Source/ModelViewViewModelBlueprint/Private/MVVMViewBlueprintCompiler.cpp
--- a/Engine/Plugins/Runtime/ModelViewViewModel/Source/ModelViewViewModelBlueprint/Private/MVVMViewBlueprintCompiler.cpp
+++ b/Engine/Plugins/Runtime/ModelViewViewModel/Source/ModelViewViewModelBlueprint/Private/MVVMViewBlueprintCompiler.cpp
@@ -335,12 +335,35 @@ void FMVVMViewBlueprintCompiler::CreateVariables(const FWidgetBlueprintCompilerC
 		FProperty* NewProperty = Context.CreateVariable(UserWidgetProperty.Name, NewPropertyPinType);
 		if (NewProperty != nullptr)
 		{
-			NewProperty->SetPropertyFlags(CPF_BlueprintVisible | CPF_RepSkip | CPF_Transient | CPF_DuplicateTransient);
+			NewProperty->SetPropertyFlags(CPF_BlueprintVisible | CPF_RepSkip);
+
+			if (UserWidgetProperty.bInstanced)
+			{
+				NewProperty->SetPropertyFlags(CPF_Edit | CPF_ExportObject | CPF_PersistentInstance | CPF_InstancedReference | CPF_NonNullable | CPF_NoClear);
+
+				if (UserWidgetProperty.AuthoritativeClass->HasAnyClassFlags(CLASS_HasInstancedReference))
+				{
+					NewProperty->SetPropertyFlags(CPF_ContainsInstancedReference);
+				}
+			}
+			else
+			{
+				NewProperty->SetPropertyFlags(CPF_Transient | CPF_DuplicateTransient);
+			}
+
+			if (UserWidgetProperty.bExposeOnSpawn)
+			{
+				NewProperty->SetPropertyFlags(CPF_ExposeOnSpawn);
+			}
+			else if (!UserWidgetProperty.bInstanced)
+			{
+				NewProperty->SetPropertyFlags(CPF_DisableEditOnInstance);
+			}
+
 			if (UserWidgetProperty.bReadOnly)
 			{
 				NewProperty->SetPropertyFlags(CPF_BlueprintReadOnly);
 			}
-			NewProperty->SetPropertyFlags(UserWidgetProperty.bExposeOnSpawn ? CPF_ExposeOnSpawn : CPF_DisableEditOnInstance);
 
 #if WITH_EDITOR
 			if (!UserWidgetProperty.BlueprintSetter.IsEmpty())
@@ -363,6 +386,10 @@ void FMVVMViewBlueprintCompiler::CreateVariables(const FWidgetBlueprintCompilerC
 			{
 				NewProperty->SetMetaData(FBlueprintMetadata::MD_Private, TEXT("true"));
 			}
+			if (UserWidgetProperty.bInstanced)
+			{
+				NewProperty->SetMetaData(FName("EditInline"), TEXT("true"));
+			}
 #endif
 		}
 		return NewProperty;
@@ -687,6 +714,7 @@ void FMVVMViewBlueprintCompiler::CreateRequiredProperties(const FWidgetBlueprint
 			SourceVariable.bExposeOnSpawn = bCreateSetterFunction;
 			SourceVariable.bPrivate = !bIsPublicReadable;
 			SourceVariable.bReadOnly = !bCreateSetterFunction;
+			SourceVariable.bInstanced = ViewModelContext.CreationType == EMVVMBlueprintViewModelContextCreationType::CreateInstance;
 		}
 
 		{
@@ -773,6 +801,7 @@ void FMVVMViewBlueprintCompiler::CreateRequiredProperties(const FWidgetBlueprint
 					//SourceVariable.bPrivate = true; Remove until the data is fix properly
 					SourceVariable.bPrivate = false;
 					SourceVariable.bReadOnly = true;
+					SourceVariable.bInstanced = false;
 
 					WidgetUserPropertyCreated.Add(Widget->GetFName());
 				}
@@ -1713,6 +1742,31 @@ void FMVVMViewBlueprintCompiler::PreCompileViewModelCreatorContexts(UWidgetBluep
 				bIsPreCompileStepValid = false;
 				continue;
 			}
+
+			FObjectPropertyBase* ViewModelProperty = FindFProperty<FObjectPropertyBase>(Class, ViewModelContext.GetViewModelName());
+			UObject* CDOInstance = Class->GetDefaultObject(true);
+
+			if (!ViewModelProperty->GetObjectPropertyValue_InContainer(CDOInstance))
+			{
+				UObject* StaticInstance = NewObject<UObject>(CDOInstance, ViewModelContext.GetViewModelClass(), NAME_None, RF_Transactional | RF_Public);
+				if (ensure(StaticInstance->GetClass()->IsChildOf(ViewModelProperty->PropertyClass)))
+				{
+					ViewModelProperty->SetObjectPropertyValue_InContainer(CDOInstance, StaticInstance);
+				}
+			}
+
+			TValueOrError<FCompiledBindingLibraryCompiler::FFieldPathHandle, FText> ReadFieldPathResult = AddObjectFieldPath(BindingLibraryCompiler, Class, ViewModelContext.GetViewModelName().ToString(), ViewModelContext.GetViewModelClass());
+			if (ReadFieldPathResult.HasError())
+			{
+				AddMessageForViewModel(ViewModelContext
+					, ReadFieldPathResult.GetError()
+					, EMessageType::Error
+				);
+				bIsPreCompileStepValid = false;
+				continue;
+			}
+
+			SourceCreatorContext.ReadPropertyPathHandle = ReadFieldPathResult.StealValue();
 		}
 		else if (ViewModelContext.CreationType == EMVVMBlueprintViewModelContextCreationType::PropertyPath)
 		{
@@ -1858,6 +1912,20 @@ void FMVVMViewBlueprintCompiler::CompileSources(const FCompiledBindingLibraryCom
 		else if (ViewModelContext.CreationType == EMVVMBlueprintViewModelContextCreationType::CreateInstance)
 		{
 			bCreateInstance = true;
+
+			const FMVVMVCompiledFieldPath* CompiledFieldPath = CompileResult.FieldPaths.Find(SourceCreatorContext.ReadPropertyPathHandle);
+			if (CompiledFieldPath == nullptr)
+			{
+				AddMessageForViewModel(ViewModelContext
+					, LOCTEXT("ViewModelInvalidInstanceNotFound", "The path for the created viewmodel instance was not found.")
+					, EMessageType::Error
+				);
+				bIsCompileStepValid = false;
+				continue;
+			}
+
+			CompiledSourceCreator.FieldPath = *CompiledFieldPath;
+			ensure(bIsOptional == false);
 		}
 		else if (ViewModelContext.CreationType == EMVVMBlueprintViewModelContextCreationType::PropertyPath)
 		{
diff --git a/Engine/Plugins/Runtime/ModelViewViewModel/Source/ModelViewViewModelBlueprint/Private/MVVMViewBlueprintCompiler.h b/Engine/Plugins/Runtime/ModelViewViewModel/Source/ModelViewViewModelBlueprint/Private/MVVMViewBlueprintCompiler.h
--- a/Engine/Plugins/Runtime/ModelViewViewModel/Source/ModelViewViewModelBlueprint/Private/MVVMViewBlueprintCompiler.h
+++ b/Engine/Plugins/Runtime/ModelViewViewModel/Source/ModelViewViewModelBlueprint/Private/MVVMViewBlueprintCompiler.h
@@ -151,6 +151,7 @@ private:
 		bool bExposeOnSpawn = false;
 		bool bPrivate = false;
 		bool bReadOnly = false;
+		bool bInstanced = false;
 
 		const FProperty* Property = nullptr;
 	};
