MVVM: Select the accessor class from the binding entry instead of calculating a new one. It may fail the verification if the property is not yet created and will be created later when compiling the widget.

diff --git a/Engine/Plugins/Runtime/ModelViewViewModel/Source/ModelViewViewModelEditor/Private/Widgets/SMVVMViewModelBindingListWidget.cpp b/Engine/Plugins/Runtime/ModelViewViewModel/Source/ModelViewViewModelEditor/Private/Widgets/SMVVMViewModelBindingListWidget.cpp
--- a/Engine/Plugins/Runtime/ModelViewViewModel/Source/ModelViewViewModelEditor/Private/Widgets/SMVVMViewModelBindingListWidget.cpp
+++ b/Engine/Plugins/Runtime/ModelViewViewModel/Source/ModelViewViewModelEditor/Private/Widgets/SMVVMViewModelBindingListWidget.cpp
@@ -495,19 +495,7 @@ FMVVMBlueprintPropertyPath SSourceBindingList::CreateBlueprintPropertyPath(SProp
 		// Backward, test if the object can be access.
 		//The last property can be a struct variable, inside a struct, inside..., inside an object. 
 		bool bPassFilter = false;
-		const UStruct* CurrentContainer = AccessorClass;
-
-		{
-			FMVVMBlueprintPropertyPath TempPropertyPath;
-			Source->Key.SetSourceTo(TempPropertyPath);
-			TArray<FMVVMConstFieldVariant> AllFields = TempPropertyPath.GetCompleteFields(WidgetBlueprintPtr);
-			if (AllFields.Num() > 0)
-			{
-				TValueOrError<const UStruct*, void> NewContainerResult = FieldPathHelper::GetFieldAsContainer(AllFields[0]);
-				CurrentContainer = NewContainerResult.HasValue() ? NewContainerResult.GetValue() : nullptr;
-			}
-		}
-
+		const UStruct* CurrentContainer = Source->Key.GetClass();
 		for (const FFieldVariant& FieldVariant : FieldPath)
 		{
 			FMVVMConstFieldVariant NewField;
