MVVM: Update the pin from a name to an array of name. That will allow the pin to be used in converion function inside conversion function.

diff --git a/Engine/Plugins/Runtime/ModelViewViewModel/Source/ModelViewViewModelBlueprint/Private/Bindings/MVVMConversionFunctionHelper.cpp b/Engine/Plugins/Runtime/ModelViewViewModel/Source/ModelViewViewModelBlueprint/Private/Bindings/MVVMConversionFunctionHelper.cpp
--- a/Engine/Plugins/Runtime/ModelViewViewModel/Source/ModelViewViewModelBlueprint/Private/Bindings/MVVMConversionFunctionHelper.cpp
+++ b/Engine/Plugins/Runtime/ModelViewViewModel/Source/ModelViewViewModelBlueprint/Private/Bindings/MVVMConversionFunctionHelper.cpp
@@ -41,13 +41,18 @@ namespace Private
 		return nullptr;
 	}
 
+	bool IsSystemInputPin(const UEdGraphPin* Pin)
+	{
+		return Pin && Pin->PinName != UEdGraphSchema_K2::PN_Execute && Pin->Direction == EGPD_Input && (!Pin->bOrphanedPin || Pin->ShouldSavePinIfOrphaned()) && !Pin->bHidden;
+	}
+
 	void MarkAsConversionFunction(const UK2Node* FunctionNode, const UEdGraph* Graph)
 	{
 		check(FunctionNode != nullptr);
 		FunctionNode->GetPackage()->GetMetaData()->SetValue(FunctionNode, ConversionFunctionMetadataKey.Resolve(), TEXT(""));
 	}
 
-	UK2Node_FunctionEntry* FindFunctionEntry(UEdGraph* Graph)
+	UK2Node_FunctionEntry* FindFunctionEntry(const UEdGraph* Graph)
 	{
 		for (UEdGraphNode* Node : Graph->Nodes)
 		{
@@ -58,8 +63,8 @@ namespace Private
 		}
 		return nullptr;
 	}
-	
-	UK2Node_FunctionResult* FindFunctionResult(UEdGraph* Graph)
+
+	UK2Node_FunctionResult* FindFunctionResult(const UEdGraph* Graph)
 	{
 		for (UEdGraphNode* Node : Graph->Nodes)
 		{
@@ -144,7 +149,7 @@ namespace Private
 		auto AddNode = [&NodesInPath](const UEdGraphPin* Pin)
 			{
 				UEdGraphNode* Result = nullptr;
-				if (Pin->Direction == EGPD_Input && Pin->LinkedTo.Num() == 1 && Pin->PinName != UEdGraphSchema_K2::PN_Execute)
+				if (Private::IsSystemInputPin(Pin) && Pin->LinkedTo.Num() == 1)
 				{
 					Result = Pin->LinkedTo[0]->GetOwningNode();
 					NodesInPath.Emplace(Result, Pin->LinkedTo[0]);
@@ -173,7 +178,7 @@ namespace Private
 
 	FMVVMBlueprintPropertyPath GetPropertyPathForPin(const UBlueprint* Blueprint, const UEdGraphPin* StartPin, bool bSkipResolve)
 	{
-		if (StartPin->Direction != EGPD_Input || StartPin->PinName == UEdGraphSchema_K2::PN_Self)
+		if (!IsInputPin(StartPin))
 		{
 			return FMVVMBlueprintPropertyPath();
 		}
@@ -309,6 +314,22 @@ namespace Private
 		return ResultPath;
 	}
 
+	UEdGraphPin* FindNewOutputPin(const UEdGraphNode* NewNode)
+	{
+		// then update our previous pin pointers
+		for (UEdGraphPin* Pin : NewNode->Pins)
+		{
+			if (Pin->Direction == EGPD_Output)
+			{
+				if (Pin->PinName != UEdGraphSchema_K2::PN_Then)
+				{
+					return Pin;
+				}
+			}
+		}
+		return nullptr;
+	};
+
 	TValueOrError<TArray<UEdGraphPin*>, void> BuildPropertyPath(const UBlueprint* Blueprint, UEdGraph* FunctionGraph, const FMVVMBlueprintPropertyPath& PropertyPath, int32 NumberOfFieldExcludingThePropertyPathSource, FVector2f EndLocation)
 	{
 		// Add new nodes
@@ -324,22 +345,6 @@ namespace Private
 					&& GetDefault<UEdGraphSchema_K2>()->ArePinsCompatible(PreviousDataPin, Pin, Context);
 			};
 
-		auto FindNewOutputPin = [](UEdGraphNode* NewNode) -> UEdGraphPin*
-			{
-				// then update our previous pin pointers
-				for (UEdGraphPin* Pin : NewNode->Pins)
-				{
-					if (Pin->Direction == EGPD_Output)
-					{
-						if (Pin->PinName != UEdGraphSchema_K2::PN_Then)
-						{
-							return Pin;
-						}
-					}
-				}
-				return nullptr;
-			};
-
 		NumberOfFieldExcludingThePropertyPathSource = FMath::Clamp(NumberOfFieldExcludingThePropertyPathSource, 0, PropertyPath.GetFieldPaths().Num());
 
 		const FVector2f LocationDelta = FVector2f(300.0f, 0.f);
@@ -691,7 +696,7 @@ namespace Private
 			ThenPin->MakeLinkTo(FunctionResultExecPin);
 		}
 	}
-}
+} //namespace
 
 bool RequiresWrapper(const UFunction* ConversionFunction)
 {
@@ -708,6 +713,11 @@ bool RequiresWrapper(const UFunction* ConversionFunction)
 	return false;
 }
 
+bool IsInputPin(const UEdGraphPin* Pin)
+{
+	return Private::IsSystemInputPin(Pin) && Pin->PinName != UEdGraphSchema_K2::PN_Self;
+}
+
 FName CreateWrapperName(const FMVVMBlueprintViewBinding& Binding, bool bSourceToDestination)
 {
 	TStringBuilder<256> StringBuilder;
@@ -1009,7 +1019,7 @@ TPair<UEdGraph*, UK2Node*> CreateGraph(UBlueprint* Blueprint, FName GraphName, c
 	return { NewGraph.FunctionGraph , CallFunctionNode };
 }
 
-UK2Node* GetWrapperNode(UEdGraph* Graph)
+UK2Node* GetWrapperNode(const UEdGraph* Graph)
 {
 	if (Graph == nullptr)
 	{
@@ -1108,39 +1118,69 @@ FMVVMBlueprintPropertyPath GetPropertyPathForArgument(const UBlueprint* WidgetBl
 	return GetPropertyPathForPin(WidgetBlueprint, ArgumentPin, bSkipResolve);
 }
 
-TMap<FName, FMVVMBlueprintPropertyPath> GetAllArgumentPropertyPaths(const UBlueprint* Blueprint, const UK2Node_CallFunction* FunctionNode, bool bSkipResolve)
+UEdGraphPin* FindPin(const UEdGraph* Graph, const TArrayView<const FName> PinNames)
 {
-	check(FunctionNode);
+	if (PinNames.Num() == 0 || Graph == nullptr)
+	{
+		return nullptr;
+	}
+
+	const UEdGraphNode* CurrentGraphNode = GetWrapperNode(Graph);
+	if (CurrentGraphNode == nullptr)
+	{
+		return nullptr;
+	}
 
-	TMap<FName, FMVVMBlueprintPropertyPath> Paths;
-	for (const UEdGraphPin* Pin : FunctionNode->GetAllPins())
+	for (int32 Index = 0; Index < PinNames.Num() - 1; ++Index)
 	{
-		FMVVMBlueprintPropertyPath Path = Private::GetPropertyPathForPin(Blueprint, Pin, bSkipResolve);
-		if (Path.IsValid())
+		FName PinName = PinNames[Index];
+		const UEdGraphPin* Pin = CurrentGraphNode->FindPin(PinName);
+		if (Pin == nullptr || Pin->LinkedTo.Num() != 1)
+		{
+			return nullptr;
+		}
+		CurrentGraphNode = Pin->LinkedTo[0]->GetOwningNode();
+		if (CurrentGraphNode == nullptr)
 		{
-			Paths.Add(Pin->PinName, Path);
+			return nullptr;
 		}
 	}
 
-	return Paths;
+	return CurrentGraphNode ? CurrentGraphNode->FindPin(PinNames.Last()) : nullptr;
 }
 
-TMap<FName, FMVVMBlueprintPropertyPath> GetAllArgumentPropertyPaths(const UBlueprint* Blueprint, const FMVVMBlueprintViewBinding& Binding, bool bSourceToDestination, bool bSkipResolve)
+TArray<FName> FindPinId(const UEdGraphPin* GraphPin)
 {
-	if (UMVVMBlueprintViewConversionFunction* ConversionFunction = Binding.Conversion.GetConversionFunction(bSourceToDestination))
+	if (GraphPin == nullptr)
+	{
+		return TArray<FName>();
+	}
+
+	const UEdGraphNode* ConversionFunctionNode = GetWrapperNode(GraphPin->GetOwningNode()->GetGraph());
+	if (ConversionFunctionNode == nullptr)
+	{
+		return TArray<FName>();
+	}
+
+	TArray<FName> Result;
+	while (GraphPin)
 	{
-		if (UEdGraph* ConversionFunctionGraph = ConversionFunction->GetWrapperGraph())
+		Result.Insert(GraphPin->GetFName(), 0);
+		const UEdGraphNode* CurrentGraphNode = GraphPin->GetOwningNode();
+		if (ConversionFunctionNode == CurrentGraphNode)
 		{
-			if (UK2Node_CallFunction* ConversionNode = Cast<UK2Node_CallFunction>(ConversionFunctionHelper::GetWrapperNode(ConversionFunctionGraph)))
-			{
-				return GetAllArgumentPropertyPaths(Blueprint, ConversionNode, bSkipResolve);
-			}
+			break;
 		}
+		const UEdGraphPin* OutputPin = Private::FindNewOutputPin(CurrentGraphNode);
+		if (OutputPin->LinkedTo.Num() != 1)
+		{
+			break;
+		}
+		GraphPin = OutputPin->LinkedTo[0];
 	}
-
-	return TMap<FName, FMVVMBlueprintPropertyPath>();
+	return Result;
 }
 
-} //namespace UE::MVVM
+} // UE::MVVM::ConversionFunctionHelper
 
 #undef LOCTEXT_NAMESPACE
diff --git a/Engine/Plugins/Runtime/ModelViewViewModel/Source/ModelViewViewModelBlueprint/Private/MVVMBlueprintPin.cpp b/Engine/Plugins/Runtime/ModelViewViewModel/Source/ModelViewViewModelBlueprint/Private/MVVMBlueprintPin.cpp
--- a/Engine/Plugins/Runtime/ModelViewViewModel/Source/ModelViewViewModelBlueprint/Private/MVVMBlueprintPin.cpp
+++ b/Engine/Plugins/Runtime/ModelViewViewModel/Source/ModelViewViewModelBlueprint/Private/MVVMBlueprintPin.cpp
@@ -2,21 +2,111 @@
 
 #include "MVVMBlueprintPin.h"
 
+#include "Algo/Reverse.h"
 #include "Bindings/MVVMConversionFunctionHelper.h"
 #include "EdGraph/EdGraphPin.h"
 #include "EdGraphSchema_K2.h"
 
 #include UE_INLINE_GENERATED_CPP_BY_NAME(MVVMBlueprintPin)
 
+/**
+ *
+ */
+FMVVMBlueprintPinId::FMVVMBlueprintPinId(const TArrayView<const FName> Names)
+	: PinNames(Names)
+{
+}
+
+FMVVMBlueprintPinId::FMVVMBlueprintPinId(TArray<FName>&& Names)
+	: PinNames(Names)
+{
+}
+
+bool FMVVMBlueprintPinId::IsValid() const
+{
+	return PinNames.Num() > 0 && !PinNames.Contains(FName());
+}
+
+bool FMVVMBlueprintPinId::IsChildOf(const FMVVMBlueprintPinId& Other) const
+{
+	if (Other.PinNames.Num() >= PinNames.Num())
+	{
+		return false;
+	}
+
+	for (int32 Index = 0; Index < Other.PinNames.Num(); ++Index)
+	{
+		if (Other.PinNames[Index] != PinNames[Index])
+		{
+			return false;
+		}
+	}
+	return true;
+}
+
+bool FMVVMBlueprintPinId::IsDirectChildOf(const FMVVMBlueprintPinId& Other) const
+{
+	if (Other.PinNames.Num() != PinNames.Num() - 1)
+	{
+		return false;
+	}
+
+	for (int32 Index = 0; Index < Other.PinNames.Num(); ++Index)
+	{
+		if (Other.PinNames[Index] != PinNames[Index])
+		{
+			return false;
+		}
+	}
+	return true;
+}
+
+bool FMVVMBlueprintPinId::operator==(const FMVVMBlueprintPinId& Other) const
+{
+	return PinNames == Other.PinNames;
+}
+
+bool FMVVMBlueprintPinId::operator==(const TArrayView<const FName> Other) const
+{
+	return PinNames.Num() == Other.Num() && CompareItems(GetData(Other), GetData(PinNames), PinNames.Num());
+}
+
+FString FMVVMBlueprintPinId::ToString() const
+{
+	TStringBuilder<256> Builder;
+	for (FName Name : PinNames)
+	{
+		if (Builder.Len() > 0)
+		{
+			Builder << TEXT('.');
+		}
+		Builder << Name;
+	}
+	return Builder.ToString();
+}
+
+/**
+ *
+ */
 FMVVMBlueprintPin::FMVVMBlueprintPin(FName InPinName)
-	: PinName(InPinName)
+	: PinId(MakeArrayView(&InPinName, 1))
+{
+}
+
+FMVVMBlueprintPin::FMVVMBlueprintPin(FMVVMBlueprintPinId InPinId)
+	: PinId(MoveTemp(InPinId))
+{
+}
+
+FMVVMBlueprintPin::FMVVMBlueprintPin(const TArrayView<const FName> InPinNames)
+	: PinId(InPinNames)
 {
 }
 
 FMVVMBlueprintPin FMVVMBlueprintPin::CreateFromPin(const UBlueprint* Blueprint, const UEdGraphPin* Pin)
 {
 	FMVVMBlueprintPin Result;
-	Result.PinName = Pin->PinName;
+	Result.PinId = FMVVMBlueprintPinId(UE::MVVM::ConversionFunctionHelper::FindPinId(Pin));
 	Result.Path = UE::MVVM::ConversionFunctionHelper::GetPropertyPathForPin(Blueprint, Pin, true);
 	Result.DefaultObject = Pin->DefaultObject;
 	Result.DefaultString = Pin->DefaultValue;
@@ -74,7 +164,7 @@ FString FMVVMBlueprintPin::GetValueAsString(const UClass* SelfContext) const
 
 bool FMVVMBlueprintPin::IsInputPin(const UEdGraphPin* Pin)
 {
-	return Pin->PinName != UEdGraphSchema_K2::PN_Self && Pin->PinName != UEdGraphSchema_K2::PN_Execute && Pin->Direction == EGPD_Input && (!Pin->bOrphanedPin || Pin->ShouldSavePinIfOrphaned()) && !Pin->bHidden;
+	return UE::MVVM::ConversionFunctionHelper::IsInputPin(Pin);
 }
 
 TArray<FMVVMBlueprintPin> FMVVMBlueprintPin::CopyAndReturnMissingPins(UBlueprint* Blueprint, UEdGraphNode* GraphNode, const TArray<FMVVMBlueprintPin>& Pins)
@@ -87,7 +177,7 @@ TArray<FMVVMBlueprintPin> FMVVMBlueprintPin::CopyAndReturnMissingPins(UBlueprint
 	for (const FMVVMBlueprintPin& Pin : Pins)
 	{
 		Pin.CopyTo(Blueprint, GraphNode);
-		if (UEdGraphPin* GraphPin = Pin.FindGraphPin(GraphNode))
+		if (UEdGraphPin* GraphPin = Pin.FindGraphPin(GraphNode->GetGraph()))
 		{
 			AllGraphPins.Add(GraphPin);
 		}
@@ -108,7 +198,7 @@ TArray<FMVVMBlueprintPin> FMVVMBlueprintPin::CopyAndReturnMissingPins(UBlueprint
 void FMVVMBlueprintPin::CopyTo(const UBlueprint* Blueprint, UEdGraphNode* Node) const
 {
 	Status = EMVVMBlueprintPinStatus::Orphaned;
-	if (UEdGraphPin* GraphPin = FindGraphPin(Node))
+	if (UEdGraphPin* GraphPin = FindGraphPin(Node->GetGraph()))
 	{
 		if (IsInputPin(GraphPin))
 		{
@@ -140,9 +230,9 @@ TArray<FMVVMBlueprintPin> FMVVMBlueprintPin::CreateFromNode(UBlueprint* Blueprin
 	return Result;
 }
 
-UEdGraphPin* FMVVMBlueprintPin::FindGraphPin(UEdGraphNode* Node) const
+UEdGraphPin* FMVVMBlueprintPin::FindGraphPin(const UEdGraph* Graph) const
 {
-	return Node->FindPin(PinName, EGPD_Input);
+	return UE::MVVM::ConversionFunctionHelper::FindPin(Graph, PinId.GetNames());
 }
 
 void FMVVMBlueprintPin::Reset()
@@ -154,3 +244,15 @@ void FMVVMBlueprintPin::Reset()
 	bSplit = false;
 	Status = EMVVMBlueprintPinStatus::Valid;
 }
+
+void FMVVMBlueprintPin::PostSerialize(const FArchive& Ar)
+{
+	if (Ar.IsLoading())
+	{
+		if (PinName_DEPRECATED.IsValid())
+		{
+			PinId = FMVVMBlueprintPinId(MakeArrayView(&PinName_DEPRECATED, 1));
+			PinName_DEPRECATED = FName();
+		}
+	}
+}
diff --git a/Engine/Plugins/Runtime/ModelViewViewModel/Source/ModelViewViewModelBlueprint/Private/MVVMBlueprintView.cpp b/Engine/Plugins/Runtime/ModelViewViewModel/Source/ModelViewViewModelBlueprint/Private/MVVMBlueprintView.cpp
--- a/Engine/Plugins/Runtime/ModelViewViewModel/Source/ModelViewViewModelBlueprint/Private/MVVMBlueprintView.cpp
+++ b/Engine/Plugins/Runtime/ModelViewViewModel/Source/ModelViewViewModelBlueprint/Private/MVVMBlueprintView.cpp
@@ -353,7 +353,7 @@ namespace UE::MVVM::Private
 					{
 						if (bGenerateGraph)
 						{
-							Binding.Conversion.DestinationToSourceConversion->SetGraphPin(WidgetBlueprint, Pin.GetName(), NewPath.GetValue());
+							Binding.Conversion.DestinationToSourceConversion->SetGraphPin(WidgetBlueprint, Pin.GetId(), NewPath.GetValue());
 						}
 						else
 						{
@@ -372,7 +372,7 @@ namespace UE::MVVM::Private
 					{
 						if (bGenerateGraph)
 						{
-							Binding.Conversion.SourceToDestinationConversion->SetGraphPin(WidgetBlueprint, Pin.GetName(), NewPath.GetValue());
+							Binding.Conversion.SourceToDestinationConversion->SetGraphPin(WidgetBlueprint, Pin.GetId(), NewPath.GetValue());
 						}
 						else
 						{
@@ -410,18 +410,18 @@ namespace UE::MVVM::Private
 					}
 				};
 
-				TArray<TTuple<FName, FMVVMBlueprintPropertyPath>> NewPins;
+				TArray<TTuple<FMVVMBlueprintPinId, FMVVMBlueprintPropertyPath>> NewPins;
 				for (const FMVVMBlueprintPin& Pin : Event->GetPins())
 				{
 					TOptional<FMVVMBlueprintPropertyPath> NewPath = PredPin(Pin);
 					if (NewPath.IsSet())
 					{
-						NewPins.Emplace(Pin.GetName(), MoveTemp(NewPath.GetValue()));
+						NewPins.Emplace(Pin.GetId(), MoveTemp(NewPath.GetValue()));
 					}
 				}
 				PredEventPath(Event->GetEventPath(), true);
 				PredEventPath(Event->GetDestinationPath(), false);
-				for (TTuple<FName, FMVVMBlueprintPropertyPath>& Pin : NewPins)
+				for (TTuple<FMVVMBlueprintPinId, FMVVMBlueprintPropertyPath>& Pin : NewPins)
 				{
 					if (bGenerateGraph)
 					{
diff --git a/Engine/Plugins/Runtime/ModelViewViewModel/Source/ModelViewViewModelBlueprint/Private/MVVMBlueprintViewConversionFunction.cpp b/Engine/Plugins/Runtime/ModelViewViewModel/Source/ModelViewViewModelBlueprint/Private/MVVMBlueprintViewConversionFunction.cpp
--- a/Engine/Plugins/Runtime/ModelViewViewModel/Source/ModelViewViewModelBlueprint/Private/MVVMBlueprintViewConversionFunction.cpp
+++ b/Engine/Plugins/Runtime/ModelViewViewModel/Source/ModelViewViewModelBlueprint/Private/MVVMBlueprintViewConversionFunction.cpp
@@ -234,7 +234,8 @@ void UMVVMBlueprintViewConversionFunction::Deprecation_InitializeFromMemberRefer
 	const FProperty* PinProperty = ConversionFunction ? UE::MVVM::BindingHelper::GetFirstArgumentProperty(ConversionFunction) : nullptr;
 	if (PinProperty)
 	{
-		FMVVMBlueprintPin NewPin = PinProperty->GetFName();
+		FName PinPropertyName = PinProperty->GetFName();
+		FMVVMBlueprintPin NewPin = FMVVMBlueprintPin(FMVVMBlueprintPinId(MakeArrayView(&PinPropertyName, 1)));
 		NewPin.SetPath(Source);
 		SavedPins.Add(MoveTemp(NewPin));
 	}
@@ -257,7 +258,8 @@ void UMVVMBlueprintViewConversionFunction::Deprecation_SetWrapperGraphName(UBlue
 		const FProperty* PinProperty = ConversionFunction ? UE::MVVM::BindingHelper::GetFirstArgumentProperty(ConversionFunction) : nullptr;
 		if (PinProperty)
 		{
-			FMVVMBlueprintPin NewPin = PinProperty->GetFName();
+			FName PinPropertyName = PinProperty->GetFName();
+			FMVVMBlueprintPin NewPin  = FMVVMBlueprintPin(FMVVMBlueprintPinId(MakeArrayView(&PinPropertyName, 1)));
 			NewPin.SetPath(Source);
 			SavedPins.Add(MoveTemp(NewPin));
 		}
@@ -372,14 +374,10 @@ UEdGraph* UMVVMBlueprintViewConversionFunction::GetOrCreateWrapperGraph(UBluepri
 	return nullptr;
 }
 
-UEdGraphPin* UMVVMBlueprintViewConversionFunction::GetOrCreateGraphPin(UBlueprint* Blueprint, FName PinName)
+UEdGraphPin* UMVVMBlueprintViewConversionFunction::GetOrCreateGraphPin(UBlueprint* Blueprint, const FMVVMBlueprintPinId& PinId)
 {
 	GetOrCreateWrapperGraph(Blueprint);
-	if (CachedWrapperNode)
-	{
-		return CachedWrapperNode->FindPin(PinName);
-	}
-	return nullptr;
+	return CachedWrapperGraph ? UE::MVVM::ConversionFunctionHelper::FindPin(CachedWrapperGraph, PinId.GetNames()) : nullptr;
 }
 
 UEdGraph* UMVVMBlueprintViewConversionFunction::GetOrCreateWrapperGraphInternal(FKismetCompilerContext& Context, const UFunction* Function)
@@ -408,13 +406,13 @@ void UMVVMBlueprintViewConversionFunction::RemoveWrapperGraph(UBlueprint* Bluepr
 	SetCachedWrapperGraph(Blueprint, nullptr, nullptr);
 }
 
-void UMVVMBlueprintViewConversionFunction::SetGraphPin(UBlueprint* Blueprint, FName PinName, const FMVVMBlueprintPropertyPath& Path)
+void UMVVMBlueprintViewConversionFunction::SetGraphPin(UBlueprint* Blueprint, const FMVVMBlueprintPinId& PinId, const FMVVMBlueprintPropertyPath& Path)
 {
-	UEdGraphPin* GraphPin = GetOrCreateGraphPin(Blueprint, PinName);
+	UEdGraphPin* GraphPin = GetOrCreateGraphPin(Blueprint, PinId);
 
 	// Set the value and make the blueprint as dirty before creating the pin.
 	//A property may not be created yet and the skeletal needs to be recreated.
-	FMVVMBlueprintPin* Pin = SavedPins.FindByPredicate([PinName](const FMVVMBlueprintPin& Other) { return PinName == Other.GetName(); });
+	FMVVMBlueprintPin* Pin = SavedPins.FindByPredicate([&PinId](const FMVVMBlueprintPin& Other) { return PinId == Other.GetId(); });
 	if (!Pin)
 	{
 		Pin = &SavedPins.Add_GetRef(FMVVMBlueprintPin::CreateFromPin(Blueprint, GraphPin));
diff --git a/Engine/Plugins/Runtime/ModelViewViewModel/Source/ModelViewViewModelBlueprint/Private/MVVMBlueprintViewEvent.cpp b/Engine/Plugins/Runtime/ModelViewViewModel/Source/ModelViewViewModelBlueprint/Private/MVVMBlueprintViewEvent.cpp
--- a/Engine/Plugins/Runtime/ModelViewViewModel/Source/ModelViewViewModelBlueprint/Private/MVVMBlueprintViewEvent.cpp
+++ b/Engine/Plugins/Runtime/ModelViewViewModel/Source/ModelViewViewModelBlueprint/Private/MVVMBlueprintViewEvent.cpp
@@ -113,14 +113,10 @@ void UMVVMBlueprintViewEvent::RemoveWrapperGraph()
 	SavedPins.Empty();
 }
 
-UEdGraphPin* UMVVMBlueprintViewEvent::GetOrCreateGraphPin(FName PinName)
+UEdGraphPin* UMVVMBlueprintViewEvent::GetOrCreateGraphPin(const FMVVMBlueprintPinId& PinId)
 {
 	GetOrCreateWrapperGraph();
-	if (CachedWrapperNode)
-	{
-		return CachedWrapperNode->FindPin(PinName);
-	}
-	return nullptr;
+	return CachedWrapperGraph ? UE::MVVM::ConversionFunctionHelper::FindPin(CachedWrapperGraph, PinId.GetNames()) : nullptr;
 }
 
 void UMVVMBlueprintViewEvent::SavePinValues()
@@ -156,21 +152,21 @@ bool UMVVMBlueprintViewEvent::HasOrphanedPin() const
 	return false;
 }
 
-FMVVMBlueprintPropertyPath UMVVMBlueprintViewEvent::GetPinPath(FName PinName) const
+FMVVMBlueprintPropertyPath UMVVMBlueprintViewEvent::GetPinPath(const FMVVMBlueprintPinId& PinId) const
 {
-	const FMVVMBlueprintPin* ViewPin = SavedPins.FindByPredicate([PinName](const FMVVMBlueprintPin& Other) { return PinName == Other.GetName(); });
+	const FMVVMBlueprintPin* ViewPin = SavedPins.FindByPredicate([&PinId](const FMVVMBlueprintPin& Other) { return PinId == Other.GetId(); });
 	return ViewPin ? ViewPin->GetPath() : FMVVMBlueprintPropertyPath();
 }
 
-void UMVVMBlueprintViewEvent::SetPinPath(FName PinName, const FMVVMBlueprintPropertyPath& Path)
+void UMVVMBlueprintViewEvent::SetPinPath(const FMVVMBlueprintPinId& PinId, const FMVVMBlueprintPropertyPath& Path)
 {
-	UEdGraphPin* GraphPin = GetOrCreateGraphPin(PinName);
+	UEdGraphPin* GraphPin = GetOrCreateGraphPin(PinId);
 
 	if (GraphPin)
 	{
 		UBlueprint* Blueprint = GetWidgetBlueprintInternal();
 		// Set the value and make the blueprint as dirty before creating the pin.
-		FMVVMBlueprintPin* ViewPin = SavedPins.FindByPredicate([PinName](const FMVVMBlueprintPin& Other) { return PinName == Other.GetName(); });
+		FMVVMBlueprintPin* ViewPin = SavedPins.FindByPredicate([&PinId](const FMVVMBlueprintPin& Other) { return PinId == Other.GetId(); });
 		if (!ViewPin)
 		{
 			ViewPin = &SavedPins.Add_GetRef(FMVVMBlueprintPin::CreateFromPin(Blueprint, GraphPin));
@@ -186,12 +182,12 @@ void UMVVMBlueprintViewEvent::SetPinPath(FName PinName, const FMVVMBlueprintProp
 	}
 }
 
-void UMVVMBlueprintViewEvent::SetPinPathNoGraphGeneration(FName PinName, const FMVVMBlueprintPropertyPath& Path)
+void UMVVMBlueprintViewEvent::SetPinPathNoGraphGeneration(const FMVVMBlueprintPinId& PinId, const FMVVMBlueprintPropertyPath& Path)
 {
-	FMVVMBlueprintPin* ViewPin = SavedPins.FindByPredicate([PinName](const FMVVMBlueprintPin& Other) { return PinName == Other.GetName(); });
+	FMVVMBlueprintPin* ViewPin = SavedPins.FindByPredicate([&PinId](const FMVVMBlueprintPin& Other) { return PinId == Other.GetId(); });
 	if (!ViewPin)
 	{
-		ViewPin = &SavedPins.Emplace_GetRef(PinName);
+		ViewPin = &SavedPins.Emplace_GetRef(PinId);
 		ViewPin->SetPath(Path);
 	}
 
diff --git a/Engine/Plugins/Runtime/ModelViewViewModel/Source/ModelViewViewModelBlueprint/Private/MVVMViewBlueprintCompiler.cpp b/Engine/Plugins/Runtime/ModelViewViewModel/Source/ModelViewViewModelBlueprint/Private/MVVMViewBlueprintCompiler.cpp
--- a/Engine/Plugins/Runtime/ModelViewViewModel/Source/ModelViewViewModelBlueprint/Private/MVVMViewBlueprintCompiler.cpp
+++ b/Engine/Plugins/Runtime/ModelViewViewModel/Source/ModelViewViewModelBlueprint/Private/MVVMViewBlueprintCompiler.cpp
@@ -183,11 +183,11 @@ void FMVVMViewBlueprintCompiler::AddMessages(TArrayView<TWeakPtr<FCompilerBindin
 {
 	for (TWeakPtr<FCompilerBinding>& Binding : Bindings)
 	{
-		AddMessageForBinding(Binding.Pin(), MessageText, MessageType, FName());
+		AddMessageForBinding(Binding.Pin(), MessageText, MessageType, FMVVMBlueprintPinId());
 	}
 	for (TWeakPtr<FCompilerEvent>& Event : Events)
 	{
-		AddMessageForEvent(Event.Pin(), MessageText, MessageType, FName());
+		AddMessageForEvent(Event.Pin(), MessageText, MessageType, FMVVMBlueprintPinId());
 	}
 
 	if (Bindings.Num() == 0 && Events.Num() == 0)
@@ -196,12 +196,13 @@ void FMVVMViewBlueprintCompiler::AddMessages(TArrayView<TWeakPtr<FCompilerBindin
 	}
 }
 
-void FMVVMViewBlueprintCompiler::AddMessageForBinding(const TSharedPtr<FCompilerBinding>& Binding, const FText& MessageText, EMessageType MessageType, FName ArgumentName) const
+
+void FMVVMViewBlueprintCompiler::AddMessageForBinding(const TSharedPtr<FCompilerBinding>& Binding, const FText& MessageText, EMessageType MessageType, const FMVVMBlueprintPinId& PinId) const
 {
 	const FMVVMBlueprintViewBinding* BindingPtr = Binding ? BlueprintView->GetBindingAt(Binding->Key.ViewBindingIndex) : nullptr;
 	if (BindingPtr)
 	{
-		AddMessageForBinding(*BindingPtr, MessageText, MessageType, ArgumentName);
+		AddMessageForBinding(*BindingPtr, MessageText, MessageType, PinId);
 	}
 	else
 	{
@@ -209,14 +210,27 @@ void FMVVMViewBlueprintCompiler::AddMessageForBinding(const TSharedPtr<FCompiler
 	}
 }
 
-void FMVVMViewBlueprintCompiler::AddMessageForBinding(const FMVVMBlueprintViewBinding& Binding, const FText& MessageText, EMessageType MessageType, FName ArgumentName) const
+
+FText GetJoinArgumentNames(const FMVVMBlueprintPinId& PinId)
+{
+	TArray<FText> JoinedNames;
+	JoinedNames.Reserve(PinId.GetNames().Num());
+	for (FName Name : PinId.GetNames())
+	{
+		JoinedNames.Add(FText::FromName(Name));
+	}
+	return FText::Join(LOCTEXT("NameSeperator", "."), JoinedNames);
+}
+
+
+void FMVVMViewBlueprintCompiler::AddMessageForBinding(const FMVVMBlueprintViewBinding& Binding, const FText& MessageText, EMessageType MessageType, const FMVVMBlueprintPinId& PinId) const
 {
 	const FText BindingName = FText::FromString(Binding.GetDisplayNameString(WidgetBlueprintCompilerContext.WidgetBlueprint()));
 
 	FText FormattedError;
-	if (!ArgumentName.IsNone())
+	if (PinId.IsValid())
 	{
-		FormattedError = FText::Format(LOCTEXT("BindingFormatWithArgument", "Binding '{0}': Argument '{1}' - {2}"), BindingName, FText::FromName(ArgumentName), MessageText);
+		FormattedError = FText::Format(LOCTEXT("BindingFormatWithArgument", "Binding '{0}': Argument '{1}' - {2}"), BindingName, UE::MVVM::Private::GetJoinArgumentNames(PinId), MessageText);
 	}
 	else
 	{
@@ -230,12 +244,12 @@ void FMVVMViewBlueprintCompiler::AddMessageForBinding(const FMVVMBlueprintViewBi
 }
 
 
-void FMVVMViewBlueprintCompiler::AddMessageForEvent(const TSharedPtr<FCompilerEvent>& Event, const FText& MessageText, EMessageType MessageType, FName ArgumentName) const
+void FMVVMViewBlueprintCompiler::AddMessageForEvent(const TSharedPtr<FCompilerEvent>& Event, const FText& MessageText, EMessageType MessageType, const FMVVMBlueprintPinId& PinId) const
 {
 	const UMVVMBlueprintViewEvent* EventPtr = Event ? Event->Event.Get() : nullptr;
 	if (EventPtr)
 	{
-		AddMessageForEvent(EventPtr, MessageText, MessageType, ArgumentName);
+		AddMessageForEvent(EventPtr, MessageText, MessageType, PinId);
 	}
 	else
 	{
@@ -244,14 +258,14 @@ void FMVVMViewBlueprintCompiler::AddMessageForEvent(const TSharedPtr<FCompilerEv
 }
 
 
-void FMVVMViewBlueprintCompiler::AddMessageForEvent(const UMVVMBlueprintViewEvent* Event, const FText& MessageText, EMessageType MessageType, FName ArgumentName) const
+void FMVVMViewBlueprintCompiler::AddMessageForEvent(const UMVVMBlueprintViewEvent* Event, const FText& MessageText, EMessageType MessageType, const FMVVMBlueprintPinId& PinId) const
 {
 	const FText EventName = Event->GetDisplayName(true);
 
 	FText FormattedError;
-	if (!ArgumentName.IsNone())
+	if (PinId.IsValid())
 	{
-		FormattedError = FText::Format(LOCTEXT("EventFormatWithArgument", "Event '{0}': Argument '{1}' - {2}"), EventName, FText::FromName(ArgumentName), MessageText);
+		FormattedError = FText::Format(LOCTEXT("EventFormatWithArgument", "Event '{0}': Argument '{1}' - {2}"), EventName, UE::MVVM::Private::GetJoinArgumentNames(PinId), MessageText);
 	}
 	else
 	{
@@ -542,7 +556,7 @@ void FMVVMViewBlueprintCompiler::CreateBindingList(const FWidgetBlueprintCompile
 		const bool bHasConversionFunction = Binding.Conversion.GetConversionFunction(true) || Binding.Conversion.GetConversionFunction(false);
 		if (bHasConversionFunction && Binding.BindingType == EMVVMBindingMode::TwoWay)
 		{
-			AddMessageForBinding(Binding, LOCTEXT("TwoWayBindingsWithConversion", "Two-way bindings are not allowed to use conversion functions."), EMessageType::Error, FName());
+			AddMessageForBinding(Binding, LOCTEXT("TwoWayBindingsWithConversion", "Two-way bindings are not allowed to use conversion functions."), EMessageType::Error, FMVVMBlueprintPinId());
 			bIsCreateVariableStepValid = false;
 			continue;
 		}
@@ -587,7 +601,7 @@ void FMVVMViewBlueprintCompiler::CreateEventList(const FWidgetBlueprintCompilerC
 
 		if (!EventPtr->GetEventPath().HasPaths())
 		{
-			AddMessageForEvent(EventPtr, LOCTEXT("EventInvalidEventPath", "The event path is invalid."), EMessageType::Error, FName());
+			AddMessageForEvent(EventPtr, LOCTEXT("EventInvalidEventPath", "The event path is invalid."), EMessageType::Error, FMVVMBlueprintPinId());
 			bIsCreateVariableStepValid = false;
 			continue;
 		}
@@ -805,12 +819,12 @@ void FMVVMViewBlueprintCompiler::CreateRequiredProperties(const FWidgetBlueprint
 	{
 		const FMVVMBlueprintViewBinding& Binding = *(BlueprintView->GetBindingAt(ValidBinding->Key.ViewBindingIndex));
 
-		auto RunGenerateCompilerSourceContext = [Self = this, &GenerateCompilerContext, &Binding](bool bCreateSource, const FMVVMBlueprintPropertyPath& PropertyPath, FName ArgumentName)
+		auto RunGenerateCompilerSourceContext = [Self = this, &GenerateCompilerContext, &Binding](bool bCreateSource, const FMVVMBlueprintPropertyPath& PropertyPath, const FMVVMBlueprintPinId& PinId)
 			{
 				TValueOrError<void, FText> SourceContextResult = GenerateCompilerContext(bCreateSource, PropertyPath);
 				if (SourceContextResult.HasError())
 				{
-					Self->AddMessageForBinding(Binding, SourceContextResult.StealError(), EMessageType::Error, ArgumentName);
+					Self->AddMessageForBinding(Binding, SourceContextResult.StealError(), EMessageType::Error, PinId);
 					Self->bIsCreateVariableStepValid = false;
 				}
 			};
@@ -823,29 +837,29 @@ void FMVVMViewBlueprintCompiler::CreateRequiredProperties(const FWidgetBlueprint
 			{
 				if (Pin.UsedPathAsValue())
 				{
-					if (Pin.GetName().IsNone())
+					if (!Pin.IsValid())
 					{
 						AddMessageForBinding(Binding
 							, FText::Format(LOCTEXT("InvalidFunctionArgumentPathName", "The conversion function {0} has an invalid argument."), FText::FromString(Binding.GetDisplayNameString(WidgetBlueprintCompilerContext.WidgetBlueprint())))
 							, EMessageType::Error
-							, FName()
+							, FMVVMBlueprintPinId()
 						);
 						bIsCreateVariableStepValid = false;
 						continue;
 					}
-					RunGenerateCompilerSourceContext(true, Pin.GetPath(), Pin.GetName());
+					RunGenerateCompilerSourceContext(true, Pin.GetPath(), Pin.GetId());
 				}
 			}
 
 			// validate the destination
 			const FMVVMBlueprintPropertyPath& BindingDestinationPath = ValidBinding->Key.bIsForwardBinding ? Binding.DestinationPath : Binding.SourcePath;
-			RunGenerateCompilerSourceContext(false, BindingDestinationPath, FName());
+			RunGenerateCompilerSourceContext(false, BindingDestinationPath, FMVVMBlueprintPinId());
 		}
 		else
 		{
 			// if we aren't using a conversion function, validate the source and destination paths
-			RunGenerateCompilerSourceContext(ValidBinding->Key.bIsForwardBinding, Binding.SourcePath, FName());
-			RunGenerateCompilerSourceContext(!ValidBinding->Key.bIsForwardBinding, Binding.DestinationPath, FName());
+			RunGenerateCompilerSourceContext(ValidBinding->Key.bIsForwardBinding, Binding.SourcePath, FMVVMBlueprintPinId());
+			RunGenerateCompilerSourceContext(!ValidBinding->Key.bIsForwardBinding, Binding.DestinationPath, FMVVMBlueprintPinId());
 		}
 	}
 
@@ -855,29 +869,33 @@ void FMVVMViewBlueprintCompiler::CreateRequiredProperties(const FWidgetBlueprint
 	{
 		UMVVMBlueprintViewEvent* EventPtr = ValidEvent->Event.Get();
 		check(EventPtr);
-		auto RunGenerateCompilerSourceContext = [Self = this, &GenerateCompilerContext, EventPtr](bool bCreateSource, const FMVVMBlueprintPropertyPath& PropertyPath, FName ArgumentName)
+		auto RunGenerateCompilerSourceContext = [Self = this, &GenerateCompilerContext, EventPtr](bool bCreateSource, const FMVVMBlueprintPropertyPath& PropertyPath, const FMVVMBlueprintPinId& PinId)
 		{
 			TValueOrError<void, FText> SourceContextResult = GenerateCompilerContext(bCreateSource, PropertyPath);
 			if (SourceContextResult.HasError())
 			{
-				Self->AddMessageForEvent(EventPtr, SourceContextResult.StealError(), EMessageType::Error, ArgumentName);
+				Self->AddMessageForEvent(EventPtr, SourceContextResult.StealError(), EMessageType::Error, PinId);
 				Self->bIsCreateVariableStepValid = false;
 			}
 		};
 
-		RunGenerateCompilerSourceContext(false, EventPtr->GetEventPath(), FName());
-		RunGenerateCompilerSourceContext(false, EventPtr->GetDestinationPath(), FName());
+		RunGenerateCompilerSourceContext(false, EventPtr->GetEventPath(), FMVVMBlueprintPinId());
+		RunGenerateCompilerSourceContext(false, EventPtr->GetDestinationPath(), FMVVMBlueprintPinId());
 		for (const FMVVMBlueprintPin& Pin : EventPtr->GetPins())
 		{
 			if (Pin.UsedPathAsValue())
 			{
-				if (Pin.GetName().IsNone())
+				if (!Pin.IsValid())
 				{
-					AddMessageForEvent(EventPtr, LOCTEXT("InvalidEventArgumentPathName", "The event has an invalid argument."), EMessageType::Error, Pin.GetName());
+					AddMessageForEvent(EventPtr
+						, LOCTEXT("InvalidEventArgumentPathName", "The event has an invalid argument.")
+						, EMessageType::Error
+						, FMVVMBlueprintPinId()
+					);
 					bIsCreateVariableStepValid = false;
 					continue;
 				}
-				RunGenerateCompilerSourceContext(false, Pin.GetPath(), Pin.GetName());
+				RunGenerateCompilerSourceContext(false, Pin.GetPath(), Pin.GetId());
 			}
 		}
 	}
@@ -961,7 +979,7 @@ void FMVVMViewBlueprintCompiler::CategorizeBindings(const FWidgetBlueprintCompil
 			UClass* NewClass = WidgetBlueprintCompilerContext.NewClass;
 			if (!ConversionFunction->IsValid(WidgetBlueprintCompilerContext.WidgetBlueprint()))
 			{
-				AddMessageForBinding(Binding, LOCTEXT("InvalidConversionFunction", "The conversion function is invalid."), EMessageType::Error, FName());
+				AddMessageForBinding(Binding, LOCTEXT("InvalidConversionFunction", "The conversion function is invalid."), EMessageType::Error, FMVVMBlueprintPinId());
 				bIsCreateFunctionsStepValid = false;
 				continue;
 			}
@@ -970,7 +988,10 @@ void FMVVMViewBlueprintCompiler::CategorizeBindings(const FWidgetBlueprintCompil
 			UEdGraph* WrapperGraph = ConversionFunction->GetOrCreateIntermediateWrapperGraph(WidgetBlueprintCompilerContext);
 			if (WrapperGraph == nullptr)
 			{
-				AddMessageForBinding(Binding, LOCTEXT("InvalidConversionFunctionGraph", "The conversion function graph could not be generated."), EMessageType::Error, FName());
+				AddMessageForBinding(Binding, LOCTEXT("InvalidConversionFunctionGraph", "The conversion function graph could not be generated.")
+					, EMessageType::Error
+					, FMVVMBlueprintPinId()
+				);
 				bIsCreateFunctionsStepValid = false;
 				continue;
 			}
@@ -979,7 +1000,10 @@ void FMVVMViewBlueprintCompiler::CategorizeBindings(const FWidgetBlueprintCompil
 
 			if (ConversionFunction->HasOrphanedPin())
 			{
-				AddMessageForBinding(Binding, LOCTEXT("InvalidConversionFunctionGraphOrphaned", "The conversion function has an orphaned pin."), EMessageType::Warning, FName());
+				AddMessageForBinding(Binding, LOCTEXT("InvalidConversionFunctionGraphOrphaned", "The conversion function has an orphaned pin.")
+					, EMessageType::Warning
+					, FMVVMBlueprintPinId()
+				);
 				bIsCreateFunctionsStepValid = false;
 				continue;
 			}
@@ -989,7 +1013,7 @@ void FMVVMViewBlueprintCompiler::CategorizeBindings(const FWidgetBlueprintCompil
 				AddMessageForBinding(Binding
 					, FText::Format(LOCTEXT("InvalidNumberOfFunctionPin", "The conversion function {0} has no source."), FText::FromString(Binding.GetDisplayNameString(WidgetBlueprintCompilerContext.WidgetBlueprint())))
 					, EMessageType::Error
-					, FName()
+					, FMVVMBlueprintPinId()
 				);
 				bIsCreateFunctionsStepValid = false;
 				continue;
@@ -1009,7 +1033,7 @@ void FMVVMViewBlueprintCompiler::CategorizeBindings(const FWidgetBlueprintCompil
 			//const FMVVMBlueprintPropertyPath& BindingSourcePath = ValidBinding->Key.bIsForwardBinding ? Binding.SourcePath : Binding.DestinationPath;
 			//if (BindingSourcePath.IsValid())
 			//{
-			//	AddMessageForBinding(Binding, LOCTEXT("ShouldNotHaveSourceWarning", "Internal Error. The binding should not have a source."), EMessageType::Warning, FName());
+			//	AddMessageForBinding(Binding, LOCTEXT("ShouldNotHaveSourceWarning", "Internal Error. The binding should not have a source."), EMessageType::Warning, FMVVMBlueprintPinId());
 			//}
 		}
 		else
@@ -1030,7 +1054,7 @@ void FMVVMViewBlueprintCompiler::CategorizeEvents(const FWidgetBlueprintCompiler
 		UEdGraph* WrapperGraph = EventPtr->GetOrCreateWrapperGraph();
 		if (WrapperGraph == nullptr)
 		{
-			AddMessageForEvent(EventPtr, LOCTEXT("InvalidEventGraph", "The event could not be generated."), EMessageType::Warning, FName());
+			AddMessageForEvent(EventPtr, LOCTEXT("InvalidEventGraph", "The event could not be generated."), EMessageType::Warning, FMVVMBlueprintPinId());
 			bIsCreateFunctionsStepValid = false;
 			continue;
 		}
@@ -1039,7 +1063,7 @@ void FMVVMViewBlueprintCompiler::CategorizeEvents(const FWidgetBlueprintCompiler
 
 		if (EventPtr->HasOrphanedPin())
 		{
-			AddMessageForEvent(EventPtr, LOCTEXT("InvalidEventGraphOrphaned", "The event has an orphaned pin."), EMessageType::Warning, FName());
+			AddMessageForEvent(EventPtr, LOCTEXT("InvalidEventGraphOrphaned", "The event has an orphaned pin."), EMessageType::Warning, FMVVMBlueprintPinId());
 			bIsCreateFunctionsStepValid = false;
 			continue;
 		}
@@ -1072,7 +1096,7 @@ void FMVVMViewBlueprintCompiler::CreateWriteFieldContexts(const FWidgetBlueprint
 			TValueOrError<FCreateFieldsResult, FText> FieldContextResult = CreateFieldContext(NewSkeletonClass, DestinationPropertyPath, false);
 			if (FieldContextResult.HasError())
 			{
-				AddMessageForBinding(Binding, FieldContextResult.StealError(), EMessageType::Error, FName());
+				AddMessageForBinding(Binding, FieldContextResult.StealError(), EMessageType::Error, FMVVMBlueprintPinId());
 				bIsCreateFunctionsStepValid = false;
 				continue;
 			}
@@ -1090,7 +1114,7 @@ void FMVVMViewBlueprintCompiler::CreateWriteFieldContexts(const FWidgetBlueprint
 					//AddMessageForBinding(Binding
 					//	, FText::Format(LOCTEXT("PropertyPathAlreadyUsed", "The property path '{0}' is already used by another binding."), PropertyPathToText(NewSkeletonClass, BlueprintView.Get(), DestinationPropertyPath))
 					//	, EMessageType::Warning
-					//	, FName()
+					//	, FMVVMBlueprintPinId()
 					//);
 					ValidBinding->WritePath = *Found;
 					continue;
@@ -1125,7 +1149,7 @@ void FMVVMViewBlueprintCompiler::CreateWriteFieldContexts(const FWidgetBlueprint
 				AddMessageForEvent(EventPtr
 					, FText::Format(Private::PropertyPathIsInvalidFormat, PropertyPathToText(NewSkeletonClass, BlueprintView.Get(), EventPtr->GetDestinationPath()))
 					, EMessageType::Error
-					, FName()
+					, FMVVMBlueprintPinId()
 				);
 				bIsCreateFunctionsStepValid = false;
 				continue;
@@ -1396,7 +1420,7 @@ void FMVVMViewBlueprintCompiler::CreateReadFieldContexts(UWidgetBlueprintGenerat
 	{
 		const FMVVMBlueprintViewBinding& Binding = *BlueprintView->GetBindingAt(ValidBinding->Key.ViewBindingIndex);
 
-		auto CreateBindingSourceContext = [Self = this, Class, &ValidBinding, &Binding](const FMVVMBlueprintPropertyPath& PropertyPath, FName ArgumentName) -> TValueOrError<FCreateFieldsResult, void>
+		auto CreateBindingSourceContext = [Self = this, Class, &ValidBinding, &Binding](const FMVVMBlueprintPropertyPath& PropertyPath, const FMVVMBlueprintPinId& PinId) -> TValueOrError<FCreateFieldsResult, void>
 		{
 			TValueOrError<FCreateFieldsResult, FText> FieldContextResult = Self->CreateFieldContext(Class, PropertyPath, true);
 			if (FieldContextResult.HasError())
@@ -1404,7 +1428,7 @@ void FMVVMViewBlueprintCompiler::CreateReadFieldContexts(UWidgetBlueprintGenerat
 				Self->AddMessageForBinding(Binding
 					, FText::Format(Private::PropertyPathIsInvalidFormat, PropertyPathToText(Class, Self->BlueprintView.Get(), PropertyPath))
 					, EMessageType::Error
-					, ArgumentName
+					, PinId
 				);
 				Self->bIsPreCompileStepValid = false;
 				return MakeError();
@@ -1413,7 +1437,7 @@ void FMVVMViewBlueprintCompiler::CreateReadFieldContexts(UWidgetBlueprintGenerat
 			return MakeValue(FieldContextResult.StealValue());
 		};
 
-		auto CreateFieldId = [Self = this, Class, &Binding, bIsOneTimeBinding = ValidBinding->bIsOneTimeBinding](const FMVVMBlueprintPropertyPath& PropertyPath, TSharedPtr<FGeneratedReadFieldPathContext>& ReadFieldContext, FName ArgumentName) -> TValueOrError<void, void>
+		auto CreateFieldId = [Self = this, Class, &Binding, bIsOneTimeBinding = ValidBinding->bIsOneTimeBinding](const FMVVMBlueprintPropertyPath& PropertyPath, TSharedPtr<FGeneratedReadFieldPathContext>& ReadFieldContext, const FMVVMBlueprintPinId& PinId) -> TValueOrError<void, void>
 		{
 			if (!ReadFieldContext->NotificationField.IsValid())
 			{
@@ -1423,7 +1447,7 @@ void FMVVMViewBlueprintCompiler::CreateReadFieldContexts(UWidgetBlueprintGenerat
 					Self->AddMessageForBinding(Binding
 						, FText::Format(LOCTEXT("CreateNotifyFieldIdFailedInvalidSelfContext", "The property path '{0}' is invalid. {1}"), PropertyPathToText(Class, Self->BlueprintView.Get(), PropertyPath), CreateFieldResult.StealError())
 						, EMessageType::Error
-						, ArgumentName
+						, PinId
 					);
 					return MakeError();
 				}
@@ -1445,7 +1469,7 @@ void FMVVMViewBlueprintCompiler::CreateReadFieldContexts(UWidgetBlueprintGenerat
 								Self->AddMessageForBinding(Binding
 									, FText::Format(LOCTEXT("CreateNotifyFieldIdFailedInvalidInvalidContext", "Internal error. The property path '{0}' is invalid. The context is invalid."), PropertyPathToText(Class, Self->BlueprintView.Get(), PropertyPath))
 									, EMessageType::Error
-									, ArgumentName
+									, PinId
 								);
 								return MakeError();
 							}
@@ -1462,7 +1486,7 @@ void FMVVMViewBlueprintCompiler::CreateReadFieldContexts(UWidgetBlueprintGenerat
 		if (ValidBinding->Type == FCompilerBinding::EType::Assignment)
 		{
 			const FMVVMBlueprintPropertyPath& BindingSourcePath = ValidBinding->Key.bIsForwardBinding ? Binding.SourcePath : Binding.DestinationPath;
-			TValueOrError<FCreateFieldsResult, void> CreateSourceResult = CreateBindingSourceContext(BindingSourcePath, FName());
+			TValueOrError<FCreateFieldsResult, void> CreateSourceResult = CreateBindingSourceContext(BindingSourcePath, FMVVMBlueprintPinId());
 			if (CreateSourceResult.HasError())
 			{
 				continue;
@@ -1476,7 +1500,7 @@ void FMVVMViewBlueprintCompiler::CreateReadFieldContexts(UWidgetBlueprintGenerat
 				Found->GeneratedFields = MoveTemp(CreateSourceResult.GetValue().GeneratedFields);
 				Found->SkeletalGeneratedFields = MoveTemp(CreateSourceResult.GetValue().SkeletalGeneratedFields);
 
-				if (CreateFieldId(BindingSourcePath, Found, FName()).HasError())
+				if (CreateFieldId(BindingSourcePath, Found, FMVVMBlueprintPinId()).HasError())
 				{
 					bIsPreCompileStepValid = false;
 					continue;
@@ -1498,7 +1522,7 @@ void FMVVMViewBlueprintCompiler::CreateReadFieldContexts(UWidgetBlueprintGenerat
 				{
 					if (Pin.UsedPathAsValue())
 					{
-						TValueOrError<FCreateFieldsResult, void> CreateSourceResult = CreateBindingSourceContext(Pin.GetPath(), Pin.GetName());
+						TValueOrError<FCreateFieldsResult, void> CreateSourceResult = CreateBindingSourceContext(Pin.GetPath(), Pin.GetId());
 						if (CreateSourceResult.HasError())
 						{
 							continue;
@@ -1512,7 +1536,7 @@ void FMVVMViewBlueprintCompiler::CreateReadFieldContexts(UWidgetBlueprintGenerat
 							Found->GeneratedFields = MoveTemp(CreateSourceResult.GetValue().GeneratedFields);
 							Found->SkeletalGeneratedFields = MoveTemp(CreateSourceResult.GetValue().SkeletalGeneratedFields);
 
-							if (CreateFieldId(Pin.GetPath(), Found, Pin.GetName()).HasError())
+							if (CreateFieldId(Pin.GetPath(), Found, Pin.GetId()).HasError())
 							{
 								bIsPreCompileStepValid = false;
 								continue;
@@ -1546,7 +1570,7 @@ void FMVVMViewBlueprintCompiler::CreateReadFieldContexts(UWidgetBlueprintGenerat
 						AddMessageForEvent(ValidEvent
 							, FText::Format(Private::PropertyPathIsInvalidFormat, PropertyPathToText(Class, BlueprintView.Get(), Pin.GetPath()))
 							, EMessageType::Error
-							, Pin.GetName()
+							, Pin.GetId()
 						);
 						bIsPreCompileStepValid = false;
 					}
@@ -2039,7 +2063,11 @@ void FMVVMViewBlueprintCompiler::PreCompileBindings(UWidgetBlueprintGeneratedCla
 		{
 			if (!GetDefault<UMVVMDeveloperProjectSettings>()->IsExecutionModeAllowed(Binding.OverrideExecutionMode))
 			{
-				Self->AddMessageForBinding(Binding, LOCTEXT("NotAllowedExecutionMode", "The binding has a restricted execution mode."), EMessageType::Error, FName());
+				Self->AddMessageForBinding(Binding
+					, LOCTEXT("NotAllowedExecutionMode", "The binding has a restricted execution mode.")
+					, EMessageType::Error
+					, FMVVMBlueprintPinId()
+				);
 				return false;
 			}
 		}
@@ -2069,7 +2097,7 @@ void FMVVMViewBlueprintCompiler::PreCompileBindings(UWidgetBlueprintGeneratedCla
 							Self->AddMessageForBinding(Binding
 								, FText::Format(LOCTEXT("CouldNotCreateFieldId", "Could not create Field. {0}"), FieldIdResult.GetError())
 								, EMessageType::Error
-								, FName()
+								, FMVVMBlueprintPinId()
 							);
 							return false;
 						}
@@ -2088,7 +2116,7 @@ void FMVVMViewBlueprintCompiler::PreCompileBindings(UWidgetBlueprintGeneratedCla
 			Self->AddMessageForBinding(Binding
 				, LOCTEXT("CouldNotCreateSourceFields", "There is no field to bind to. The binding must be a OneTime binding.")
 				, EMessageType::Error
-				, FName()
+				, FMVVMBlueprintPinId()
 			);
 			return false;
 		}
@@ -2115,7 +2143,7 @@ void FMVVMViewBlueprintCompiler::PreCompileBindings(UWidgetBlueprintGeneratedCla
 						Self->AddMessageForBinding(Binding
 							, FText::Format(Private::CouldNotCreateSourceFieldPathFormat, ::UE::MVVM::FieldPathHelper::ToText(ReadPath->GeneratedFields), FieldPathResult.GetError())
 							, EMessageType::Error
-							, FName()
+							, FMVVMBlueprintPinId()
 						);
 						return false;
 					}
@@ -2133,7 +2161,7 @@ void FMVVMViewBlueprintCompiler::PreCompileBindings(UWidgetBlueprintGeneratedCla
 			Self->AddMessageForBinding(Binding
 				, LOCTEXT("CouldNotCreateSourceReadIsPresent", "Internal error. There should be no property to read from.")
 				, EMessageType::Error
-				, FName()
+				, FMVVMBlueprintPinId()
 			);
 			return false;
 		}
@@ -2153,7 +2181,7 @@ void FMVVMViewBlueprintCompiler::PreCompileBindings(UWidgetBlueprintGeneratedCla
 			Self->AddMessageForBinding(Binding
 				, LOCTEXT("InvalidWritePath", "Internal Error. The write path is invalid.")
 				, EMessageType::Error
-				, FName()
+				, FMVVMBlueprintPinId()
 			);
 			return false;
 		}
@@ -2164,7 +2192,7 @@ void FMVVMViewBlueprintCompiler::PreCompileBindings(UWidgetBlueprintGeneratedCla
 			Self->AddMessageForBinding(Binding
 				, FText::Format(Private::CouldNotCreateDestinationFieldPathFormat, ::UE::MVVM::FieldPathHelper::ToText(ValidBinding.WritePath->GeneratedFields), FieldPathResult.GetError())
 				, EMessageType::Error
-				, FName()
+				, FMVVMBlueprintPinId()
 			);
 			return false;
 		}
@@ -2188,7 +2216,7 @@ void FMVVMViewBlueprintCompiler::PreCompileBindings(UWidgetBlueprintGeneratedCla
 					Self->AddMessageForBinding(Binding
 						, FText::Format(LOCTEXT("ConversionFunctionNotFound", "The conversion function '{0}' could not be found."), FText::FromName(ViewConversionFunction->GetCompiledFunctionName(Class)))
 						, EMessageType::Error
-						, FName()
+						, FMVVMBlueprintPinId()
 					);
 					return false;
 				}
@@ -2201,14 +2229,18 @@ void FMVVMViewBlueprintCompiler::PreCompileBindings(UWidgetBlueprintGeneratedCla
 						Self->AddMessageForBinding(Binding
 							, FText::Format(LOCTEXT("ConversionFunctionNotAllow", "The conversion function {0} is not allowed."), FText::FromName(FunctionOrWrapperFunction.Get<const UFunction*>()->GetFName()))
 							, EMessageType::Error
-							, FName()
+							, FMVVMBlueprintPinId()
 						);
 						return false;
 					}
 				}
 				else
 				{
-					Self->AddMessageForBinding(Binding, LOCTEXT("ConversionFunctionNodeNotAllow", "The conversion function node is not allowed."), EMessageType::Error, FName());
+					Self->AddMessageForBinding(Binding
+						, LOCTEXT("ConversionFunctionNodeNotAllow", "The conversion function node is not allowed.")
+						, EMessageType::Error
+						, FMVVMBlueprintPinId()
+					);
 					return false;
 				}
 			}
@@ -2224,7 +2256,7 @@ void FMVVMViewBlueprintCompiler::PreCompileBindings(UWidgetBlueprintGeneratedCla
 						, FText::FromString(ConversionFunction->GetPathName())
 						, FieldPathResult.GetError())
 					, EMessageType::Error
-					, FName()
+					, FMVVMBlueprintPinId()
 				);
 				return false;
 			}
@@ -2237,14 +2269,14 @@ void FMVVMViewBlueprintCompiler::PreCompileBindings(UWidgetBlueprintGeneratedCla
 			const bool bShouldHaveConversionFunction = ValidBinding.Type == FCompilerBinding::EType::ComplexConversionFunction || ValidBinding.Type == FCompilerBinding::EType::SimpleConversionFunction;
 			if ((ConversionFunction != nullptr) != bShouldHaveConversionFunction)
 			{
-				Self->AddMessageForBinding(Binding, LOCTEXT("ConversionFunctionShouldExist", "Internal error. The conversion function should exist."), EMessageType::Error, FName());
+				Self->AddMessageForBinding(Binding, LOCTEXT("ConversionFunctionShouldExist", "Internal error. The conversion function should exist."), EMessageType::Error, FMVVMBlueprintPinId());
 				return false;
 			}
 
 			const bool bShouldHaveComplexConversionFunction = ValidBinding.Type == FCompilerBinding::EType::ComplexConversionFunction;
 			if (bShouldHaveComplexConversionFunction != BindingHelper::IsValidForComplexRuntimeConversion(ConversionFunction))
 			{
-				Self->AddMessageForBinding(Binding, LOCTEXT("ConversionFunctionIsNotComplex", "Internal Error. The complex conversion function does not respect the prerequisite."), EMessageType::Error, FName());
+				Self->AddMessageForBinding(Binding, LOCTEXT("ConversionFunctionIsNotComplex", "Internal Error. The complex conversion function does not respect the prerequisite."), EMessageType::Error, FMVVMBlueprintPinId());
 				return false;
 			}
 		}
@@ -2259,7 +2291,7 @@ void FMVVMViewBlueprintCompiler::PreCompileBindings(UWidgetBlueprintGeneratedCla
 			&& ValidBinding->Type != FCompilerBinding::EType::ComplexConversionFunction
 			&& ValidBinding->Type != FCompilerBinding::EType::SimpleConversionFunction)
 		{
-			AddMessageForBinding(Binding, LOCTEXT("UnsupportedBindingType", "The binding is invalid."), EMessageType::Error, FName());
+			AddMessageForBinding(Binding, LOCTEXT("UnsupportedBindingType", "The binding is invalid."), EMessageType::Error, FMVVMBlueprintPinId());
 			bIsPreCompileStepValid = false;
 			continue;
 		}
@@ -2289,7 +2321,7 @@ void FMVVMViewBlueprintCompiler::PreCompileBindings(UWidgetBlueprintGeneratedCla
 			AddMessageForBinding(Binding
 				, FText::Format(LOCTEXT("CouldNotCreateBinding", "Could not create binding. {0}"), BindingResult.StealError())
 				, EMessageType::Error
-				, FName()
+				, FMVVMBlueprintPinId()
 			);
 			bIsPreCompileStepValid = false;
 			continue;
@@ -2355,7 +2387,7 @@ void FMVVMViewBlueprintCompiler::CompileBindings(const FCompiledBindingLibraryCo
 		const FMVVMVCompiledBinding* CompiledBinding = CompileResult.Bindings.Find(ValidBinding->BindingHandle);
 		if (CompiledBinding == nullptr)
 		{
-			AddMessageForBinding(Binding, LOCTEXT("CompiledBindingNotGenerated", "Could not generate compiled binding."), EMessageType::Error, FName());
+			AddMessageForBinding(Binding, LOCTEXT("CompiledBindingNotGenerated", "Could not generate compiled binding."), EMessageType::Error, FMVVMBlueprintPinId());
 			bIsCompileStepValid = false;
 			continue;
 		}
@@ -2378,7 +2410,7 @@ void FMVVMViewBlueprintCompiler::CompileBindings(const FCompiledBindingLibraryCo
 		{
 			if (ReadPath->OptionalSource == nullptr)
 			{
-				AddMessageForBinding(Binding, LOCTEXT("InvalidSourceInternal", "Internal error. The binding has an invalid source."), EMessageType::Error, FName());
+				AddMessageForBinding(Binding, LOCTEXT("InvalidSourceInternal", "Internal error. The binding has an invalid source."), EMessageType::Error, FMVVMBlueprintPinId());
 				bIsCompileStepValid = false;
 				continue;
 			}
@@ -2389,7 +2421,7 @@ void FMVVMViewBlueprintCompiler::CompileBindings(const FCompiledBindingLibraryCo
 				});
 			if (!ViewExtension->Sources.IsValidIndex(ViewExtensionSourceCreatorsIndex))
 			{
-				AddMessageForBinding(Binding, LOCTEXT("CompiledSourceCreatorNotGenerated", "Internal error. The source creator was not generated."), EMessageType::Error, FName());
+				AddMessageForBinding(Binding, LOCTEXT("CompiledSourceCreatorNotGenerated", "Internal error. The source creator was not generated."), EMessageType::Error, FMVVMBlueprintPinId());
 				bIsCompileStepValid = false;
 				continue;
 			}
@@ -2403,7 +2435,7 @@ void FMVVMViewBlueprintCompiler::CompileBindings(const FCompiledBindingLibraryCo
 			const UE::FieldNotification::FFieldId* CompiledFieldId = bHasField ? CompileResult.FieldIds.Find(ReadPath->NotificationField->LibraryCompilerHandle) : nullptr;
 			if (CompiledFieldId == nullptr && bHasField)
 			{
-				AddMessageForBinding(Binding, LOCTEXT("CompiledFieldNotGenerated", "Internal error. The FieldId was not generated."), EMessageType::Error, FName());
+				AddMessageForBinding(Binding, LOCTEXT("CompiledFieldNotGenerated", "Internal error. The FieldId was not generated."), EMessageType::Error, FMVVMBlueprintPinId());
 				bIsCompileStepValid = false;
 				continue;
 			}
@@ -2549,7 +2581,7 @@ void FMVVMViewBlueprintCompiler::PreCompileEvents(UWidgetBlueprintGeneratedClass
 			AddMessageForEvent(EventPtr
 				, FText::Format(Private::PropertyPathIsInvalidFormat, PropertyPathToText(Class, BlueprintView.Get(), EventPtr->GetEventPath()))
 				, EMessageType::Error
-				, FName());
+				, FMVVMBlueprintPinId());
 			bIsPreCompileStepValid = false;
 			continue;
 		}
@@ -2561,7 +2593,7 @@ void FMVVMViewBlueprintCompiler::PreCompileEvents(UWidgetBlueprintGeneratedClass
 			AddMessageForEvent(EventPtr
 				, FText::Format(LOCTEXT("EventPathIsNotMulticastDelegate", "The event {0} is not a multicast delegate."), PropertyPathToText(Class, BlueprintView.Get(), EventPtr->GetEventPath()))
 				, EMessageType::Error
-				, FName());
+				, FMVVMBlueprintPinId());
 			bIsPreCompileStepValid = false;
 			continue;
 		}
@@ -2572,7 +2604,7 @@ void FMVVMViewBlueprintCompiler::PreCompileEvents(UWidgetBlueprintGeneratedClass
 			AddMessageForEvent(EventPtr
 				, FText::Format(LOCTEXT("EventPathFunctionSignatureError", "The event {0} doesn't match the function signature."), DelegateProperty->GetDisplayNameText())
 				, EMessageType::Error
-				, FName());
+				, FMVVMBlueprintPinId());
 			bIsPreCompileStepValid = false;
 			continue;
 		}
@@ -2584,7 +2616,7 @@ void FMVVMViewBlueprintCompiler::PreCompileEvents(UWidgetBlueprintGeneratedClass
 			AddMessageForEvent(EventPtr
 				, FText::Format(Private::CouldNotCreateSourceFieldPathFormat, PropertyPathToText(Class, BlueprintView.Get(), EventPtr->GetEventPath()), EventFieldPathResult.GetError())
 				, EMessageType::Error
-				, FName());
+				, FMVVMBlueprintPinId());
 			bIsPreCompileStepValid = false;
 			continue;
 		}
@@ -2636,14 +2668,14 @@ void FMVVMViewBlueprintCompiler::CompileEvents(const FCompiledBindingLibraryComp
 		const FMVVMVCompiledFieldPath* CompiledFieldPath = CompileResult.FieldPaths.Find(ValidEvent->DelegateFieldPathHandle);
 		if (CompiledFieldPath == nullptr)
 		{
-			AddMessageForEvent(ValidEvent->Event.Get(), LOCTEXT("CompiledEventFieldPathNotGenerated", "Could not generate the event path."), EMessageType::Error, FName());
+			AddMessageForEvent(ValidEvent->Event.Get(), LOCTEXT("CompiledEventFieldPathNotGenerated", "Could not generate the event path."), EMessageType::Error, FMVVMBlueprintPinId());
 			bIsCompileStepValid = false;
 			continue;
 		}
 
 		if (ValidEvent->GeneratedGraphName.IsNone() || Class->FindFunctionByName(ValidEvent->GeneratedGraphName) == nullptr)
 		{
-			AddMessageForEvent(ValidEvent->Event.Get(), LOCTEXT("CompiledEventFieldPathNotGenerated", "Could not generate the event path."), EMessageType::Error, FName());
+			AddMessageForEvent(ValidEvent->Event.Get(), LOCTEXT("CompiledEventFieldPathNotGenerated", "Could not generate the event path."), EMessageType::Error, FMVVMBlueprintPinId());
 			bIsCompileStepValid = false;
 			continue;
 		}
diff --git a/Engine/Plugins/Runtime/ModelViewViewModel/Source/ModelViewViewModelBlueprint/Private/MVVMViewBlueprintCompiler.h b/Engine/Plugins/Runtime/ModelViewViewModel/Source/ModelViewViewModelBlueprint/Private/MVVMViewBlueprintCompiler.h
--- a/Engine/Plugins/Runtime/ModelViewViewModel/Source/ModelViewViewModelBlueprint/Private/MVVMViewBlueprintCompiler.h
+++ b/Engine/Plugins/Runtime/ModelViewViewModel/Source/ModelViewViewModelBlueprint/Private/MVVMViewBlueprintCompiler.h
@@ -11,9 +11,7 @@
 #include "UObject/StrongObjectPtr.h"
 #include "WidgetBlueprintCompiler.h"
 
-struct FMVVMBlueprintPropertyPath;
-
-struct FMVVMViewModelPropertyPath;
+struct FMVVMBlueprintPinId;
 struct FMVVMBlueprintViewBinding;
 class FWidgetBlueprintCompilerContext;
 class UEdGraph;
@@ -346,10 +344,10 @@ private:
 	};
 	void AddMessage(const FText& MessageText, EMessageType MessageType) const;
 	void AddMessages(TArrayView<TWeakPtr<FCompilerBinding>> Bindings, TArrayView<TWeakPtr<FCompilerEvent>> Events, const FText& MessageText, EMessageType MessageType) const;
-	void AddMessageForBinding(const TSharedPtr<FCompilerBinding>& Binding, const FText& MessageText, EMessageType MessageType, FName ArgumentName) const;
-	void AddMessageForBinding(const FMVVMBlueprintViewBinding& Binding, const FText& MessageText, EMessageType MessageType, FName ArgumentName) const;
-	void AddMessageForEvent(const TSharedPtr<FCompilerEvent>& Event, const FText& MessageText, EMessageType MessageType, FName ArgumentName) const;
-	void AddMessageForEvent(const UMVVMBlueprintViewEvent* Event, const FText& MessageText, EMessageType MessageType, FName ArgumentName) const;
+	void AddMessageForBinding(const TSharedPtr<FCompilerBinding>& Binding, const FText& MessageText, EMessageType MessageType, const FMVVMBlueprintPinId& ArgumentName) const;
+	void AddMessageForBinding(const FMVVMBlueprintViewBinding& Binding, const FText& MessageText, EMessageType MessageType, const FMVVMBlueprintPinId& ArgumentName) const;
+	void AddMessageForEvent(const TSharedPtr<FCompilerEvent>& Event, const FText& MessageText, EMessageType MessageType, const FMVVMBlueprintPinId& ArgumentName) const;
+	void AddMessageForEvent(const UMVVMBlueprintViewEvent* Event, const FText& MessageText, EMessageType MessageType, const FMVVMBlueprintPinId& ArgumentName) const;
 	void AddMessageForViewModel(const FMVVMBlueprintViewModelContext& ViewModel, const FText& Message, EMessageType MessageType) const;
 	void AddMessageForViewModel(const FText& ViewModelDisplayName, const FText& Message, EMessageType MessageType) const;
 
diff --git a/Engine/Plugins/Runtime/ModelViewViewModel/Source/ModelViewViewModelBlueprint/Public/Bindings/MVVMConversionFunctionHelper.h b/Engine/Plugins/Runtime/ModelViewViewModel/Source/ModelViewViewModelBlueprint/Public/Bindings/MVVMConversionFunctionHelper.h
--- a/Engine/Plugins/Runtime/ModelViewViewModel/Source/ModelViewViewModelBlueprint/Public/Bindings/MVVMConversionFunctionHelper.h
+++ b/Engine/Plugins/Runtime/ModelViewViewModel/Source/ModelViewViewModelBlueprint/Public/Bindings/MVVMConversionFunctionHelper.h
@@ -7,12 +7,14 @@
 #include "Templates/SubclassOf.h"
 #include "Templates/ValueOrError.h"
 
+struct FMVVMBlueprintFunctionReference;
 struct FMVVMBlueprintPropertyPath;
 struct FMVVMBlueprintViewBinding;
 
 class UBlueprint;
 class UClass;
 class UEdGraph;
+class UEdGraphNode;
 class UEdGraphPin;
 class UK2Node;
 class UK2Node_CallFunction;
@@ -25,11 +27,8 @@ namespace UE::MVVM::ConversionFunctionHelper
 	/** Conversion function requires a wrapper. */
 	MODELVIEWVIEWMODELBLUEPRINT_API bool RequiresWrapper(const UFunction* Function);
 
-	/** Find all BlueprintPropertyPath used by the given binding. */
-	MODELVIEWVIEWMODELBLUEPRINT_API TMap<FName, FMVVMBlueprintPropertyPath> GetAllArgumentPropertyPaths(const UBlueprint* WidgetBlueprint, const FMVVMBlueprintViewBinding& Binding, bool bSourceToDestination, bool bSkipResolve);
-
-	/** Find all BlueprintPropertyPath used by the given conversion function node. */
-	MODELVIEWVIEWMODELBLUEPRINT_API TMap<FName, FMVVMBlueprintPropertyPath> GetAllArgumentPropertyPaths(const UBlueprint* WidgetBlueprint, const UK2Node_CallFunction* FunctionNode, bool bSkipResolve);
+	/** The pin is valid to use with a PropertyPath. */
+	MODELVIEWVIEWMODELBLUEPRINT_API bool IsInputPin(const UEdGraphPin* Pin);
 	
 	/** Find the property path of a given argument in the conversion function. */
 	MODELVIEWVIEWMODELBLUEPRINT_API FMVVMBlueprintPropertyPath GetPropertyPathForPin(const UBlueprint* WidgetBlueprint, const UEdGraphPin* Pin, bool bSkipResolve);
@@ -67,8 +66,14 @@ namespace UE::MVVM::ConversionFunctionHelper
 	/** */
 	MODELVIEWVIEWMODELBLUEPRINT_API TPair<UEdGraph*, UK2Node*> CreateGraph(UBlueprint* WidgetBlueprint, FName GraphName, const UFunction* Signature, const TSubclassOf<UK2Node> Node, bool bIsConst, bool bTransient, TFunctionRef<void(UK2Node*)> InitNodeCallback);
 
+	/** Find the main conversion function node from the given graph. */
+	MODELVIEWVIEWMODELBLUEPRINT_API UK2Node* GetWrapperNode(const UEdGraph* Graph);
+
+	/** Find the conversion function node from the given graph. */
+	MODELVIEWVIEWMODELBLUEPRINT_API UEdGraphPin* FindPin(const UEdGraph* Graph, const TArrayView<const FName> PinNames);
+
 	/** Find the conversion function node from the given graph. */
-	MODELVIEWVIEWMODELBLUEPRINT_API UK2Node* GetWrapperNode(UEdGraph* Graph);
+	MODELVIEWVIEWMODELBLUEPRINT_API TArray<FName> FindPinId(const UEdGraphPin* GraphPin);
 } //namespace
 
 #if UE_ENABLE_INCLUDE_ORDER_DEPRECATED_IN_5_2
diff --git a/Engine/Plugins/Runtime/ModelViewViewModel/Source/ModelViewViewModelBlueprint/Public/MVVMBlueprintPin.h b/Engine/Plugins/Runtime/ModelViewViewModel/Source/ModelViewViewModelBlueprint/Public/MVVMBlueprintPin.h
--- a/Engine/Plugins/Runtime/ModelViewViewModel/Source/ModelViewViewModelBlueprint/Public/MVVMBlueprintPin.h
+++ b/Engine/Plugins/Runtime/ModelViewViewModel/Source/ModelViewViewModelBlueprint/Public/MVVMBlueprintPin.h
@@ -21,47 +21,99 @@ enum class EMVVMBlueprintPinStatus : uint8
 };
 
 /**
-*
-*/
+ * Pin name type to help with compare operation and moving it around.
+ */
 USTRUCT()
-struct MODELVIEWVIEWMODELBLUEPRINT_API FMVVMBlueprintPin
+struct MODELVIEWVIEWMODELBLUEPRINT_API FMVVMBlueprintPinId
 {
 	GENERATED_BODY()
 
+	FMVVMBlueprintPinId() = default;
+	explicit FMVVMBlueprintPinId(const TArrayView<const FName> Names);
+	explicit FMVVMBlueprintPinId(TArray<FName>&& Names);
+
+	bool IsValid() const;
+
+	const TArrayView<const FName> GetNames() const
+	{
+		return PinNames;
+	}
+
+	/** return true if the Pin is part of the Other pin. It can be a grand child. */
+	bool IsChildOf(const FMVVMBlueprintPinId& Other) const;
+
+	/** return true if the Pin is the directly child of the Other pin. It can be a child but not a grand child. */
+	bool IsDirectChildOf(const FMVVMBlueprintPinId& Other) const;
+
+	bool operator==(const FMVVMBlueprintPinId& Other) const;
+	bool operator==(const TArrayView<const FName> Other) const;
+
+	FString ToString() const;
+
 private:
-	UPROPERTY(VisibleAnywhere, Category = "Viewmodel")
-	FName PinName;
+	UPROPERTY(VisibleAnywhere, Category = "MVVM")
+	TArray<FName> PinNames;
+};
 
-	UPROPERTY(VisibleAnywhere, Category = "Viewmodel")
+/**
+ *
+ */
+USTRUCT()
+struct MODELVIEWVIEWMODELBLUEPRINT_API FMVVMBlueprintPin
+{
+	GENERATED_BODY()
+
+private:	
+	UPROPERTY(VisibleAnywhere, Category = "MVVM")
+	FMVVMBlueprintPinId PinId;
+
+	UPROPERTY(VisibleAnywhere, Category = "MVVM")
 	FMVVMBlueprintPropertyPath Path;
 
 	/** Default value for this pin (used if the pin has no connections), stored as a string */
-	UPROPERTY(VisibleAnywhere, Category = "Viewmodel")
+	UPROPERTY(VisibleAnywhere, Category = "MVVM")
 	FString DefaultString;
 
 	/** If the default value for this pin should be an FText, it is stored here. */
-	UPROPERTY(VisibleAnywhere, Category = "Viewmodel")
+	UPROPERTY(VisibleAnywhere, Category = "MVVM")
 	FText DefaultText;
 
 	/** If the default value for this pin should be an object, we store a pointer to it */
-	UPROPERTY(VisibleAnywhere, Category = "Viewmodel")
+	UPROPERTY(VisibleAnywhere, Category = "MVVM")
 	TObjectPtr<class UObject> DefaultObject;
 
 	/** The pin is split. */
-	UPROPERTY(VisibleAnywhere, Category = "Viewmodel")
+	UPROPERTY(VisibleAnywhere, Category = "MVVM")
 	bool bSplit = false;
 
 	/** The pin could not be set. */
-	UPROPERTY(VisibleAnywhere, Category = "Viewmodel")
+	UPROPERTY(VisibleAnywhere, Category = "MVVM")
 	mutable EMVVMBlueprintPinStatus Status = EMVVMBlueprintPinStatus::Valid;
 
+	UPROPERTY()
+	FName PinName_DEPRECATED;
+
 public:
 	FMVVMBlueprintPin() = default;
+	UE_DEPRECATED(5.4, "FMVVMBlueprintPin with a single name is deprecated. Use the TArrayView constructor instead")
 	FMVVMBlueprintPin(FName PinName);
+	explicit FMVVMBlueprintPin(FMVVMBlueprintPinId PinId);
+	explicit FMVVMBlueprintPin(const TArrayView<const FName> PinName);
 
+	UE_DEPRECATED(5.4, "GetName is deprecated. Use GetNames instead")
 	FName GetName() const
 	{
-		return PinName;
+		return PinId.GetNames().Num() > 0 ? PinId.GetNames().Last() : FName();
+	}
+
+	const FMVVMBlueprintPinId& GetId() const
+	{
+		return PinId;
+	}
+
+	bool IsValid() const
+	{
+		return PinId.IsValid();
 	}
 
 	/** The pin is split into its different components. */
@@ -82,6 +134,7 @@ public:
 		return !bSplit && Path.IsValid();
 	}
 
+	/** Get the path used by this pin. */
 	const FMVVMBlueprintPropertyPath& GetPath() const
 	{
 		return Path;
@@ -94,14 +147,26 @@ public:
 	void SetDefaultValue(const FString& Value);
 	void SetPath(const FMVVMBlueprintPropertyPath& Value);
 
+public:
+	void PostSerialize(const FArchive& Ar);
+
+public:
 	static bool IsInputPin(const UEdGraphPin* Pin);
 	static TArray<FMVVMBlueprintPin> CopyAndReturnMissingPins(UBlueprint* Blueprint, UEdGraphNode* GraphNode, const TArray<FMVVMBlueprintPin>& Pins);
 	static TArray<FMVVMBlueprintPin> CreateFromNode(UBlueprint* Blueprint, UEdGraphNode* GraphNode);
 	static FMVVMBlueprintPin CreateFromPin(const UBlueprint* Blueprint, const UEdGraphPin* Pin);
 
 	void CopyTo(const UBlueprint* WidgetBlueprint, UEdGraphNode* Node) const;
-	UEdGraphPin* FindGraphPin(UEdGraphNode* Node) const;
+	UEdGraphPin* FindGraphPin(const UEdGraph* Graph) const;
 
-private:
 	void Reset();
 };
+
+template<>
+struct TStructOpsTypeTraits<FMVVMBlueprintPin> : public TStructOpsTypeTraitsBase2<FMVVMBlueprintPin>
+{
+	enum
+	{
+		WithPostSerialize = true,
+	};
+};
diff --git a/Engine/Plugins/Runtime/ModelViewViewModel/Source/ModelViewViewModelBlueprint/Public/MVVMBlueprintViewConversionFunction.h b/Engine/Plugins/Runtime/ModelViewViewModel/Source/ModelViewViewModelBlueprint/Public/MVVMBlueprintViewConversionFunction.h
--- a/Engine/Plugins/Runtime/ModelViewViewModel/Source/ModelViewViewModelBlueprint/Public/MVVMBlueprintViewConversionFunction.h
+++ b/Engine/Plugins/Runtime/ModelViewViewModel/Source/ModelViewViewModelBlueprint/Public/MVVMBlueprintViewConversionFunction.h
@@ -96,15 +96,15 @@ public:
 	 * Returns the pin from the graph.
 	 * Create the graph and all the nodes for that graph if the graph doesn't exist and it's needed.
 	 */
-	UEdGraphPin* GetOrCreateGraphPin(UBlueprint* Blueprint, FName PinName);
+	UEdGraphPin* GetOrCreateGraphPin(UBlueprint* Blueprint, const FMVVMBlueprintPinId& PinId);
 
-	TArrayView<const FMVVMBlueprintPin> GetPins() const
+	const TArrayView<const FMVVMBlueprintPin> GetPins() const
 	{
 		return SavedPins;
 	}
 
 	/** */
-	void SetGraphPin(UBlueprint* Blueprint, FName PinName, const FMVVMBlueprintPropertyPath& Path);
+	void SetGraphPin(UBlueprint* Blueprint, const FMVVMBlueprintPinId& PinId, const FMVVMBlueprintPropertyPath& Value);
 
 	/** Generates SavedPins from the wrapper graph, if it exists. */
 	void SavePinValues(UBlueprint* Blueprint);
diff --git a/Engine/Plugins/Runtime/ModelViewViewModel/Source/ModelViewViewModelBlueprint/Public/MVVMBlueprintViewEvent.h b/Engine/Plugins/Runtime/ModelViewViewModel/Source/ModelViewViewModelBlueprint/Public/MVVMBlueprintViewEvent.h
--- a/Engine/Plugins/Runtime/ModelViewViewModel/Source/ModelViewViewModelBlueprint/Public/MVVMBlueprintViewEvent.h
+++ b/Engine/Plugins/Runtime/ModelViewViewModel/Source/ModelViewViewModelBlueprint/Public/MVVMBlueprintViewEvent.h
@@ -92,12 +92,12 @@ public:
 	/** Keep the orphaned pins. Add the missing pins. */
 	bool HasOrphanedPin() const;
 
-	UEdGraphPin* GetOrCreateGraphPin(FName PinName);
+	UEdGraphPin* GetOrCreateGraphPin(const FMVVMBlueprintPinId& Pin);
 
-	FMVVMBlueprintPropertyPath GetPinPath(FName PinName) const;
-	void SetPinPath(FName PinName, const FMVVMBlueprintPropertyPath& Path);
+	FMVVMBlueprintPropertyPath GetPinPath(const FMVVMBlueprintPinId& Pin) const;
+	void SetPinPath(const FMVVMBlueprintPinId& Pin, const FMVVMBlueprintPropertyPath& Path);
 	// To set a pin when loading the asset (no graph generation)
-	void SetPinPathNoGraphGeneration(FName PinName, const FMVVMBlueprintPropertyPath& Path);
+	void SetPinPathNoGraphGeneration(const FMVVMBlueprintPinId& Pin, const FMVVMBlueprintPropertyPath& Path);
 
 	FSimpleMulticastDelegate OnWrapperGraphModified;
 
diff --git a/Engine/Plugins/Runtime/ModelViewViewModel/Source/ModelViewViewModelEditor/Private/Framework/MVVMRowHelper.cpp b/Engine/Plugins/Runtime/ModelViewViewModel/Source/ModelViewViewModelEditor/Private/Framework/MVVMRowHelper.cpp
--- a/Engine/Plugins/Runtime/ModelViewViewModel/Source/ModelViewViewModelEditor/Private/Framework/MVVMRowHelper.cpp
+++ b/Engine/Plugins/Runtime/ModelViewViewModel/Source/ModelViewViewModelEditor/Private/Framework/MVVMRowHelper.cpp
@@ -292,14 +292,14 @@ void HandleResetSelectedPin(UWidgetBlueprint* WidgetBlueprint, UMVVMBlueprintVie
 	{
 		if (Entry->GetRowType() == FBindingEntry::ERowType::EventParameter)
 		{
-			EditorSubsystem->ResetPinToDefaultValue(WidgetBlueprint, Entry->GetEvent(), Entry->GetEventParameterName());
+			EditorSubsystem->ResetPinToDefaultValue(WidgetBlueprint, Entry->GetEvent(), Entry->GetEventParameterId());
 		}
 		else if (Entry->GetRowType() == FBindingEntry::ERowType::BindingParameter)
 		{
 			if (FMVVMBlueprintViewBinding* Binding = Entry->GetBinding(View))
 			{
 				const bool bSourceToDestination = UE::MVVM::IsForwardBinding(Binding->BindingType);
-				EditorSubsystem->ResetPinToDefaultValue(WidgetBlueprint, *Binding, Entry->GetBindingParameterName(), bSourceToDestination);
+				EditorSubsystem->ResetPinToDefaultValue(WidgetBlueprint, *Binding, Entry->GetBindingParameterId(), bSourceToDestination);
 			}
 		}
 	}
@@ -312,14 +312,14 @@ void HandleBreakSelectedPin(UWidgetBlueprint* WidgetBlueprint, UMVVMBlueprintVie
 	{
 		if (Entry->GetRowType() == FBindingEntry::ERowType::EventParameter)
 		{
-			EditorSubsystem->SplitPin(WidgetBlueprint, Entry->GetEvent(), Entry->GetEventParameterName());
+			EditorSubsystem->SplitPin(WidgetBlueprint, Entry->GetEvent(), Entry->GetEventParameterId());
 		}
 		else if (Entry->GetRowType() == FBindingEntry::ERowType::BindingParameter)
 		{
 			if (FMVVMBlueprintViewBinding* Binding = Entry->GetBinding(View))
 			{
 				const bool bSourceToDestination = UE::MVVM::IsForwardBinding(Binding->BindingType);
-				EditorSubsystem->SplitPin(WidgetBlueprint, *Binding, Entry->GetBindingParameterName(), bSourceToDestination);
+				EditorSubsystem->SplitPin(WidgetBlueprint, *Binding, Entry->GetBindingParameterId(), bSourceToDestination);
 			}
 		}
 	}
@@ -332,14 +332,14 @@ void HandleRecombineSelectedPin(UWidgetBlueprint* WidgetBlueprint, UMVVMBlueprin
 	{
 		if (Entry->GetRowType() == FBindingEntry::ERowType::EventParameter)
 		{
-			EditorSubsystem->RecombinePin(WidgetBlueprint, Entry->GetEvent(), Entry->GetEventParameterName());
+			EditorSubsystem->RecombinePin(WidgetBlueprint, Entry->GetEvent(), Entry->GetEventParameterId());
 		}
 		else if (Entry->GetRowType() == FBindingEntry::ERowType::BindingParameter)
 		{
 			if (FMVVMBlueprintViewBinding* Binding = Entry->GetBinding(View))
 			{
 				const bool bSourceToDestination = UE::MVVM::IsForwardBinding(Binding->BindingType);
-				EditorSubsystem->RecombinePin(WidgetBlueprint, *Binding, Entry->GetBindingParameterName(), bSourceToDestination);
+				EditorSubsystem->RecombinePin(WidgetBlueprint, *Binding, Entry->GetBindingParameterId(), bSourceToDestination);
 			}
 		}
 	}
@@ -352,14 +352,14 @@ void HandleResetOrphanedSelectedPin(UWidgetBlueprint* WidgetBlueprint, UMVVMBlue
 	{
 		if (Entry->GetRowType() == FBindingEntry::ERowType::EventParameter)
 		{
-			EditorSubsystem->ResetOrphanedPin(WidgetBlueprint, Entry->GetEvent(), Entry->GetEventParameterName());
+			EditorSubsystem->ResetOrphanedPin(WidgetBlueprint, Entry->GetEvent(), Entry->GetEventParameterId());
 		}
 		else if (Entry->GetRowType() == FBindingEntry::ERowType::BindingParameter)
 		{
 			if (FMVVMBlueprintViewBinding* Binding = Entry->GetBinding(View))
 			{
 				const bool bSourceToDestination = UE::MVVM::IsForwardBinding(Binding->BindingType);
-				EditorSubsystem->ResetOrphanedPin(WidgetBlueprint, *Binding, Entry->GetBindingParameterName(), bSourceToDestination);
+				EditorSubsystem->ResetOrphanedPin(WidgetBlueprint, *Binding, Entry->GetBindingParameterId(), bSourceToDestination);
 			}
 		}
 	}
@@ -433,25 +433,25 @@ FMenuBuilder FRowHelper::CreateContextMenu(UWidgetBlueprint* WidgetBlueprint, UM
 			{
 				if (Entry->GetRowType() == FBindingEntry::ERowType::EventParameter)
 				{
-					if (!EditorSubsystem->CanSplitPin(WidgetBlueprint, Entry->GetEvent(), Entry->GetEventParameterName()))
+					if (!EditorSubsystem->CanSplitPin(WidgetBlueprint, Entry->GetEvent(), Entry->GetEventParameterId()))
 					{
 						bCanSplitPin = false;
 					}
-					if (!EditorSubsystem->CanRecombinePin(WidgetBlueprint, Entry->GetEvent(), Entry->GetEventParameterName()))
+					if (!EditorSubsystem->CanRecombinePin(WidgetBlueprint, Entry->GetEvent(), Entry->GetEventParameterId()))
 					{
 						bCanRecombinePin = false;
 					}
-					if (!EditorSubsystem->CanResetPinToDefaultValue(WidgetBlueprint, Entry->GetEvent(), Entry->GetEventParameterName()))
+					if (!EditorSubsystem->CanResetPinToDefaultValue(WidgetBlueprint, Entry->GetEvent(), Entry->GetEventParameterId()))
 					{
 						bCanResetPin = false;
 					}
-					if (!EditorSubsystem->CanResetOrphanedPin(WidgetBlueprint, Entry->GetEvent(), Entry->GetEventParameterName()))
+					if (!EditorSubsystem->CanResetOrphanedPin(WidgetBlueprint, Entry->GetEvent(), Entry->GetEventParameterId()))
 					{
 						bCanResetOrphanedPin = false;
 					}
 
 					UMVVMBlueprintViewEvent* ViewEvent = Entry->GetEvent();
-					UEdGraphPin* GraphPin = ViewEvent ? ViewEvent->GetOrCreateGraphPin(Entry->GetEventParameterName()) : nullptr;
+					UEdGraphPin* GraphPin = ViewEvent ? ViewEvent->GetOrCreateGraphPin(Entry->GetEventParameterId()) : nullptr;
 					if (GraphPin == nullptr)
 					{
 						bCanRecombinePinVisible = false;
@@ -475,24 +475,24 @@ FMenuBuilder FRowHelper::CreateContextMenu(UWidgetBlueprint* WidgetBlueprint, UM
 					if (Binding)
 					{
 						const bool bSourceToDestination = UE::MVVM::IsForwardBinding(Binding->BindingType);
-						if (!EditorSubsystem->CanSplitPin(WidgetBlueprint, *Binding, Entry->GetBindingParameterName(), bSourceToDestination))
+						if (!EditorSubsystem->CanSplitPin(WidgetBlueprint, *Binding, Entry->GetBindingParameterId(), bSourceToDestination))
 						{
 							bCanSplitPin = false;
 						}
-						if (!EditorSubsystem->CanRecombinePin(WidgetBlueprint, *Binding, Entry->GetBindingParameterName(), bSourceToDestination))
+						if (!EditorSubsystem->CanRecombinePin(WidgetBlueprint, *Binding, Entry->GetBindingParameterId(), bSourceToDestination))
 						{
 							bCanRecombinePin = false;
 						}
-						if (!EditorSubsystem->CanResetPinToDefaultValue(WidgetBlueprint, *Binding, Entry->GetBindingParameterName(), bSourceToDestination))
+						if (!EditorSubsystem->CanResetPinToDefaultValue(WidgetBlueprint, *Binding, Entry->GetBindingParameterId(), bSourceToDestination))
 						{
 							bCanResetPin = false;
 						}
-						if (!EditorSubsystem->CanResetOrphanedPin(WidgetBlueprint, *Binding, Entry->GetBindingParameterName(), bSourceToDestination))
+						if (!EditorSubsystem->CanResetOrphanedPin(WidgetBlueprint, *Binding, Entry->GetBindingParameterId(), bSourceToDestination))
 						{
 							bCanResetOrphanedPin = false;
 						}
 
-						UEdGraphPin* GraphPin = EditorSubsystem->GetConversionFunctionArgumentPin(WidgetBlueprint, *Binding, Entry->GetBindingParameterName(), bSourceToDestination);
+						UEdGraphPin* GraphPin = EditorSubsystem->GetConversionFunctionArgumentPin(WidgetBlueprint, *Binding, Entry->GetBindingParameterId(), bSourceToDestination);
 						if (GraphPin == nullptr)
 						{
 							bCanRecombinePinVisible = false;
diff --git a/Engine/Plugins/Runtime/ModelViewViewModel/Source/ModelViewViewModelEditor/Private/MVVMEditorSubsystem.cpp b/Engine/Plugins/Runtime/ModelViewViewModel/Source/ModelViewViewModelEditor/Private/MVVMEditorSubsystem.cpp
--- a/Engine/Plugins/Runtime/ModelViewViewModel/Source/ModelViewViewModelEditor/Private/MVVMEditorSubsystem.cpp
+++ b/Engine/Plugins/Runtime/ModelViewViewModel/Source/ModelViewViewModelEditor/Private/MVVMEditorSubsystem.cpp
@@ -347,11 +347,11 @@ UFunction* UMVVMEditorSubsystem::GetConversionFunction(const UWidgetBlueprint* W
 	return nullptr;
 }
 
-UEdGraphPin* UMVVMEditorSubsystem::GetConversionFunctionArgumentPin(const UWidgetBlueprint* WidgetBlueprint, const FMVVMBlueprintViewBinding& Binding, FName ParameterName, bool bSourceToDestination) const
+UEdGraphPin* UMVVMEditorSubsystem::GetConversionFunctionArgumentPin(const UWidgetBlueprint* WidgetBlueprint, const FMVVMBlueprintViewBinding& Binding, const FMVVMBlueprintPinId& ParameterId, bool bSourceToDestination) const
 {
 	if (UMVVMBlueprintViewConversionFunction* ConversionFunction = Binding.Conversion.GetConversionFunction(bSourceToDestination))
 	{
-		return ConversionFunction->GetOrCreateGraphPin(const_cast<UWidgetBlueprint*>(WidgetBlueprint), ParameterName);
+		return ConversionFunction->GetOrCreateGraphPin(const_cast<UWidgetBlueprint*>(WidgetBlueprint), ParameterId);
 	}
 	return nullptr;
 }
@@ -611,7 +611,7 @@ void UMVVMEditorSubsystem::SetEventDestinationPath(UMVVMBlueprintViewEvent* Even
 	}
 }
 
-void UMVVMEditorSubsystem::SetEventArgumentPath(UMVVMBlueprintViewEvent* Event, FName ArgumentName, const FMVVMBlueprintPropertyPath& Path) const
+void UMVVMEditorSubsystem::SetEventArgumentPath(UMVVMBlueprintViewEvent* Event,  const FMVVMBlueprintPinId& ParameterId, const FMVVMBlueprintPropertyPath& Path) const
 {
 	const UMVVMBlueprintView* View = Event ? Event->GetOuterUMVVMBlueprintView() : nullptr;
 	if (View)
@@ -620,7 +620,7 @@ void UMVVMEditorSubsystem::SetEventArgumentPath(UMVVMBlueprintViewEvent* Event,
 
 		UE::MVVM::Private::OnEventPreEditChange(Event, "SavedPins");
 	
-		Event->SetPinPath(ArgumentName, Path);
+		Event->SetPinPath(ParameterId, Path);
 
 		UE::MVVM::Private::OnEventPostEditChange(Event, "SavedPins");
 	}
@@ -848,7 +848,7 @@ TArray<UFunction*> UMVVMEditorSubsystem::GetAvailableConversionFunctions(const U
 	return ConversionFunctions;
 }
 
-FMVVMBlueprintPropertyPath UMVVMEditorSubsystem::GetPathForConversionFunctionArgument(const UWidgetBlueprint* WidgetBlueprint, const FMVVMBlueprintViewBinding& Binding, FName ArgumentName, bool bSourceToDestination) const
+FMVVMBlueprintPropertyPath UMVVMEditorSubsystem::GetPathForConversionFunctionArgument(const UWidgetBlueprint* WidgetBlueprint, const FMVVMBlueprintViewBinding& Binding,  const FMVVMBlueprintPinId& ParameterId, bool bSourceToDestination) const
 {
 	UMVVMBlueprintViewConversionFunction* ConversionFunction = Binding.Conversion.GetConversionFunction(bSourceToDestination);
 	if (ConversionFunction == nullptr)
@@ -856,7 +856,7 @@ FMVVMBlueprintPropertyPath UMVVMEditorSubsystem::GetPathForConversionFunctionArg
 		return FMVVMBlueprintPropertyPath();
 	}
 
-	UEdGraphPin* GraphPin = ConversionFunction->GetOrCreateGraphPin(const_cast<UWidgetBlueprint*>(WidgetBlueprint), ArgumentName);
+	UEdGraphPin* GraphPin = ConversionFunction->GetOrCreateGraphPin(const_cast<UWidgetBlueprint*>(WidgetBlueprint), ParameterId);
 	if (GraphPin == nullptr)
 	{
 		return FMVVMBlueprintPropertyPath();
@@ -865,44 +865,44 @@ FMVVMBlueprintPropertyPath UMVVMEditorSubsystem::GetPathForConversionFunctionArg
 	return UE::MVVM::ConversionFunctionHelper::GetPropertyPathForPin(WidgetBlueprint, GraphPin, false);
 }
 
-void UMVVMEditorSubsystem::SetPathForConversionFunctionArgument(UWidgetBlueprint* WidgetBlueprint, FMVVMBlueprintViewBinding& Binding, FName ArgumentName, const FMVVMBlueprintPropertyPath& Path, bool bSourceToDestination) const
+void UMVVMEditorSubsystem::SetPathForConversionFunctionArgument(UWidgetBlueprint* WidgetBlueprint, FMVVMBlueprintViewBinding& Binding,  const FMVVMBlueprintPinId& ParameterId, const FMVVMBlueprintPropertyPath& Path, bool bSourceToDestination) const
 {
 	UMVVMBlueprintViewConversionFunction* ConversionFunction = Binding.Conversion.GetConversionFunction(bSourceToDestination);
 	if (ConversionFunction)
 	{
-		ConversionFunction->SetGraphPin(WidgetBlueprint, ArgumentName, Path);
+		ConversionFunction->SetGraphPin(WidgetBlueprint, ParameterId, Path);
 	}
 	FBlueprintEditorUtils::MarkBlueprintAsModified(WidgetBlueprint);
 }
 
 namespace Private
 {
-	UEdGraphPin* GetGraphPin(const UMVVMEditorSubsystem* Subsystem, const UWidgetBlueprint* WidgetBlueprint, const FMVVMBlueprintViewBinding& ViewBinding, FName PinName, bool bSourceToDestination)
+	UEdGraphPin* GetGraphPin(const UMVVMEditorSubsystem* Subsystem, const UWidgetBlueprint* WidgetBlueprint, const FMVVMBlueprintViewBinding& ViewBinding, const FMVVMBlueprintPinId& ParameterId, bool bSourceToDestination)
 	{
 		if (WidgetBlueprint == nullptr)
 		{
 			return nullptr;
 		}
-		return Subsystem->GetConversionFunctionArgumentPin(WidgetBlueprint, ViewBinding, PinName, bSourceToDestination);
+		return Subsystem->GetConversionFunctionArgumentPin(WidgetBlueprint, ViewBinding, ParameterId, bSourceToDestination);
 	}
 
-	UEdGraphPin* GetGraphPin(const UWidgetBlueprint* WidgetBlueprint, UMVVMBlueprintViewEvent* ViewEvent, FName PinName)
+	UEdGraphPin* GetGraphPin(const UWidgetBlueprint* WidgetBlueprint, UMVVMBlueprintViewEvent* ViewEvent,  const FMVVMBlueprintPinId& ParameterId)
 	{
 		if (WidgetBlueprint == nullptr || ViewEvent == nullptr)
 		{
 			return nullptr;
 		}
 
-		return ViewEvent->GetOrCreateGraphPin(PinName);
+		return ViewEvent->GetOrCreateGraphPin(ParameterId);
 	}
 
-	void DoAction(const UMVVMEditorSubsystem* Subsystem, UWidgetBlueprint* WidgetBlueprint, UMVVMBlueprintViewEvent* ViewEvent, FMVVMBlueprintViewBinding* Binding, FName PinName, bool bSourceToDestination
+	void DoAction(const UMVVMEditorSubsystem* Subsystem, UWidgetBlueprint* WidgetBlueprint, UMVVMBlueprintViewEvent* ViewEvent, FMVVMBlueprintViewBinding* Binding,  const FMVVMBlueprintPinId& ParameterId, bool bSourceToDestination
 		, FText TransactionName
 		, TFunctionRef<bool(const UEdGraphSchema_K2*, UEdGraphPin*)> Test, TFunctionRef<void(const UEdGraphSchema_K2*, UEdGraphPin*)> Action)
 	{
 		UEdGraphPin* GraphPin = ViewEvent
-			? Private::GetGraphPin(WidgetBlueprint, ViewEvent, PinName)
-			: Private::GetGraphPin(Subsystem, WidgetBlueprint, *Binding, PinName, bSourceToDestination);
+			? Private::GetGraphPin(WidgetBlueprint, ViewEvent, ParameterId)
+			: Private::GetGraphPin(Subsystem, WidgetBlueprint, *Binding, ParameterId, bSourceToDestination);
 		const UEdGraphSchema_K2* K2Schema = GetDefault<UEdGraphSchema_K2>();
 		if (!Test(K2Schema, GraphPin))
 		{
@@ -933,9 +933,9 @@ namespace Private
 		}
 	}
 
-	void SplitPin(const UMVVMEditorSubsystem* Subsystem, UWidgetBlueprint* WidgetBlueprint, UMVVMBlueprintViewEvent* ViewEvent, FMVVMBlueprintViewBinding* Binding, FName PinName, bool bSourceToDestination)
+	void SplitPin(const UMVVMEditorSubsystem* Subsystem, UWidgetBlueprint* WidgetBlueprint, UMVVMBlueprintViewEvent* ViewEvent, FMVVMBlueprintViewBinding* Binding,  const FMVVMBlueprintPinId& ParameterId, bool bSourceToDestination)
 	{
-		DoAction(Subsystem, WidgetBlueprint, ViewEvent, Binding, PinName, bSourceToDestination, LOCTEXT("BreakPin", "Split Struct Pin")
+		DoAction(Subsystem, WidgetBlueprint, ViewEvent, Binding, ParameterId, bSourceToDestination, LOCTEXT("BreakPin", "Split Struct Pin")
 			, [](const UEdGraphSchema_K2* K2Schema, UEdGraphPin* GraphPin)->bool
 			{
 				return GraphPin != nullptr && K2Schema->CanSplitStructPin(*GraphPin);
@@ -952,9 +952,9 @@ namespace Private
 		return GraphPin ? K2Schema->CanSplitStructPin(*GraphPin) && !GraphPin->bOrphanedPin : false;
 	}
 	
-	void RecombinePin(const UMVVMEditorSubsystem* Subsystem, UWidgetBlueprint* WidgetBlueprint, UMVVMBlueprintViewEvent* ViewEvent, FMVVMBlueprintViewBinding* Binding, FName PinName, bool bSourceToDestination)
+	void RecombinePin(const UMVVMEditorSubsystem* Subsystem, UWidgetBlueprint* WidgetBlueprint, UMVVMBlueprintViewEvent* ViewEvent, FMVVMBlueprintViewBinding* Binding,  const FMVVMBlueprintPinId& ParameterId, bool bSourceToDestination)
 	{
-		DoAction(Subsystem, WidgetBlueprint, ViewEvent, Binding, PinName, bSourceToDestination, LOCTEXT("BreakPin", "Split Struct Pin")
+		DoAction(Subsystem, WidgetBlueprint, ViewEvent, Binding, ParameterId, bSourceToDestination, LOCTEXT("BreakPin", "Split Struct Pin")
 			, [](const UEdGraphSchema_K2* K2Schema, UEdGraphPin* GraphPin)->bool
 			{
 				return GraphPin != nullptr && K2Schema->CanRecombineStructPin(*GraphPin);
@@ -971,9 +971,9 @@ namespace Private
 		return GraphPin ? K2Schema->CanRecombineStructPin(*GraphPin) && !GraphPin->bOrphanedPin : false;
 	}
 	
-	void ResetPinToDefaultValue(const UMVVMEditorSubsystem* Subsystem, UWidgetBlueprint* WidgetBlueprint, UMVVMBlueprintViewEvent* ViewEvent, FMVVMBlueprintViewBinding* Binding, FName PinName, bool bSourceToDestination)
+	void ResetPinToDefaultValue(const UMVVMEditorSubsystem* Subsystem, UWidgetBlueprint* WidgetBlueprint, UMVVMBlueprintViewEvent* ViewEvent, FMVVMBlueprintViewBinding* Binding,  const FMVVMBlueprintPinId& ParameterId, bool bSourceToDestination)
 	{
-		DoAction(Subsystem, WidgetBlueprint, ViewEvent, Binding, PinName, bSourceToDestination, LOCTEXT("BreakPin", "Split Struct Pin")
+		DoAction(Subsystem, WidgetBlueprint, ViewEvent, Binding, ParameterId, bSourceToDestination, LOCTEXT("BreakPin", "Split Struct Pin")
 			, [](const UEdGraphSchema_K2* K2Schema, UEdGraphPin* GraphPin)->bool
 			{
 				return GraphPin != nullptr && !K2Schema->DoesDefaultValueMatchAutogenerated(*GraphPin);
@@ -990,9 +990,9 @@ namespace Private
 		return GraphPin ? !K2Schema->DoesDefaultValueMatchAutogenerated(*GraphPin) && !GraphPin->bOrphanedPin : false;
 	}
 
-	void ResetOrphanedPin(const UMVVMEditorSubsystem* Subsystem, UWidgetBlueprint* WidgetBlueprint, UMVVMBlueprintViewEvent* ViewEvent, FMVVMBlueprintViewBinding* Binding, FName PinName, bool bSourceToDestination)
+	void ResetOrphanedPin(const UMVVMEditorSubsystem* Subsystem, UWidgetBlueprint* WidgetBlueprint, UMVVMBlueprintViewEvent* ViewEvent, FMVVMBlueprintViewBinding* Binding,  const FMVVMBlueprintPinId& ParameterId, bool bSourceToDestination)
 	{
-		DoAction(Subsystem, WidgetBlueprint, ViewEvent, Binding, PinName, bSourceToDestination, LOCTEXT("BreakPin", "Split Struct Pin")
+		DoAction(Subsystem, WidgetBlueprint, ViewEvent, Binding, ParameterId, bSourceToDestination, LOCTEXT("BreakPin", "Split Struct Pin")
 			, [](const UEdGraphSchema_K2* K2Schema, UEdGraphPin* GraphPin)->bool
 			{
 				return GraphPin == nullptr || GraphPin->bOrphanedPin;
@@ -1012,91 +1012,91 @@ namespace Private
 	}
 } //namespace
 
-void UMVVMEditorSubsystem::SplitPin(UWidgetBlueprint* WidgetBlueprint, FMVVMBlueprintViewBinding& Binding, FName PinName, bool bSourceToDestination) const
+void UMVVMEditorSubsystem::SplitPin(UWidgetBlueprint* WidgetBlueprint, FMVVMBlueprintViewBinding& Binding,  const FMVVMBlueprintPinId& ParameterId, bool bSourceToDestination) const
 {
-	Private::SplitPin(this, WidgetBlueprint, nullptr, &Binding, PinName, bSourceToDestination);
+	Private::SplitPin(this, WidgetBlueprint, nullptr, &Binding, ParameterId, bSourceToDestination);
 }
 
-bool UMVVMEditorSubsystem::CanSplitPin(const UWidgetBlueprint* WidgetBlueprint, const FMVVMBlueprintViewBinding& Binding, FName PinName, bool bSourceToDestination) const
+bool UMVVMEditorSubsystem::CanSplitPin(const UWidgetBlueprint* WidgetBlueprint, const FMVVMBlueprintViewBinding& Binding,  const FMVVMBlueprintPinId& ParameterId, bool bSourceToDestination) const
 {
-	UEdGraphPin* GraphPin = Private::GetGraphPin(this, WidgetBlueprint, Binding, PinName, bSourceToDestination);
+	UEdGraphPin* GraphPin = Private::GetGraphPin(this, WidgetBlueprint, Binding, ParameterId, bSourceToDestination);
 	return Private::CanSplitPin(GraphPin);
 }
 
-void UMVVMEditorSubsystem::SplitPin(UWidgetBlueprint* WidgetBlueprint, UMVVMBlueprintViewEvent* ViewEvent, FName PinName) const
+void UMVVMEditorSubsystem::SplitPin(UWidgetBlueprint* WidgetBlueprint, UMVVMBlueprintViewEvent* ViewEvent,  const FMVVMBlueprintPinId& ParameterId) const
 {
-	Private::SplitPin(this, WidgetBlueprint, ViewEvent, nullptr, PinName, true);
+	Private::SplitPin(this, WidgetBlueprint, ViewEvent, nullptr, ParameterId, true);
 }
 
-bool UMVVMEditorSubsystem::CanSplitPin(const UWidgetBlueprint* WidgetBlueprint, UMVVMBlueprintViewEvent* Event, FName PinName) const
+bool UMVVMEditorSubsystem::CanSplitPin(const UWidgetBlueprint* WidgetBlueprint, UMVVMBlueprintViewEvent* Event,  const FMVVMBlueprintPinId& ParameterId) const
 {
-	UEdGraphPin* GraphPin = Private::GetGraphPin(WidgetBlueprint, Event, PinName);
+	UEdGraphPin* GraphPin = Private::GetGraphPin(WidgetBlueprint, Event, ParameterId);
 	return Private::CanSplitPin(GraphPin);
 }
 
-void UMVVMEditorSubsystem::RecombinePin(UWidgetBlueprint* WidgetBlueprint, FMVVMBlueprintViewBinding& Binding, FName PinName, bool bSourceToDestination) const
+void UMVVMEditorSubsystem::RecombinePin(UWidgetBlueprint* WidgetBlueprint, FMVVMBlueprintViewBinding& Binding,  const FMVVMBlueprintPinId& ParameterId, bool bSourceToDestination) const
 {
-	Private::RecombinePin(this, WidgetBlueprint, nullptr, &Binding, PinName, bSourceToDestination);
+	Private::RecombinePin(this, WidgetBlueprint, nullptr, &Binding, ParameterId, bSourceToDestination);
 }
 
-bool UMVVMEditorSubsystem::CanRecombinePin(const UWidgetBlueprint* WidgetBlueprint, const FMVVMBlueprintViewBinding& Binding, FName PinName, bool bSourceToDestination) const
+bool UMVVMEditorSubsystem::CanRecombinePin(const UWidgetBlueprint* WidgetBlueprint, const FMVVMBlueprintViewBinding& Binding,  const FMVVMBlueprintPinId& ParameterId, bool bSourceToDestination) const
 {
-	UEdGraphPin* GraphPin = Private::GetGraphPin(this, WidgetBlueprint, Binding, PinName, bSourceToDestination);
+	UEdGraphPin* GraphPin = Private::GetGraphPin(this, WidgetBlueprint, Binding, ParameterId, bSourceToDestination);
 	return Private::CanRecombinePin(GraphPin);
 }
 
-void UMVVMEditorSubsystem::RecombinePin(UWidgetBlueprint* WidgetBlueprint, UMVVMBlueprintViewEvent* ViewEvent, FName PinName) const
+void UMVVMEditorSubsystem::RecombinePin(UWidgetBlueprint* WidgetBlueprint, UMVVMBlueprintViewEvent* ViewEvent,   const FMVVMBlueprintPinId& ParameterId) const
 {
-	Private::RecombinePin(this, WidgetBlueprint, ViewEvent, nullptr, PinName, true);
+	Private::RecombinePin(this, WidgetBlueprint, ViewEvent, nullptr, ParameterId, true);
 }
 
-bool UMVVMEditorSubsystem::CanRecombinePin(const UWidgetBlueprint* WidgetBlueprint, UMVVMBlueprintViewEvent* Event, FName PinName) const
+bool UMVVMEditorSubsystem::CanRecombinePin(const UWidgetBlueprint* WidgetBlueprint, UMVVMBlueprintViewEvent* Event,  const FMVVMBlueprintPinId& ParameterId) const
 {
-	UEdGraphPin* GraphPin = Private::GetGraphPin(WidgetBlueprint, Event, PinName);
+	UEdGraphPin* GraphPin = Private::GetGraphPin(WidgetBlueprint, Event, ParameterId);
 	return Private::CanRecombinePin(GraphPin);
 }
 
-void UMVVMEditorSubsystem::ResetPinToDefaultValue(UWidgetBlueprint* WidgetBlueprint, FMVVMBlueprintViewBinding& Binding, FName PinName, bool bSourceToDestination) const
+void UMVVMEditorSubsystem::ResetPinToDefaultValue(UWidgetBlueprint* WidgetBlueprint, FMVVMBlueprintViewBinding& Binding,  const FMVVMBlueprintPinId& ParameterId, bool bSourceToDestination) const
 {
-	Private::ResetPinToDefaultValue(this, WidgetBlueprint, nullptr, &Binding, PinName, bSourceToDestination);
+	Private::ResetPinToDefaultValue(this, WidgetBlueprint, nullptr, &Binding, ParameterId, bSourceToDestination);
 }
 
-bool UMVVMEditorSubsystem::CanResetPinToDefaultValue(const UWidgetBlueprint* WidgetBlueprint, const FMVVMBlueprintViewBinding& Binding, FName PinName, bool bSourceToDestination) const
+bool UMVVMEditorSubsystem::CanResetPinToDefaultValue(const UWidgetBlueprint* WidgetBlueprint, const FMVVMBlueprintViewBinding& Binding,  const FMVVMBlueprintPinId& ParameterId, bool bSourceToDestination) const
 {
-	UEdGraphPin* GraphPin = Private::GetGraphPin(this, WidgetBlueprint, Binding, PinName, bSourceToDestination);
+	UEdGraphPin* GraphPin = Private::GetGraphPin(this, WidgetBlueprint, Binding, ParameterId, bSourceToDestination);
 	return Private::CanResetPinToDefaultValue(GraphPin);
 }
 
-void UMVVMEditorSubsystem::ResetPinToDefaultValue(UWidgetBlueprint* WidgetBlueprint, UMVVMBlueprintViewEvent* ViewEvent, FName PinName) const
+void UMVVMEditorSubsystem::ResetPinToDefaultValue(UWidgetBlueprint* WidgetBlueprint, UMVVMBlueprintViewEvent* ViewEvent,  const FMVVMBlueprintPinId& ParameterId) const
 {
-	Private::ResetPinToDefaultValue(this, WidgetBlueprint, ViewEvent, nullptr, PinName, true);
+	Private::ResetPinToDefaultValue(this, WidgetBlueprint, ViewEvent, nullptr, ParameterId, true);
 }
 
-bool UMVVMEditorSubsystem::CanResetPinToDefaultValue(const UWidgetBlueprint* WidgetBlueprint, UMVVMBlueprintViewEvent* Event, FName PinName) const
+bool UMVVMEditorSubsystem::CanResetPinToDefaultValue(const UWidgetBlueprint* WidgetBlueprint, UMVVMBlueprintViewEvent* Event,  const FMVVMBlueprintPinId& ParameterId) const
 {
-	UEdGraphPin* GraphPin = Private::GetGraphPin(WidgetBlueprint, Event, PinName);
+	UEdGraphPin* GraphPin = Private::GetGraphPin(WidgetBlueprint, Event, ParameterId);
 	return Private::CanResetPinToDefaultValue(GraphPin);
 }
 
-void UMVVMEditorSubsystem::ResetOrphanedPin(UWidgetBlueprint* WidgetBlueprint, FMVVMBlueprintViewBinding& Binding, FName PinName, bool bSourceToDestination) const
+void UMVVMEditorSubsystem::ResetOrphanedPin(UWidgetBlueprint* WidgetBlueprint, FMVVMBlueprintViewBinding& Binding, const FMVVMBlueprintPinId& ParameterId, bool bSourceToDestination) const
 {
-	Private::ResetOrphanedPin(this, WidgetBlueprint, nullptr, &Binding, PinName, bSourceToDestination);
+	Private::ResetOrphanedPin(this, WidgetBlueprint, nullptr, &Binding, ParameterId, bSourceToDestination);
 }
 
-bool UMVVMEditorSubsystem::CanResetOrphanedPin(const UWidgetBlueprint* WidgetBlueprint, const FMVVMBlueprintViewBinding& Binding, FName PinName, bool bSourceToDestination) const
+bool UMVVMEditorSubsystem::CanResetOrphanedPin(const UWidgetBlueprint* WidgetBlueprint, const FMVVMBlueprintViewBinding& Binding, const FMVVMBlueprintPinId& ParameterId, bool bSourceToDestination) const
 {
-	UEdGraphPin* GraphPin = Private::GetGraphPin(this, WidgetBlueprint, Binding, PinName, bSourceToDestination);
+	UEdGraphPin* GraphPin = Private::GetGraphPin(this, WidgetBlueprint, Binding, ParameterId, bSourceToDestination);
 	return Private::CanResetOrphanedPin(GraphPin);
 }
 
-void UMVVMEditorSubsystem::ResetOrphanedPin(UWidgetBlueprint* WidgetBlueprint, UMVVMBlueprintViewEvent* ViewEvent, FName PinName) const
+void UMVVMEditorSubsystem::ResetOrphanedPin(UWidgetBlueprint* WidgetBlueprint, UMVVMBlueprintViewEvent* ViewEvent, const FMVVMBlueprintPinId& ParameterId) const
 {
-	Private::ResetOrphanedPin(this, WidgetBlueprint, ViewEvent, nullptr, PinName, true);
+	Private::ResetOrphanedPin(this, WidgetBlueprint, ViewEvent, nullptr, ParameterId, true);
 }
 
-bool UMVVMEditorSubsystem::CanResetOrphanedPin(const UWidgetBlueprint* WidgetBlueprint, UMVVMBlueprintViewEvent* Event, FName PinName) const
+bool UMVVMEditorSubsystem::CanResetOrphanedPin(const UWidgetBlueprint* WidgetBlueprint, UMVVMBlueprintViewEvent* Event, const FMVVMBlueprintPinId& ParameterId) const
 {
-	UEdGraphPin* GraphPin = Private::GetGraphPin(WidgetBlueprint, Event, PinName);
+	UEdGraphPin* GraphPin = Private::GetGraphPin(WidgetBlueprint, Event, ParameterId);
 	return Private::CanResetOrphanedPin(GraphPin);
 }
 
@@ -1205,21 +1205,16 @@ FGuid UMVVMEditorSubsystem::GetFirstBindingThatUsesViewModel(const UWidgetBluepr
 
 			auto TestConversionFunction = [&](bool bForward)
 			{
-				if (UE::MVVM::IsForwardBinding(Binding.BindingType))
+				UMVVMBlueprintViewConversionFunction* ConversionFunction = Binding.Conversion.GetConversionFunction(true);
+				if (ConversionFunction)
 				{
-					const UFunction* Function = GetConversionFunction(WidgetBlueprint, Binding, bForward);
-					if (Function != nullptr)
+					for (const FMVVMBlueprintPin& Pin : ConversionFunction->GetPins())
 					{
-						TValueOrError<TArray<const FProperty*>, FText> ArgumentsResult = UE::MVVM::BindingHelper::TryGetArgumentsForConversionFunction(Function);
-						if (ArgumentsResult.HasValue())
+						if (Pin.UsedPathAsValue() && Pin.GetPath().GetSource(WidgetBlueprint) == EMVVMBlueprintFieldPathSource::ViewModel)
 						{
-							for (const FProperty* Property : ArgumentsResult.GetValue())
+							if (Pin.GetPath().GetViewModelId() == ViewModelId)
 							{
-								FMVVMBlueprintPropertyPath Path = GetPathForConversionFunctionArgument(WidgetBlueprint, Binding, Property->GetFName(), bForward);
-								if (Path.GetViewModelId() == ViewModelId)
-								{
-									return Binding.BindingId;
-								}
+								return Binding.BindingId;
 							}
 						}
 					}
diff --git a/Engine/Plugins/Runtime/ModelViewViewModel/Source/ModelViewViewModelEditor/Private/Types/MVVMBindingEntry.cpp b/Engine/Plugins/Runtime/ModelViewViewModel/Source/ModelViewViewModelEditor/Private/Types/MVVMBindingEntry.cpp
--- a/Engine/Plugins/Runtime/ModelViewViewModel/Source/ModelViewViewModelEditor/Private/Types/MVVMBindingEntry.cpp
+++ b/Engine/Plugins/Runtime/ModelViewViewModel/Source/ModelViewViewModelEditor/Private/Types/MVVMBindingEntry.cpp
@@ -51,12 +51,12 @@ void FBindingEntry::SetBindingId(FGuid Id)
 	BindingId = Id;
 }
 
-void FBindingEntry::SetBindingParameter(FGuid Id, FName ParameterName)
+void FBindingEntry::SetBindingParameter(FGuid Id, FMVVMBlueprintPinId Parameter)
 {
 	check(RowType == ERowType::None);
 	RowType = ERowType::BindingParameter;
 	BindingId = Id;
-	Name = ParameterName;
+	PinId = MoveTemp(Parameter);
 }
 
 UMVVMBlueprintViewEvent* FBindingEntry::GetEvent() const
@@ -71,12 +71,12 @@ void FBindingEntry::SetEvent(UMVVMBlueprintViewEvent* InEvent)
 	Event = InEvent;
 }
 
-void FBindingEntry::SetEventParameter(UMVVMBlueprintViewEvent* InEvent, FName ParameterName)
+void FBindingEntry::SetEventParameter(UMVVMBlueprintViewEvent* InEvent, FMVVMBlueprintPinId Parameter)
 {
 	check(RowType == ERowType::None);
 	RowType = ERowType::EventParameter;
 	Event = InEvent;
-	Name = ParameterName;
+	PinId = MoveTemp(Parameter);
 }
 
 void FBindingEntry::AddChild(TSharedPtr<FBindingEntry> Child)
@@ -107,6 +107,7 @@ bool FBindingEntry::operator==(const FBindingEntry& Other) const
 	return RowType == Other.RowType
 		&& Name == Other.Name
 		&& BindingId == Other.BindingId
+		&& PinId == Other.PinId
 		&& Event == Other.Event;
 }
 
@@ -118,9 +119,11 @@ FString FBindingEntry::GetSearchNameString(UMVVMBlueprintView* View, UWidgetBlue
 	switch (RowType)
 	{
 	case UE::MVVM::FBindingEntry::ERowType::Group:
+		RowToString = Name.ToString();
+		break;
 	case UE::MVVM::FBindingEntry::ERowType::BindingParameter:
 	case UE::MVVM::FBindingEntry::ERowType::EventParameter:
-		RowToString = Name.ToString();
+		RowToString = PinId.ToString();
 		break;
 	case UE::MVVM::FBindingEntry::ERowType::Binding:
 		{
diff --git a/Engine/Plugins/Runtime/ModelViewViewModel/Source/ModelViewViewModelEditor/Private/Types/MVVMBindingEntry.h b/Engine/Plugins/Runtime/ModelViewViewModel/Source/ModelViewViewModelEditor/Private/Types/MVVMBindingEntry.h
--- a/Engine/Plugins/Runtime/ModelViewViewModel/Source/ModelViewViewModelEditor/Private/Types/MVVMBindingEntry.h
+++ b/Engine/Plugins/Runtime/ModelViewViewModel/Source/ModelViewViewModelEditor/Private/Types/MVVMBindingEntry.h
@@ -3,6 +3,7 @@
 #pragma once
 
 #include "CoreMinimal.h"
+#include "MVVMBlueprintPin.h"
 #include "UObject/Object.h"
 #include "UObject/WeakObjectPtr.h"
 #include "UObject/WeakObjectPtrTemplates.h"
@@ -67,12 +68,12 @@ struct FBindingEntry
 	void SetBindingId(FGuid Id);
 
 	//~ binding parameter
-	FName GetBindingParameterName() const
+	const FMVVMBlueprintPinId& GetBindingParameterId() const
 	{
-		return Name;
+		return PinId;
 	}
 
-	void SetBindingParameter(FGuid Id, FName ParameterName);
+	void SetBindingParameter(FGuid Id, FMVVMBlueprintPinId Parameter);
 
 	//~ event
 	UMVVMBlueprintViewEvent* GetEvent() const;
@@ -80,12 +81,12 @@ struct FBindingEntry
 	void SetEvent(UMVVMBlueprintViewEvent* InEvent);
 
 	//~ event parameter
-	FName GetEventParameterName() const
+	const FMVVMBlueprintPinId& GetEventParameterId() const
 	{
-		return Name;
+		return PinId;
 	}
 
-	void SetEventParameter(UMVVMBlueprintViewEvent* InEvent, FName ParameterName);
+	void SetEventParameter(UMVVMBlueprintViewEvent* Event, FMVVMBlueprintPinId Parameter);
 
 	//~ children
 	TConstArrayView<TSharedPtr<FBindingEntry>> GetAllChildren() const
@@ -114,6 +115,7 @@ private:
 	ERowType RowType = ERowType::None;
 	FName Name;
 	FGuid BindingId;
+	FMVVMBlueprintPinId PinId;
 	TWeakObjectPtr<UMVVMBlueprintViewEvent> Event;
 	TArray<TSharedPtr<FBindingEntry>> AllChildren;
 	TArray<TSharedPtr<FBindingEntry>> FilteredChildren;
diff --git a/Engine/Plugins/Runtime/ModelViewViewModel/Source/ModelViewViewModelEditor/Private/Widgets/BindingEntry/SMVVMFunctionParameterRow.cpp b/Engine/Plugins/Runtime/ModelViewViewModel/Source/ModelViewViewModelEditor/Private/Widgets/BindingEntry/SMVVMFunctionParameterRow.cpp
--- a/Engine/Plugins/Runtime/ModelViewViewModel/Source/ModelViewViewModelEditor/Private/Widgets/BindingEntry/SMVVMFunctionParameterRow.cpp
+++ b/Engine/Plugins/Runtime/ModelViewViewModel/Source/ModelViewViewModelEditor/Private/Widgets/BindingEntry/SMVVMFunctionParameterRow.cpp
@@ -41,12 +41,12 @@ TSharedRef<SWidget> SFunctionParameterRow::BuildRowWidget()
 		UMVVMBlueprintViewConversionFunction* ConversionFunction = Binding->Conversion.GetConversionFunction(bSourceToDestination);
 		check(ConversionFunction);
 
-		Pin = ConversionFunction->GetPins().FindByPredicate([Name = GetEntry()->GetBindingParameterName()](const FMVVMBlueprintPin& Other) { return Other.GetName() == Name; });
-		GraphPin = ConversionFunction->GetOrCreateGraphPin(GetBlueprint(), GetEntry()->GetBindingParameterName());
+		Pin = ConversionFunction->GetPins().FindByPredicate([ArgId = GetEntry()->GetBindingParameterId()](const FMVVMBlueprintPin& Other) { return Other.GetId() == ArgId; });
+		GraphPin = ConversionFunction->GetOrCreateGraphPin(GetBlueprint(), GetEntry()->GetBindingParameterId());
 
 		ContentWidget = SNew(SFunctionParameter, GetBlueprint())
 			.BindingId(Binding->BindingId)
-			.ParameterName(GetEntry()->GetBindingParameterName())
+			.ParameterId(GetEntry()->GetBindingParameterId())
 			.SourceToDestination(bSourceToDestination)
 			.AllowDefault(!bSimpleConversionFunction);
 	}
@@ -55,12 +55,12 @@ TSharedRef<SWidget> SFunctionParameterRow::BuildRowWidget()
 		UMVVMBlueprintViewEvent* ViewEvent = GetEntry()->GetEvent();
 		check(ViewEvent);
 
-		Pin = ViewEvent->GetPins().FindByPredicate([Name = GetEntry()->GetEventParameterName()](const FMVVMBlueprintPin& Other) { return Other.GetName() == Name; });
-		GraphPin = ViewEvent->GetOrCreateGraphPin(GetEntry()->GetEventParameterName());
+		Pin = ViewEvent->GetPins().FindByPredicate([ArgId = GetEntry()->GetEventParameterId()](const FMVVMBlueprintPin& Other) { return Other.GetId() == ArgId; });
+		GraphPin = ViewEvent->GetOrCreateGraphPin(GetEntry()->GetEventParameterId());
 
 		ContentWidget = SNew(SEventParameter, GetBlueprint())
 			.Event(GetEntry()->GetEvent())
-			.ParameterName(GetEntry()->GetEventParameterName())
+			.ParameterId(GetEntry()->GetEventParameterId())
 			.AllowDefault(true);
 	}
 
@@ -73,7 +73,7 @@ TSharedRef<SWidget> SFunctionParameterRow::BuildRowWidget()
 	{
 		PrimaryBrush = FBlueprintEditorUtils::GetIconFromPin(FEdGraphPinType());
 		PrimaryColor = FLinearColor::Red;
-		DisplayName = GraphPin ? GraphPin->GetDisplayName() : (Pin ? FText::FromName(Pin->GetName()) : FText::GetEmpty());
+		DisplayName = GraphPin ? GraphPin->GetDisplayName() : (Pin ? FText::FromName(Pin->GetId().GetNames().Last()) : FText::GetEmpty());
 		bTextColorIsRed = true;
 	}
 	else
diff --git a/Engine/Plugins/Runtime/ModelViewViewModel/Source/ModelViewViewModelEditor/Private/Widgets/SMVVMEventParameter.cpp b/Engine/Plugins/Runtime/ModelViewViewModel/Source/ModelViewViewModelEditor/Private/Widgets/SMVVMEventParameter.cpp
--- a/Engine/Plugins/Runtime/ModelViewViewModel/Source/ModelViewViewModelEditor/Private/Widgets/SMVVMEventParameter.cpp
+++ b/Engine/Plugins/Runtime/ModelViewViewModel/Source/ModelViewViewModelEditor/Private/Widgets/SMVVMEventParameter.cpp
@@ -32,14 +32,14 @@ void SEventParameter::Construct(const FArguments& InArgs, UWidgetBlueprint* InWi
 	ViewEvent = InArgs._Event;
 	check(InArgs._Event);
 
-	ParameterName = InArgs._ParameterName;
-	check(!ParameterName.IsNone());
+	ParameterId = InArgs._ParameterId;
+	check(ParameterId.IsValid());
 
 	bAllowDefault = InArgs._AllowDefault;
 
 	bool bIsBooleanPin = false;
 	TSharedRef<SWidget> ValueWidget = SNullWidget::NullWidget;
-	if (UEdGraphPin* Pin = InArgs._Event->GetOrCreateGraphPin(ParameterName))
+	if (UEdGraphPin* Pin = InArgs._Event->GetOrCreateGraphPin(ParameterId))
 	{
 		bIsBooleanPin = Pin->PinType.PinCategory == UEdGraphSchema_K2::PC_Boolean;
 		// create a new pin widget so that we can get the default value widget out of it
@@ -155,7 +155,7 @@ FMVVMBlueprintPropertyPath SEventParameter::OnGetSelectedField() const
 {
 	if (const UMVVMBlueprintViewEvent* EventPtr = ViewEvent.Get())
 	{
-		return EventPtr->GetPinPath(ParameterName);
+		return EventPtr->GetPinPath(ParameterId);
 	}
 	return FMVVMBlueprintPropertyPath();
 }
@@ -165,7 +165,7 @@ void SEventParameter::SetSelectedField(const FMVVMBlueprintPropertyPath& Path)
 	if (UMVVMBlueprintViewEvent* EventPtr = ViewEvent.Get())
 	{
 		const UMVVMEditorSubsystem* Subsystem = GEditor->GetEditorSubsystem<UMVVMEditorSubsystem>();
-		Subsystem->SetEventArgumentPath(EventPtr, ParameterName, Path);
+		Subsystem->SetEventArgumentPath(EventPtr, ParameterId, Path);
 	}
 }
 
@@ -185,7 +185,7 @@ FFieldSelectionContext SEventParameter::GetSelectedSelectionContext() const
 	}
 	
 	Result.BindingMode = EMVVMBindingMode::OneTimeToDestination;
-	//Result.AssignableTo = ConversionFunction->FindPropertyByName(ParameterName);
+	//Result.AssignableTo = ConversionFunction->FindPropertyByName(ParameterId);
 	Result.bAllowWidgets = true;
 	Result.bAllowViewModels = true;
 	Result.bAllowConversionFunctions = false;
diff --git a/Engine/Plugins/Runtime/ModelViewViewModel/Source/ModelViewViewModelEditor/Private/Widgets/SMVVMEventParameter.h b/Engine/Plugins/Runtime/ModelViewViewModel/Source/ModelViewViewModelEditor/Private/Widgets/SMVVMEventParameter.h
--- a/Engine/Plugins/Runtime/ModelViewViewModel/Source/ModelViewViewModelEditor/Private/Widgets/SMVVMEventParameter.h
+++ b/Engine/Plugins/Runtime/ModelViewViewModel/Source/ModelViewViewModelEditor/Private/Widgets/SMVVMEventParameter.h
@@ -2,6 +2,7 @@
 
 #pragma once
 
+#include "MVVMBlueprintPin.h"
 #include "MVVMPropertyPath.h"
 #include "Widgets/SCompoundWidget.h"
 #include "Widgets/SMVVMFieldSelectorMenu.h"
@@ -24,7 +25,7 @@ class SEventParameter : public SCompoundWidget
 public:
 	SLATE_BEGIN_ARGS(SEventParameter) {}
 		SLATE_ARGUMENT_DEFAULT(UMVVMBlueprintViewEvent*, Event) =  nullptr;
-		SLATE_ARGUMENT(FName, ParameterName)
+		SLATE_ARGUMENT(FMVVMBlueprintPinId, ParameterId)
 		SLATE_ARGUMENT_DEFAULT(bool, AllowDefault) = true;
 	SLATE_END_ARGS()
 
@@ -45,7 +46,7 @@ private:
 private:
 	TWeakObjectPtr<UWidgetBlueprint> WidgetBlueprint;
 	TWeakObjectPtr<UMVVMBlueprintViewEvent> ViewEvent;
-	FName ParameterName;
+	FMVVMBlueprintPinId ParameterId;
 	/** This reference is just to keep the default value widget alive. */
 	TSharedPtr<SGraphPin> GraphPin;
 
diff --git a/Engine/Plugins/Runtime/ModelViewViewModel/Source/ModelViewViewModelEditor/Private/Widgets/SMVVMFunctionParameter.cpp b/Engine/Plugins/Runtime/ModelViewViewModel/Source/ModelViewViewModelEditor/Private/Widgets/SMVVMFunctionParameter.cpp
--- a/Engine/Plugins/Runtime/ModelViewViewModel/Source/ModelViewViewModelEditor/Private/Widgets/SMVVMFunctionParameter.cpp
+++ b/Engine/Plugins/Runtime/ModelViewViewModel/Source/ModelViewViewModelEditor/Private/Widgets/SMVVMFunctionParameter.cpp
@@ -32,8 +32,8 @@ void SFunctionParameter::Construct(const FArguments& InArgs, UWidgetBlueprint* I
 	BindingId = InArgs._BindingId;
 	check(BindingId.IsValid());
 
-	ParameterName = InArgs._ParameterName;
-	check(!ParameterName.IsNone());
+	ParameterId = InArgs._ParameterId;
+	check(ParameterId.IsValid());
 	
 	bSourceToDestination = InArgs._SourceToDestination;
 	bAllowDefault = InArgs._AllowDefault;
@@ -48,7 +48,7 @@ void SFunctionParameter::Construct(const FArguments& InArgs, UWidgetBlueprint* I
 	TSharedRef<SWidget> ValueWidget = SNullWidget::NullWidget;
 	bool bIsBooleanPin = false;
 
-	if (UEdGraphPin* Pin = EditorSubsystem->GetConversionFunctionArgumentPin(InWidgetBlueprint, *Binding, ParameterName, bSourceToDestination))
+	if (UEdGraphPin* Pin = EditorSubsystem->GetConversionFunctionArgumentPin(InWidgetBlueprint, *Binding, ParameterId, bSourceToDestination))
 	{
 		bIsBooleanPin = Pin->PinType.PinCategory == UEdGraphSchema_K2::PC_Boolean;
 		// create a new pin widget so that we can get the default value widget out of it
@@ -172,7 +172,7 @@ FMVVMBlueprintPropertyPath SFunctionParameter::OnGetSelectedField() const
 		{
 			if (const FMVVMBlueprintViewBinding* Binding = View->GetBinding(BindingId))
 			{
-				return Subsystem->GetPathForConversionFunctionArgument(WidgetBlueprint.Get(), *Binding, ParameterName, bSourceToDestination);
+				return Subsystem->GetPathForConversionFunctionArgument(WidgetBlueprint.Get(), *Binding, ParameterId, bSourceToDestination);
 			}
 		}
 	}
@@ -188,7 +188,7 @@ void SFunctionParameter::SetSelectedField(const FMVVMBlueprintPropertyPath& Path
 		{
 			if (FMVVMBlueprintViewBinding* Binding = View->GetBinding(BindingId))
 			{
-				Subsystem->SetPathForConversionFunctionArgument(WidgetBlueprint.Get(), *Binding, ParameterName, Path, bSourceToDestination);
+				Subsystem->SetPathForConversionFunctionArgument(WidgetBlueprint.Get(), *Binding, ParameterId, Path, bSourceToDestination);
 			}
 		}
 	}
@@ -225,7 +225,7 @@ FFieldSelectionContext SFunctionParameter::GetSelectedSelectionContext() const
 	}
 
 	Result.BindingMode = EMVVMBindingMode::OneWayToDestination;
-	Result.AssignableTo = ConversionFunction->FindPropertyByName(ParameterName);
+	Result.AssignableTo = ParameterId.GetNames().Num() > 0 ? ConversionFunction->FindPropertyByName(ParameterId.GetNames().Last()) : nullptr;
 	Result.bAllowWidgets = true;
 	Result.bAllowViewModels = true;
 	Result.bAllowConversionFunctions = false;
diff --git a/Engine/Plugins/Runtime/ModelViewViewModel/Source/ModelViewViewModelEditor/Private/Widgets/SMVVMFunctionParameter.h b/Engine/Plugins/Runtime/ModelViewViewModel/Source/ModelViewViewModelEditor/Private/Widgets/SMVVMFunctionParameter.h
--- a/Engine/Plugins/Runtime/ModelViewViewModel/Source/ModelViewViewModelEditor/Private/Widgets/SMVVMFunctionParameter.h
+++ b/Engine/Plugins/Runtime/ModelViewViewModel/Source/ModelViewViewModelEditor/Private/Widgets/SMVVMFunctionParameter.h
@@ -2,6 +2,7 @@
 
 #pragma once
 
+#include "MVVMBlueprintPin.h"
 #include "MVVMPropertyPath.h"
 #include "Widgets/SCompoundWidget.h"
 #include "Widgets/SMVVMFieldSelectorMenu.h"
@@ -25,7 +26,7 @@ class SFunctionParameter : public SCompoundWidget
 public:
 	SLATE_BEGIN_ARGS(SFunctionParameter) {}
 		SLATE_ARGUMENT(FGuid, BindingId)
-		SLATE_ARGUMENT(FName, ParameterName)
+		SLATE_ARGUMENT(FMVVMBlueprintPinId, ParameterId)
 		SLATE_ARGUMENT_DEFAULT(bool, SourceToDestination) = true;
 		SLATE_ARGUMENT_DEFAULT(bool, AllowDefault) = true;
 	SLATE_END_ARGS()
@@ -47,7 +48,7 @@ private:
 private:
 	TWeakObjectPtr<UWidgetBlueprint> WidgetBlueprint;
 	FGuid BindingId;
-	FName ParameterName;
+	FMVVMBlueprintPinId ParameterId;
 	/** This reference is just to keep the default value widget alive. */
 	TSharedPtr<SGraphPin> GraphPin;
 
diff --git a/Engine/Plugins/Runtime/ModelViewViewModel/Source/ModelViewViewModelEditor/Private/Widgets/SMVVMViewBindingListView.cpp b/Engine/Plugins/Runtime/ModelViewViewModel/Source/ModelViewViewModelEditor/Private/Widgets/SMVVMViewBindingListView.cpp
--- a/Engine/Plugins/Runtime/ModelViewViewModel/Source/ModelViewViewModelEditor/Private/Widgets/SMVVMViewBindingListView.cpp
+++ b/Engine/Plugins/Runtime/ModelViewViewModel/Source/ModelViewViewModelEditor/Private/Widgets/SMVVMViewBindingListView.cpp
@@ -342,7 +342,7 @@ void SBindingsList::Refresh()
 				ConversionFunction->GetOrCreateWrapperGraph(MVVMExtensionPtr->GetWidgetBlueprint());
 				for (const FMVVMBlueprintPin& Pin : ConversionFunction->GetPins())
 				{
-					UEdGraphPin* GraphPin = ConversionFunction->GetOrCreateGraphPin(MVVMExtensionPtr->GetWidgetBlueprint(), Pin.GetName());
+					UEdGraphPin* GraphPin = ConversionFunction->GetOrCreateGraphPin(MVVMExtensionPtr->GetWidgetBlueprint(), Pin.GetId());
 					if (GraphPin && GraphPin->bHidden)
 					{
 						continue;
@@ -351,8 +351,12 @@ void SBindingsList::Refresh()
 					TSharedPtr<FBindingEntry> ArgumentEntry;
 					if (PreviousGroupEntry)
 					{
-						if (TSharedPtr<FBindingEntry>* FoundParameter = PreviousGroupEntry->Children.FindByPredicate([BindingId, ArgumentName = Pin.GetName()](const TSharedPtr<FBindingEntry>& Other)
-							{ return Other->GetBindingId() == BindingId && Other->GetRowType() == FBindingEntry::ERowType::BindingParameter && Other->GetBindingParameterName() == ArgumentName; }))
+						TSharedPtr<FBindingEntry>* FoundParameter = PreviousGroupEntry->Children.FindByPredicate(
+							[BindingId, ArgumentId = Pin.GetId()](const TSharedPtr<FBindingEntry>& Other)
+							{
+								return Other->GetBindingId() == BindingId && Other->GetRowType() == FBindingEntry::ERowType::BindingParameter && Other->GetBindingParameterId() == ArgumentId;
+							});
+						if (FoundParameter)
 						{
 							ArgumentEntry = *FoundParameter;
 						}
@@ -361,7 +365,7 @@ void SBindingsList::Refresh()
 					if (!ArgumentEntry.IsValid())
 					{
 						ArgumentEntry = MakeShared<FBindingEntry>();
-						ArgumentEntry->SetBindingParameter(Binding.BindingId, Pin.GetName());
+						ArgumentEntry->SetBindingParameter(Binding.BindingId, Pin.GetId());
 						ConversionFunction->OnWrapperGraphModified.AddSP(this, &SBindingsList::ForceRefresh);
 
 						NewEntries.Add(ArgumentEntry);
@@ -425,7 +429,7 @@ void SBindingsList::Refresh()
 			// Create/Find entries for function parameters
 			for (const FMVVMBlueprintPin& Pin : Event->GetPins())
 			{
-				UEdGraphPin* GraphPin = Event->GetOrCreateGraphPin(Pin.GetName());
+				UEdGraphPin* GraphPin = Event->GetOrCreateGraphPin(Pin.GetId());
 				if (GraphPin && GraphPin->bHidden)
 				{
 					continue;
@@ -434,8 +438,12 @@ void SBindingsList::Refresh()
 				TSharedPtr<FBindingEntry> ArgumentEntry;
 				if (PreviousGroupEntry)
 				{
-					if (TSharedPtr<FBindingEntry>* FoundParameter = PreviousGroupEntry->Children.FindByPredicate([Event, ArgumentName = Pin.GetName()](const TSharedPtr<FBindingEntry>& Other)
-						{ return Other->GetRowType() == FBindingEntry::ERowType::EventParameter && Other->GetEvent() == Event && Other->GetEventParameterName() == ArgumentName; }))
+					TSharedPtr<FBindingEntry>* FoundParameter = PreviousGroupEntry->Children.FindByPredicate(
+						[Event, ArgumentId = Pin.GetId()](const TSharedPtr<FBindingEntry>& Other)
+						{
+							return Other->GetRowType() == FBindingEntry::ERowType::EventParameter && Other->GetEvent() == Event && Other->GetEventParameterId() == ArgumentId;
+						});
+					if (FoundParameter)
 					{
 						ArgumentEntry = *FoundParameter;
 					}
@@ -444,7 +452,7 @@ void SBindingsList::Refresh()
 				if (!ArgumentEntry.IsValid())
 				{
 					ArgumentEntry = MakeShared<FBindingEntry>();
-					ArgumentEntry->SetEventParameter(Event, Pin.GetName());
+					ArgumentEntry->SetEventParameter(Event, Pin.GetId());
 					Event->OnWrapperGraphModified.AddSP(this, &SBindingsList::ForceRefresh);
 
 					NewEntries.Add(ArgumentEntry);
diff --git a/Engine/Plugins/Runtime/ModelViewViewModel/Source/ModelViewViewModelEditor/Public/MVVMEditorSubsystem.h b/Engine/Plugins/Runtime/ModelViewViewModel/Source/ModelViewViewModelEditor/Public/MVVMEditorSubsystem.h
--- a/Engine/Plugins/Runtime/ModelViewViewModel/Source/ModelViewViewModelEditor/Public/MVVMEditorSubsystem.h
+++ b/Engine/Plugins/Runtime/ModelViewViewModel/Source/ModelViewViewModelEditor/Public/MVVMEditorSubsystem.h
@@ -3,9 +3,9 @@
 #pragma once
 
 #include "EditorSubsystem.h"
-
-
+#include "MVVMBlueprintPin.h"
 #include "UObject/Package.h"
+
 #include "MVVMEditorSubsystem.generated.h"
 
 class UEdGraphPin;
@@ -81,7 +81,7 @@ public:
 
 	void SetEventPath(UMVVMBlueprintViewEvent* Event, FMVVMBlueprintPropertyPath PropertyPath);
 	void SetEventDestinationPath(UMVVMBlueprintViewEvent* Event, FMVVMBlueprintPropertyPath PropertyPath);
-	void SetEventArgumentPath(UMVVMBlueprintViewEvent* Event, FName ArgumentName, const FMVVMBlueprintPropertyPath& PropertyPath) const;
+	void SetEventArgumentPath(UMVVMBlueprintViewEvent* Event, const FMVVMBlueprintPinId& PinId, const FMVVMBlueprintPropertyPath& PropertyPath) const;
 	void SetEnabledForEvent(UMVVMBlueprintViewEvent* Event, bool bEnabled);
 	void SetCompileForEvent(UMVVMBlueprintViewEvent* Event, bool bCompile);
 
@@ -104,26 +104,26 @@ public:
 	UFUNCTION(BlueprintCallable, Category = "Viewmodel")
 	TArray<UFunction*> GetAvailableConversionFunctions(const UWidgetBlueprint* WidgetBlueprint, const FMVVMBlueprintPropertyPath& Source, const FMVVMBlueprintPropertyPath& Destination) const;
 
-	FMVVMBlueprintPropertyPath GetPathForConversionFunctionArgument(const UWidgetBlueprint* WidgetBlueprint, const FMVVMBlueprintViewBinding& Binding, FName ArgumentName, bool bSourceToDestination) const;
-	void SetPathForConversionFunctionArgument(UWidgetBlueprint* WidgetBlueprint, FMVVMBlueprintViewBinding& Binding, FName ArgumentName, const FMVVMBlueprintPropertyPath& Path, bool bSourceToDestination) const;
-	UEdGraphPin* GetConversionFunctionArgumentPin(const UWidgetBlueprint* WidgetBlueprint, const FMVVMBlueprintViewBinding& Binding, FName ParameterName, bool bSourceToDestination) const;
-
-	void SplitPin(UWidgetBlueprint* WidgetBlueprint, FMVVMBlueprintViewBinding& Binding, FName PinName, bool bSourceToDestination) const;
-	bool CanSplitPin(const UWidgetBlueprint* WidgetBlueprint, const FMVVMBlueprintViewBinding& Binding, FName PinName, bool bSourceToDestination) const;
-	void SplitPin(UWidgetBlueprint* WidgetBlueprint, UMVVMBlueprintViewEvent* Event, FName PinName) const;
-	bool CanSplitPin(const UWidgetBlueprint* WidgetBlueprint, UMVVMBlueprintViewEvent* Event, FName PinName) const;
-	void RecombinePin(UWidgetBlueprint* WidgetBlueprint, FMVVMBlueprintViewBinding& Binding, FName PinName, bool bSourceToDestination) const;
-	bool CanRecombinePin(const UWidgetBlueprint* WidgetBlueprint, const FMVVMBlueprintViewBinding& Binding, FName PinName, bool bSourceToDestination) const;
-	void RecombinePin(UWidgetBlueprint* WidgetBlueprint, UMVVMBlueprintViewEvent* Event, FName PinName) const;
-	bool CanRecombinePin(const UWidgetBlueprint* WidgetBlueprint, UMVVMBlueprintViewEvent* Event, FName PinName) const;
-	void ResetPinToDefaultValue(UWidgetBlueprint* WidgetBlueprint, FMVVMBlueprintViewBinding& Binding, FName PinName, bool bSourceToDestination) const;
-	bool CanResetPinToDefaultValue(const UWidgetBlueprint* WidgetBlueprint, const FMVVMBlueprintViewBinding& Binding, FName PinName, bool bSourceToDestination) const;
-	void ResetPinToDefaultValue(UWidgetBlueprint* WidgetBlueprint, UMVVMBlueprintViewEvent* Event, FName PinName) const;
-	bool CanResetPinToDefaultValue(const UWidgetBlueprint* WidgetBlueprint, UMVVMBlueprintViewEvent* Event, FName PinName) const;
-	void ResetOrphanedPin(UWidgetBlueprint* WidgetBlueprint, FMVVMBlueprintViewBinding& Binding, FName PinName, bool bSourceToDestination) const;
-	bool CanResetOrphanedPin(const UWidgetBlueprint* WidgetBlueprint, const FMVVMBlueprintViewBinding& Binding, FName PinName, bool bSourceToDestination) const;
-	void ResetOrphanedPin(UWidgetBlueprint* WidgetBlueprint, UMVVMBlueprintViewEvent* Event, FName PinName) const;
-	bool CanResetOrphanedPin(const UWidgetBlueprint* WidgetBlueprint, UMVVMBlueprintViewEvent* Event, FName PinName) const;
+	FMVVMBlueprintPropertyPath GetPathForConversionFunctionArgument(const UWidgetBlueprint* WidgetBlueprint, const FMVVMBlueprintViewBinding& Binding, const FMVVMBlueprintPinId& PinId, bool bSourceToDestination) const;
+	void SetPathForConversionFunctionArgument(UWidgetBlueprint* WidgetBlueprint, FMVVMBlueprintViewBinding& Binding, const FMVVMBlueprintPinId& PinId, const FMVVMBlueprintPropertyPath& Path, bool bSourceToDestination) const;
+	UEdGraphPin* GetConversionFunctionArgumentPin(const UWidgetBlueprint* WidgetBlueprint, const FMVVMBlueprintViewBinding& Binding, const FMVVMBlueprintPinId& PinId, bool bSourceToDestination) const;
+
+	void SplitPin(UWidgetBlueprint* WidgetBlueprint, FMVVMBlueprintViewBinding& Binding, const FMVVMBlueprintPinId& PinId, bool bSourceToDestination) const;
+	bool CanSplitPin(const UWidgetBlueprint* WidgetBlueprint, const FMVVMBlueprintViewBinding& Binding, const FMVVMBlueprintPinId& PinId, bool bSourceToDestination) const;
+	void SplitPin(UWidgetBlueprint* WidgetBlueprint, UMVVMBlueprintViewEvent* Event, const FMVVMBlueprintPinId& PinId) const;
+	bool CanSplitPin(const UWidgetBlueprint* WidgetBlueprint, UMVVMBlueprintViewEvent* Event, const FMVVMBlueprintPinId& PinId) const;
+	void RecombinePin(UWidgetBlueprint* WidgetBlueprint, FMVVMBlueprintViewBinding& Binding, const FMVVMBlueprintPinId& PinId, bool bSourceToDestination) const;
+	bool CanRecombinePin(const UWidgetBlueprint* WidgetBlueprint, const FMVVMBlueprintViewBinding& Binding, const FMVVMBlueprintPinId& PinId, bool bSourceToDestination) const;
+	void RecombinePin(UWidgetBlueprint* WidgetBlueprint, UMVVMBlueprintViewEvent* Event, const FMVVMBlueprintPinId& PinId) const;
+	bool CanRecombinePin(const UWidgetBlueprint* WidgetBlueprint, UMVVMBlueprintViewEvent* Event, const FMVVMBlueprintPinId& PinId) const;
+	void ResetPinToDefaultValue(UWidgetBlueprint* WidgetBlueprint, FMVVMBlueprintViewBinding& Binding,const FMVVMBlueprintPinId& PinId, bool bSourceToDestination) const;
+	bool CanResetPinToDefaultValue(const UWidgetBlueprint* WidgetBlueprint, const FMVVMBlueprintViewBinding& Binding, const FMVVMBlueprintPinId& PinId, bool bSourceToDestination) const;
+	void ResetPinToDefaultValue(UWidgetBlueprint* WidgetBlueprint, UMVVMBlueprintViewEvent* Event, const FMVVMBlueprintPinId& PinId) const;
+	bool CanResetPinToDefaultValue(const UWidgetBlueprint* WidgetBlueprint, UMVVMBlueprintViewEvent* Event, const FMVVMBlueprintPinId& PinId) const;
+	void ResetOrphanedPin(UWidgetBlueprint* WidgetBlueprint, FMVVMBlueprintViewBinding& Binding, const FMVVMBlueprintPinId& PinId, bool bSourceToDestination) const;
+	bool CanResetOrphanedPin(const UWidgetBlueprint* WidgetBlueprint, const FMVVMBlueprintViewBinding& Binding, const FMVVMBlueprintPinId& PinId, bool bSourceToDestination) const;
+	void ResetOrphanedPin(UWidgetBlueprint* WidgetBlueprint, UMVVMBlueprintViewEvent* Event,const FMVVMBlueprintPinId& PinId) const;
+	bool CanResetOrphanedPin(const UWidgetBlueprint* WidgetBlueprint, UMVVMBlueprintViewEvent* Event, const FMVVMBlueprintPinId& PinId) const;
 
 	TArray<UE::MVVM::FBindingSource> GetBindableWidgets(const UWidgetBlueprint* WidgetBlueprint) const;
 	TArray<UE::MVVM::FBindingSource> GetAllViewModels(const UWidgetBlueprint* WidgetBlueprint) const;
