MVVM: Rename the property introduced with CL 30396310. It was in conflict with an old property that got removed and was not deprecated properly. Old data tries to load the old property with the bad type.

diff --git a/Engine/Plugins/Runtime/ModelViewViewModel/Source/ModelViewViewModelBlueprint/Private/MVVMBlueprintPin.cpp b/Engine/Plugins/Runtime/ModelViewViewModel/Source/ModelViewViewModelBlueprint/Private/MVVMBlueprintPin.cpp
--- a/Engine/Plugins/Runtime/ModelViewViewModel/Source/ModelViewViewModelBlueprint/Private/MVVMBlueprintPin.cpp
+++ b/Engine/Plugins/Runtime/ModelViewViewModel/Source/ModelViewViewModelBlueprint/Private/MVVMBlueprintPin.cpp
@@ -89,24 +89,24 @@ FString FMVVMBlueprintPinId::ToString() const
  *
  */
 FMVVMBlueprintPin::FMVVMBlueprintPin(FName InPinName)
-	: PinId(MakeArrayView(&InPinName, 1))
+	: Id(MakeArrayView(&InPinName, 1))
 {
 }
 
 FMVVMBlueprintPin::FMVVMBlueprintPin(FMVVMBlueprintPinId InPinId)
-	: PinId(MoveTemp(InPinId))
+	: Id(MoveTemp(InPinId))
 {
 }
 
 FMVVMBlueprintPin::FMVVMBlueprintPin(const TArrayView<const FName> InPinNames)
-	: PinId(InPinNames)
+	: Id(InPinNames)
 {
 }
 
 FMVVMBlueprintPin FMVVMBlueprintPin::CreateFromPin(const UBlueprint* Blueprint, const UEdGraphPin* Pin)
 {
 	FMVVMBlueprintPin Result;
-	Result.PinId = FMVVMBlueprintPinId(UE::MVVM::ConversionFunctionHelper::FindPinId(Pin));
+	Result.Id = FMVVMBlueprintPinId(UE::MVVM::ConversionFunctionHelper::FindPinId(Pin));
 	Result.Path = UE::MVVM::ConversionFunctionHelper::GetPropertyPathForPin(Blueprint, Pin, true);
 	Result.DefaultObject = Pin->DefaultObject;
 	Result.DefaultString = Pin->DefaultValue;
@@ -232,7 +232,7 @@ TArray<FMVVMBlueprintPin> FMVVMBlueprintPin::CreateFromNode(UBlueprint* Blueprin
 
 UEdGraphPin* FMVVMBlueprintPin::FindGraphPin(const UEdGraph* Graph) const
 {
-	return UE::MVVM::ConversionFunctionHelper::FindPin(Graph, PinId.GetNames());
+	return UE::MVVM::ConversionFunctionHelper::FindPin(Graph, Id.GetNames());
 }
 
 void FMVVMBlueprintPin::Reset()
@@ -251,7 +251,7 @@ void FMVVMBlueprintPin::PostSerialize(const FArchive& Ar)
 	{
 		if (PinName_DEPRECATED.IsValid())
 		{
-			PinId = FMVVMBlueprintPinId(MakeArrayView(&PinName_DEPRECATED, 1));
+			Id = FMVVMBlueprintPinId(MakeArrayView(&PinName_DEPRECATED, 1));
 			PinName_DEPRECATED = FName();
 		}
 	}
diff --git a/Engine/Plugins/Runtime/ModelViewViewModel/Source/ModelViewViewModelBlueprint/Public/MVVMBlueprintPin.h b/Engine/Plugins/Runtime/ModelViewViewModel/Source/ModelViewViewModelBlueprint/Public/MVVMBlueprintPin.h
--- a/Engine/Plugins/Runtime/ModelViewViewModel/Source/ModelViewViewModelBlueprint/Public/MVVMBlueprintPin.h
+++ b/Engine/Plugins/Runtime/ModelViewViewModel/Source/ModelViewViewModelBlueprint/Public/MVVMBlueprintPin.h
@@ -65,7 +65,7 @@ struct MODELVIEWVIEWMODELBLUEPRINT_API FMVVMBlueprintPin
 
 private:	
 	UPROPERTY(VisibleAnywhere, Category = "MVVM")
-	FMVVMBlueprintPinId PinId;
+	FMVVMBlueprintPinId Id;
 
 	UPROPERTY(VisibleAnywhere, Category = "MVVM")
 	FMVVMBlueprintPropertyPath Path;
@@ -92,6 +92,9 @@ private:
 
 	UPROPERTY()
 	FName PinName_DEPRECATED;
+	
+	UPROPERTY()
+	FGuid PinId_DEPRECATED;
 
 public:
 	FMVVMBlueprintPin() = default;
@@ -103,17 +106,17 @@ public:
 	UE_DEPRECATED(5.4, "GetName is deprecated. Use GetNames instead")
 	FName GetName() const
 	{
-		return PinId.GetNames().Num() > 0 ? PinId.GetNames().Last() : FName();
+		return Id.GetNames().Num() > 0 ? Id.GetNames().Last() : FName();
 	}
 
 	const FMVVMBlueprintPinId& GetId() const
 	{
-		return PinId;
+		return Id;
 	}
 
 	bool IsValid() const
 	{
-		return PinId.IsValid();
+		return Id.IsValid();
 	}
 
 	/** The pin is split into its different components. */
