MVVM: Fix bad deprecation for PinName. Confused IsValid with IsNone.

diff --git a/Engine/Plugins/Runtime/ModelViewViewModel/Source/ModelViewViewModelBlueprint/Private/MVVMBlueprintPin.cpp b/Engine/Plugins/Runtime/ModelViewViewModel/Source/ModelViewViewModelBlueprint/Private/MVVMBlueprintPin.cpp
--- a/Engine/Plugins/Runtime/ModelViewViewModel/Source/ModelViewViewModelBlueprint/Private/MVVMBlueprintPin.cpp
+++ b/Engine/Plugins/Runtime/ModelViewViewModel/Source/ModelViewViewModelBlueprint/Private/MVVMBlueprintPin.cpp
@@ -249,7 +249,7 @@ void FMVVMBlueprintPin::PostSerialize(const FArchive& Ar)
 {
 	if (Ar.IsLoading())
 	{
-		if (PinName_DEPRECATED.IsValid())
+		if (!PinName_DEPRECATED.IsNone())
 		{
 			Id = FMVVMBlueprintPinId(MakeArrayView(&PinName_DEPRECATED, 1));
 			PinName_DEPRECATED = FName();
