MVVM: Add MVVMBlueprintFunctionReference. A struct that can be a K2Node or a UFunction.

diff --git a/Engine/Plugins/Runtime/ModelViewViewModel/Source/ModelViewViewModelBlueprint/Private/MVVMBlueprintFunctionReference.cpp b/Engine/Plugins/Runtime/ModelViewViewModel/Source/ModelViewViewModelBlueprint/Private/MVVMBlueprintFunctionReference.cpp
new file mode 100644
--- /dev/null
+++ b/Engine/Plugins/Runtime/ModelViewViewModel/Source/ModelViewViewModelBlueprint/Private/MVVMBlueprintFunctionReference.cpp
@@ -0,0 +1,135 @@
+// Copyright Epic Games, Inc. All Rights Reserved.
+
+#include "MVVMBlueprintFunctionReference.h"
+#include "BlueprintCompilationManager.h"
+#include "Engine/Blueprint.h"
+#include "Kismet2/BlueprintEditorUtils.h"
+
+#include UE_INLINE_GENERATED_CPP_BY_NAME(MVVMBlueprintFunctionReference)
+
+#define LOCTEXT_NAMESPACE "MVVMBlueprintFunctionReference"
+
+FMVVMBlueprintFunctionReference::FMVVMBlueprintFunctionReference(const UBlueprint* InContext, const UFunction* InFunction)
+	: Type(EMVVMBlueprintFunctionReferenceType::None)
+{
+	UClass* OwnerClass = InFunction ? InFunction->GetOwnerClass() : nullptr;
+	ensure(OwnerClass);
+	if (!OwnerClass)
+	{
+		return;
+	}
+
+	Type = EMVVMBlueprintFunctionReferenceType::Function;
+	FName FieldName = InFunction->GetFName();
+	FGuid MemberGuid;
+	UBlueprint::GetGuidFromClassByFieldName<UFunction>(OwnerClass, FieldName, MemberGuid);
+
+	// Set the member reference
+	bool bIsSelf = (InContext->GeneratedClass && InContext->GeneratedClass->IsChildOf(OwnerClass))
+		|| (InContext->SkeletonGeneratedClass && InContext->SkeletonGeneratedClass->IsChildOf(OwnerClass));
+	if (bIsSelf)
+	{
+		FunctionReference.SetSelfMember(FieldName, MemberGuid);
+	}
+	else
+	{
+		FGuid Guid;
+		if (UBlueprint* VariableOwnerBP = Cast<UBlueprint>(OwnerClass->ClassGeneratedBy))
+		{
+			OwnerClass = VariableOwnerBP->SkeletonGeneratedClass ? VariableOwnerBP->SkeletonGeneratedClass : VariableOwnerBP->GeneratedClass;
+		}
+
+		FunctionReference.SetExternalMember(FieldName, OwnerClass, MemberGuid);
+	}
+}
+
+FMVVMBlueprintFunctionReference::FMVVMBlueprintFunctionReference(FMemberReference InReference)
+	: FunctionReference(MoveTemp(InReference))
+	, Type(EMVVMBlueprintFunctionReferenceType::Function)
+{}
+
+FMVVMBlueprintFunctionReference::FMVVMBlueprintFunctionReference(const TSubclassOf<UK2Node> InNode)
+	: Node(InNode)
+	, Type(InNode.Get() != nullptr ? EMVVMBlueprintFunctionReferenceType::Node : EMVVMBlueprintFunctionReferenceType::None)
+{
+}
+
+const UFunction* FMVVMBlueprintFunctionReference::GetFunction(const UBlueprint* SelfContext) const
+{
+	UClass* ContextClass = SelfContext->SkeletonGeneratedClass ? SelfContext->SkeletonGeneratedClass : SelfContext->GeneratedClass;
+	return GetFunction(ContextClass);
+}
+
+const UFunction* FMVVMBlueprintFunctionReference::GetFunction(const UClass* SelfContext) const
+{
+	if (Type != EMVVMBlueprintFunctionReferenceType::Function)
+	{
+		return nullptr;
+	}
+
+	if (FunctionReference.GetMemberName().IsNone())
+	{
+		return nullptr;
+	}
+
+	if (!FBlueprintCompilationManager::IsGeneratedClassLayoutReady())
+	{
+		SelfContext = FBlueprintEditorUtils::GetSkeletonClass(FunctionReference.GetMemberParentClass(const_cast<UClass*>(SelfContext)));
+	}
+	return FunctionReference.ResolveMember<UFunction>(const_cast<UClass*>(SelfContext), false);
+}
+
+const TSubclassOf<UK2Node> FMVVMBlueprintFunctionReference::GetNode() const
+{
+	return Type == EMVVMBlueprintFunctionReferenceType::Node ? Node : TSubclassOf<UK2Node>();
+}
+
+bool FMVVMBlueprintFunctionReference::operator==(const FMVVMBlueprintFunctionReference& Other) const
+{
+	if (Type != Other.Type)
+	{
+		return false;
+	}
+	switch (Type)
+	{
+	case EMVVMBlueprintFunctionReferenceType::Function:
+		return FunctionReference.IsSameReference(Other.FunctionReference);
+
+	case EMVVMBlueprintFunctionReferenceType::Node:
+		return Node == Other.Node;
+
+	case EMVVMBlueprintFunctionReferenceType::None:
+		return true;
+
+	default:
+		check(false);
+		return false;
+	}
+}
+
+FString FMVVMBlueprintFunctionReference::ToString() const
+{
+	switch (Type)
+	{
+		case EMVVMBlueprintFunctionReferenceType::Function:
+			return FunctionReference.GetMemberName().ToString();
+		case EMVVMBlueprintFunctionReferenceType::Node:
+			return Node.Get() ? Node.Get()->GetName() : FString();
+	}
+	return FString();
+}
+
+
+FName FMVVMBlueprintFunctionReference::GetName() const
+{
+	switch (Type)
+	{
+	case EMVVMBlueprintFunctionReferenceType::Function:
+		return FunctionReference.GetMemberName();
+	case EMVVMBlueprintFunctionReferenceType::Node:
+		return Node.Get() ? Node.Get()->GetFName() : FName();
+	}
+	return FName();
+}
+
+#undef LOCTEXT_NAMESPACE
diff --git a/Engine/Plugins/Runtime/ModelViewViewModel/Source/ModelViewViewModelBlueprint/Private/MVVMBlueprintViewBinding.cpp b/Engine/Plugins/Runtime/ModelViewViewModel/Source/ModelViewViewModelBlueprint/Private/MVVMBlueprintViewBinding.cpp
--- a/Engine/Plugins/Runtime/ModelViewViewModel/Source/ModelViewViewModelBlueprint/Private/MVVMBlueprintViewBinding.cpp
+++ b/Engine/Plugins/Runtime/ModelViewViewModel/Source/ModelViewViewModelBlueprint/Private/MVVMBlueprintViewBinding.cpp
@@ -104,10 +104,10 @@ namespace UE::MVVM::Private
 
 		if (UMVVMBlueprintViewConversionFunction* ViewConversionFunction = Binding.Conversion.GetConversionFunction(bIsSource))
 		{
-			TVariant<const UFunction*, TSubclassOf<UK2Node>> ConversionFunction = ViewConversionFunction->GetConversionFunction(WidgetBlueprint);
-			if (ConversionFunction.IsType<const UFunction*>())
+			FMVVMBlueprintFunctionReference ConversionFunction = ViewConversionFunction->GetConversionFunction();
+			if (ConversionFunction.GetType() == EMVVMBlueprintFunctionReferenceType::Function)
 			{
-				if (const UFunction* ConversionFunctionPtr = ConversionFunction.Get<const UFunction*>())
+				if (const UFunction* ConversionFunctionPtr = ConversionFunction.GetFunction(WidgetBlueprint))
 				{
 					if (bUseDisplayName)
 					{
@@ -133,9 +133,9 @@ namespace UE::MVVM::Private
 					NameBuilder << TEXT("<ERROR>");
 				}
 			}
-			else
+			else if (ConversionFunction.GetType() == EMVVMBlueprintFunctionReferenceType::Node)
 			{
-				TSubclassOf<UK2Node> ConversionFunctionNode = ConversionFunction.Get<TSubclassOf<UK2Node>>();
+				TSubclassOf<UK2Node> ConversionFunctionNode = ConversionFunction.GetNode();
 				if (ConversionFunctionNode.Get())
 				{
 					if (bUseDisplayName)
@@ -152,6 +152,10 @@ namespace UE::MVVM::Private
 					NameBuilder << TEXT("<ERROR>");
 				}
 			}
+			else
+			{
+				NameBuilder << TEXT("<ERROR>");
+			}
 
 			// AddPins
 			{
diff --git a/Engine/Plugins/Runtime/ModelViewViewModel/Source/ModelViewViewModelBlueprint/Private/MVVMBlueprintViewConversionFunction.cpp b/Engine/Plugins/Runtime/ModelViewViewModel/Source/ModelViewViewModelBlueprint/Private/MVVMBlueprintViewConversionFunction.cpp
--- a/Engine/Plugins/Runtime/ModelViewViewModel/Source/ModelViewViewModelBlueprint/Private/MVVMBlueprintViewConversionFunction.cpp
+++ b/Engine/Plugins/Runtime/ModelViewViewModel/Source/ModelViewViewModelBlueprint/Private/MVVMBlueprintViewConversionFunction.cpp
@@ -62,21 +62,21 @@ bool UMVVMBlueprintViewConversionFunction::IsValidConversionFunction(const UBlue
 
 bool UMVVMBlueprintViewConversionFunction::IsValid(const UBlueprint* SelfContext) const
 {
-	if (FunctionNode.Get())
+	if (ConversionFunction.GetType() != EMVVMBlueprintFunctionReferenceType::Function)
 	{
 		return false; // not valid for now
 	}
 	
 	UClass* SelfClass = SelfContext->SkeletonGeneratedClass ? SelfContext->SkeletonGeneratedClass : SelfContext->GeneratedClass;
-	const UFunction* ConversionFunction = FunctionReference.ResolveMember<UFunction>(SelfClass);
-	if (ConversionFunction)
+	const UFunction* Function = ConversionFunction.GetFunction(SelfContext);
+	if (Function)
 	{
-		if (!GetDefault<UMVVMDeveloperProjectSettings>()->IsConversionFunctionAllowed(SelfContext, ConversionFunction))
+		if (!GetDefault<UMVVMDeveloperProjectSettings>()->IsConversionFunctionAllowed(SelfContext, Function))
 		{
 			return false;
 		}
 
-		return IsValidConversionFunction(SelfContext, ConversionFunction);
+		return IsValidConversionFunction(SelfContext, Function);
 	}
 	return false;
 }
@@ -89,16 +89,16 @@ bool UMVVMBlueprintViewConversionFunction::NeedsWrapperGraph(const UBlueprint* S
 
 bool UMVVMBlueprintViewConversionFunction::NeedsWrapperGraphInternal(const UClass* SkeletalSelfContext) const
 {
-	if (FunctionNode.Get())
+	if (ConversionFunction.GetType() != EMVVMBlueprintFunctionReferenceType::Function)
 	{
 		return true; // not valid for now
 	}
 
 
-	const UFunction* ConversionFunction = FunctionReference.ResolveMember<UFunction>(const_cast<UClass*>(SkeletalSelfContext));
-	if (ensure(ConversionFunction))
+	const UFunction* Function = ConversionFunction.GetFunction(SkeletalSelfContext);
+	if (ensure(Function))
 	{
-		return SavedPins.Num() > 1 || !UE::MVVM::BindingHelper::IsValidForSimpleRuntimeConversion(ConversionFunction);
+		return SavedPins.Num() > 1 || !UE::MVVM::BindingHelper::IsValidForSimpleRuntimeConversion(Function);
 	}
 	return false;
 }
@@ -117,7 +117,12 @@ const UFunction* UMVVMBlueprintViewConversionFunction::GetCompiledFunction(const
 		CompiledFunction.SetSelfMember(GraphName);
 		return CompiledFunction.ResolveMember<UFunction>(const_cast<UClass*>(SelfContext));
 	}
-	return FunctionReference.ResolveMember<UFunction>(const_cast<UClass*>(SelfContext));
+
+	if (ConversionFunction.GetType() == EMVVMBlueprintFunctionReferenceType::Function)
+	{
+		return ConversionFunction.GetFunction(SelfContext);
+	}
+	return nullptr;
 }
 
 FName UMVVMBlueprintViewConversionFunction::GetCompiledFunctionName(const UClass* SelfContext) const
@@ -127,27 +132,17 @@ FName UMVVMBlueprintViewConversionFunction::GetCompiledFunctionName(const UClass
 		ensure(!GraphName.IsNone());
 		return GraphName;
 	}
-	return FunctionReference.GetMemberName();
+	return ConversionFunction.GetName();
 }
 
-TVariant<const UFunction*, TSubclassOf<UK2Node>> UMVVMBlueprintViewConversionFunction::GetConversionFunction(const UBlueprint* SelfContext) const
+FMVVMBlueprintFunctionReference UMVVMBlueprintViewConversionFunction::GetConversionFunction() const
 {
-	if (FunctionNode.Get())
-	{
-		return TVariant<const UFunction*, TSubclassOf<UK2Node>>( TInPlaceType<TSubclassOf<UK2Node>>(), FunctionNode);
-	}
-	else
-	{
-		UClass* SelfClass = SelfContext->SkeletonGeneratedClass ? SelfContext->SkeletonGeneratedClass : SelfContext->GeneratedClass;
-		const UFunction* ConversionFunction = FunctionReference.ResolveMember<UFunction>(SelfClass);
-		return TVariant<const UFunction*, TSubclassOf<UK2Node>>( TInPlaceType<const UFunction*>(), ConversionFunction);
-	}
+	return ConversionFunction;
 }
 
 void UMVVMBlueprintViewConversionFunction::Reset()
 {
-	FunctionReference = FMemberReference();
-	FunctionNode = TSubclassOf<UK2Node>();
+	ConversionFunction = FMVVMBlueprintFunctionReference();
 	GraphName = FName();
 	bWrapperGraphTransient = false;
 	SavedPins.Reset();
@@ -160,26 +155,7 @@ void UMVVMBlueprintViewConversionFunction::InitializeFromFunction(UBlueprint* In
 
 	if (ensure(InFunction && !InGraphName.IsNone()))
 	{
-		UClass* OwnerClass = InFunction->GetOuterUClass();
-		FGuid MemberGuid;
-		UBlueprint::GetGuidFromClassByFieldName<UFunction>(OwnerClass, InFunction->GetFName(), MemberGuid);
-
-		bool bIsSelf = (InContext->GeneratedClass && InContext->GeneratedClass->IsChildOf(OwnerClass))
-			|| (InContext->SkeletonGeneratedClass && InContext->SkeletonGeneratedClass->IsChildOf(OwnerClass));
-		if (bIsSelf)
-		{
-			FunctionReference.SetSelfMember(InFunction->GetFName(), MemberGuid);
-		}
-		else
-		{
-			if (UBlueprint* VariableOwnerBP = Cast<UBlueprint>(OwnerClass->ClassGeneratedBy))
-			{
-				OwnerClass = VariableOwnerBP->SkeletonGeneratedClass ? VariableOwnerBP->SkeletonGeneratedClass : VariableOwnerBP->GeneratedClass;
-			}
-
-			FunctionReference.SetExternalMember(InFunction->GetFName(), OwnerClass, MemberGuid);
-		}
-
+		ConversionFunction = FMVVMBlueprintFunctionReference(InContext, InFunction);
 		check(GraphName.IsNone()); // the name needs to be set before a GetOrCreateWrapperGraph
 		GraphName = InGraphName;
 		bWrapperGraphTransient = !GetDefault<UMVVMDeveloperProjectSettings>()->bAllowConversionFunctionGeneratedGraphInEditor;
@@ -196,11 +172,11 @@ void UMVVMBlueprintViewConversionFunction::Deprecation_InitializeFromWrapperGrap
 	{
 		if (UK2Node_CallFunction* CallFunction = Cast<UK2Node_CallFunction>(WrapperNode))
 		{
-			FunctionReference = CallFunction->FunctionReference;
+			ConversionFunction = FMVVMBlueprintFunctionReference(CallFunction->FunctionReference);
 		}
 		else
 		{
-			FunctionNode = WrapperNode->GetClass();
+			ConversionFunction = FMVVMBlueprintFunctionReference(WrapperNode->GetClass());
 		}
 		check(WrapperNode->GetGraph());
 		SetCachedWrapperGraph(SelfContext, WrapperNode->GetGraph(), WrapperNode);
@@ -222,7 +198,7 @@ void UMVVMBlueprintViewConversionFunction::Deprecation_InitializeFromMemberRefer
 {
 	Reset();
 
-	FunctionReference = MemberReference;
+	ConversionFunction = FMVVMBlueprintFunctionReference(MemberReference);
 
 	check(GraphName.IsNone()); // the name needs to be set before a GetOrCreateWrapperGraph
 	GraphName = InGraphName;
@@ -230,8 +206,8 @@ void UMVVMBlueprintViewConversionFunction::Deprecation_InitializeFromMemberRefer
 
 	// since it is a new object, we can't create a the graph right away
 	UClass* GeneratedClass = SelfContext->SkeletonGeneratedClass ? SelfContext->SkeletonGeneratedClass : SelfContext->GeneratedClass;
-	const UFunction* ConversionFunction = FunctionReference.ResolveMember<UFunction>(GeneratedClass);
-	const FProperty* PinProperty = ConversionFunction ? UE::MVVM::BindingHelper::GetFirstArgumentProperty(ConversionFunction) : nullptr;
+	const UFunction* Function = ConversionFunction.GetFunction(GeneratedClass);
+	const FProperty* PinProperty = Function ? UE::MVVM::BindingHelper::GetFirstArgumentProperty(Function) : nullptr;
 	if (PinProperty)
 	{
 		FName PinPropertyName = PinProperty->GetFName();
@@ -254,8 +230,8 @@ void UMVVMBlueprintViewConversionFunction::Deprecation_SetWrapperGraphName(UBlue
 
 		// since it is a new object, we can't create a the graph right away
 		UClass* GeneratedClass = SelfContext->SkeletonGeneratedClass ? SelfContext->SkeletonGeneratedClass : SelfContext->GeneratedClass;
-		const UFunction* ConversionFunction = FunctionReference.ResolveMember<UFunction>(GeneratedClass);
-		const FProperty* PinProperty = ConversionFunction ? UE::MVVM::BindingHelper::GetFirstArgumentProperty(ConversionFunction) : nullptr;
+		const UFunction* Function = ConversionFunction.GetFunction(GeneratedClass);
+		const FProperty* PinProperty = Function ? UE::MVVM::BindingHelper::GetFirstArgumentProperty(Function) : nullptr;
 		if (PinProperty)
 		{
 			FName PinPropertyName = PinProperty->GetFName();
@@ -330,15 +306,14 @@ UEdGraph* UMVVMBlueprintViewConversionFunction::GetOrCreateIntermediateWrapperGr
 	}
 	else if (IsValid(Context.Blueprint))
 	{
-		TVariant<const UFunction*, TSubclassOf<UK2Node>> ConversionFunction = GetConversionFunction(Context.Blueprint);
-		if (ConversionFunction.IsType<TSubclassOf<UK2Node>>())
+		if (ConversionFunction.GetType() != EMVVMBlueprintFunctionReferenceType::Function)
 		{
 			ensure(false); // not supported for now
 			return nullptr;
 		}
 
 		CreateWrapperGraphName();
-		return GetOrCreateWrapperGraphInternal(Context, ConversionFunction.Get<const UFunction*>());
+		return GetOrCreateWrapperGraphInternal(Context, ConversionFunction.GetFunction(Context.Blueprint));
 	}
 	return nullptr;
 }
@@ -361,15 +336,14 @@ UEdGraph* UMVVMBlueprintViewConversionFunction::GetOrCreateWrapperGraph(UBluepri
 	}
 	else if (IsValid(Blueprint))
 	{
-		TVariant<const UFunction*, TSubclassOf<UK2Node>> ConversionFunction = GetConversionFunction(Blueprint);
-		if (ConversionFunction.IsType<TSubclassOf<UK2Node>>())
+		if (ConversionFunction.GetType() != EMVVMBlueprintFunctionReferenceType::Function)
 		{
 			ensure(false); // not supported for now
 			return nullptr;
 		}
 
 		CreateWrapperGraphName();
-		return GetOrCreateWrapperGraphInternal(Blueprint, ConversionFunction.Get<const UFunction*>());
+		return GetOrCreateWrapperGraphInternal(Blueprint, ConversionFunction.GetFunction(Blueprint));
 	}
 	return nullptr;
 }
@@ -380,17 +354,17 @@ UEdGraphPin* UMVVMBlueprintViewConversionFunction::GetOrCreateGraphPin(UBlueprin
 	return CachedWrapperGraph ? UE::MVVM::ConversionFunctionHelper::FindPin(CachedWrapperGraph, PinId.GetNames()) : nullptr;
 }
 
-UEdGraph* UMVVMBlueprintViewConversionFunction::GetOrCreateWrapperGraphInternal(FKismetCompilerContext& Context, const UFunction* Function)
+UEdGraph* UMVVMBlueprintViewConversionFunction::GetOrCreateWrapperGraphInternal(FKismetCompilerContext& Context, const UFunction* InFunction)
 {
-	return GetOrCreateWrapperGraphInternal(Context.Blueprint, Function);
+	return GetOrCreateWrapperGraphInternal(Context.Blueprint, InFunction);
 }
 
-UEdGraph* UMVVMBlueprintViewConversionFunction::GetOrCreateWrapperGraphInternal(UBlueprint* Blueprint, const UFunction* Function)
+UEdGraph* UMVVMBlueprintViewConversionFunction::GetOrCreateWrapperGraphInternal(UBlueprint* Blueprint, const UFunction* InFunction)
 {
 	check(Blueprint);
-	check(Function);
+	check(InFunction);
 	bool bConst = true;
-	TPair<UEdGraph*, UK2Node*> Result = UE::MVVM::ConversionFunctionHelper::CreateGraph(Blueprint, GraphName, nullptr, Function, bConst, bWrapperGraphTransient);
+	TPair<UEdGraph*, UK2Node*> Result = UE::MVVM::ConversionFunctionHelper::CreateGraph(Blueprint, GraphName, nullptr, InFunction, bConst, bWrapperGraphTransient);
 	const_cast<UMVVMBlueprintViewConversionFunction*>(this)->SetCachedWrapperGraph(Blueprint, Result.Get<0>(), Result.Get<1>());
 	LoadPinValuesInternal(Blueprint);
 	return CachedWrapperGraph;
@@ -501,3 +475,19 @@ void UMVVMBlueprintViewConversionFunction::HandleUserDefinedPinRenamed(UK2Node*
 		OnWrapperGraphModified.Broadcast();
 	}
 }
+
+void UMVVMBlueprintViewConversionFunction::PostLoad()
+{
+	Super::PostLoad();
+
+	if (FunctionNode_DEPRECATED.Get())
+	{
+		ConversionFunction = FMVVMBlueprintFunctionReference(FunctionNode_DEPRECATED.Get());
+		FunctionNode_DEPRECATED = TSubclassOf<UK2Node>();
+	}
+	else if (!FunctionReference_DEPRECATED.GetMemberName().IsNone())
+	{
+		ConversionFunction = FMVVMBlueprintFunctionReference(FunctionReference_DEPRECATED);
+		FunctionReference_DEPRECATED = FMemberReference();
+	}
+}
diff --git a/Engine/Plugins/Runtime/ModelViewViewModel/Source/ModelViewViewModelBlueprint/Private/MVVMViewBlueprintCompiler.cpp b/Engine/Plugins/Runtime/ModelViewViewModel/Source/ModelViewViewModelBlueprint/Private/MVVMViewBlueprintCompiler.cpp
--- a/Engine/Plugins/Runtime/ModelViewViewModel/Source/ModelViewViewModelBlueprint/Private/MVVMViewBlueprintCompiler.cpp
+++ b/Engine/Plugins/Runtime/ModelViewViewModel/Source/ModelViewViewModelBlueprint/Private/MVVMViewBlueprintCompiler.cpp
@@ -2221,13 +2221,14 @@ void FMVVMViewBlueprintCompiler::PreCompileBindings(UWidgetBlueprintGeneratedCla
 					return false;
 				}
 
-				TVariant<const UFunction*, TSubclassOf<UK2Node>> FunctionOrWrapperFunction = ViewConversionFunction->GetConversionFunction(Self->WidgetBlueprintCompilerContext.WidgetBlueprint());
-				if (FunctionOrWrapperFunction.IsType<const UFunction*>())
+				FMVVMBlueprintFunctionReference FunctionOrWrapperFunction = ViewConversionFunction->GetConversionFunction();
+				if (FunctionOrWrapperFunction.GetType() == EMVVMBlueprintFunctionReferenceType::Function)
 				{
-					if (!GetDefault<UMVVMDeveloperProjectSettings>()->IsConversionFunctionAllowed(Self->WidgetBlueprintCompilerContext.WidgetBlueprint(), FunctionOrWrapperFunction.Get<const UFunction*>()))
+					const UFunction* WrapperFunction = FunctionOrWrapperFunction.GetFunction(Self->WidgetBlueprintCompilerContext.WidgetBlueprint());
+					if (!GetDefault<UMVVMDeveloperProjectSettings>()->IsConversionFunctionAllowed(Self->WidgetBlueprintCompilerContext.WidgetBlueprint(), WrapperFunction))
 					{
 						Self->AddMessageForBinding(Binding
-							, FText::Format(LOCTEXT("ConversionFunctionNotAllow", "The conversion function {0} is not allowed."), FText::FromName(FunctionOrWrapperFunction.Get<const UFunction*>()->GetFName()))
+							, FText::Format(LOCTEXT("ConversionFunctionNotAllow", "The conversion function {0} is not allowed."), FText::FromName(WrapperFunction->GetFName()))
 							, EMessageType::Error
 							, FMVVMBlueprintPinId()
 						);
diff --git a/Engine/Plugins/Runtime/ModelViewViewModel/Source/ModelViewViewModelBlueprint/Public/MVVMBlueprintFunctionReference.h b/Engine/Plugins/Runtime/ModelViewViewModel/Source/ModelViewViewModelBlueprint/Public/MVVMBlueprintFunctionReference.h
new file mode 100644
--- /dev/null
+++ b/Engine/Plugins/Runtime/ModelViewViewModel/Source/ModelViewViewModelBlueprint/Public/MVVMBlueprintFunctionReference.h
@@ -0,0 +1,75 @@
+// Copyright Epic Games, Inc. All Rights Reserved.
+
+#pragma once
+
+#include "Engine/MemberReference.h"
+#include "K2Node.h"
+#include "Templates/SubclassOf.h"
+#include "MVVMBlueprintFunctionReference.generated.h"
+
+class UBlueprint;
+class UFunction;
+
+/**
+ *
+ */
+UENUM()
+enum class EMVVMBlueprintFunctionReferenceType : uint8
+{
+	None,
+	Function,
+	Node
+};
+
+/**
+ * A type that point to a function or a node that can be used to create a conversion function.
+ */
+USTRUCT(BlueprintType)
+struct FMVVMBlueprintFunctionReference
+{
+	GENERATED_BODY()
+
+public:
+	FMVVMBlueprintFunctionReference() = default;
+	MODELVIEWVIEWMODELBLUEPRINT_API explicit FMVVMBlueprintFunctionReference(const UBlueprint* InContext, const UFunction* InFunction);
+	MODELVIEWVIEWMODELBLUEPRINT_API explicit FMVVMBlueprintFunctionReference(FMemberReference InFunctionReference);
+	MODELVIEWVIEWMODELBLUEPRINT_API explicit FMVVMBlueprintFunctionReference(const TSubclassOf<UK2Node> InNode);
+
+	/** If the node is of type Function, resolves the function. Uses the blueprint, if the function is defined on the BP itself. */
+	MODELVIEWVIEWMODELBLUEPRINT_API const UFunction* GetFunction(const UBlueprint* SelfContext) const;
+	/** If the node is of type Function, resolves the function. Uses the blueprint, if the function is defined on the BP itself. */
+	MODELVIEWVIEWMODELBLUEPRINT_API const UFunction* GetFunction(const UClass* SelfContext) const;
+	/** If the node is of type Node, return the node class. */
+	MODELVIEWVIEWMODELBLUEPRINT_API const TSubclassOf<UK2Node> GetNode() const;
+
+	EMVVMBlueprintFunctionReferenceType GetType() const
+	{
+		return Type;
+	}
+
+	/** return a string representation of the reference. */
+	MODELVIEWVIEWMODELBLUEPRINT_API FString ToString() const;
+
+	/** return a name of the reference. */
+	MODELVIEWVIEWMODELBLUEPRINT_API FName GetName() const;
+
+	MODELVIEWVIEWMODELBLUEPRINT_API bool operator==(const FMVVMBlueprintFunctionReference& Other) const;
+
+	friend uint32 GetTypeHash(const FMVVMBlueprintFunctionReference& Value)
+	{
+		uint32 Hash = GetTypeHash(static_cast<uint8>(Value.Type));
+		Hash = HashCombine(Hash, GetTypeHash(Value.FunctionReference.GetMemberName()));
+		Hash = HashCombine(Hash, GetTypeHash(Value.Node));
+		return Hash;
+	}
+
+private:
+	UPROPERTY(VisibleAnywhere, Category = "MVVM")
+	FMemberReference FunctionReference;
+
+	UPROPERTY(VisibleAnywhere, Category = "MVVM")
+	TSubclassOf<UK2Node> Node;
+
+	UPROPERTY(VisibleAnywhere, Category = "MVVM")
+	EMVVMBlueprintFunctionReferenceType Type = EMVVMBlueprintFunctionReferenceType::None;
+};
\ No newline at end of file
diff --git a/Engine/Plugins/Runtime/ModelViewViewModel/Source/ModelViewViewModelBlueprint/Public/MVVMBlueprintViewConversionFunction.h b/Engine/Plugins/Runtime/ModelViewViewModel/Source/ModelViewViewModelBlueprint/Public/MVVMBlueprintViewConversionFunction.h
--- a/Engine/Plugins/Runtime/ModelViewViewModel/Source/ModelViewViewModelBlueprint/Public/MVVMBlueprintViewConversionFunction.h
+++ b/Engine/Plugins/Runtime/ModelViewViewModel/Source/ModelViewViewModelBlueprint/Public/MVVMBlueprintViewConversionFunction.h
@@ -2,6 +2,7 @@
 
 #pragma once
 
+#include "MVVMBlueprintFunctionReference.h"
 #include "MVVMBlueprintPin.h"
 #include "MVVMBlueprintView.h"
 #include "Engine/MemberReference.h"
@@ -34,8 +35,12 @@ public:
 	FName GetCompiledFunctionName(const UClass* SelfContext) const;
 
 	/** @return the conversion function. */
+	UE_DEPRECATED(5.4, "GetConversionFunction that returns a variant is deprecated.")
 	TVariant<const UFunction*, TSubclassOf<UK2Node>> GetConversionFunction(const UBlueprint* SelfContext) const;
 
+	/** @return the conversion function. */
+	FMVVMBlueprintFunctionReference GetConversionFunction() const;
+
 	/** Set the function. Generate a Graph. */
 	void InitializeFromFunction(UBlueprint* SelfContext, FName GraphName, const UFunction* Function);
 
@@ -72,11 +77,6 @@ public:
 		return GraphName;
 	}
 
-	UK2Node* GetWrapperNode() const
-	{
-		return CachedWrapperNode;
-	}
-
 	/**
 	 * If needed, create the graph and all the nodes for that graph when compiling.
 	 * Returns the existing one, if one was created from GetOrCreateWrapperGraph.
@@ -115,6 +115,8 @@ public:
 
 	FSimpleMulticastDelegate OnWrapperGraphModified;
 
+	virtual void PostLoad() override;
+
 private:
 	void HandleGraphChanged(const FEdGraphEditAction& Action, TWeakObjectPtr<UBlueprint> Context);
 	void HandleUserDefinedPinRenamed(UK2Node* InNode, FName OldPinName, FName NewPinName, TWeakObjectPtr<UBlueprint> WeakBlueprint);
@@ -128,19 +130,11 @@ private:
 
 private:
 	/**
-	 * The conversion UFunction when simple or when it's complex.
-	 * @note Only one of FunctionReference or the GraphNode can be valid.
+	 * Conversion reference. It can be simple, complex or a K2Node.
+	 * @note The conversion is complex
 	 */
 	UPROPERTY(VisibleAnywhere, Category = "Viewmodel")
-	FMemberReference FunctionReference;
-	
-	/**
-	 * The conversion K2Node the graph is generated for.
-	 * @note Only one of FunctionReference or the GraphNode can be valid.
-	 * @note The conversion is complex.
-	 */
-	UPROPERTY(VisibleAnywhere, Category = "Viewmodel")
-	TSubclassOf<UK2Node> FunctionNode;
+	FMVVMBlueprintFunctionReference ConversionFunction;
 
 	/** Name of the generated graph if a wrapper is needed. */
 	UPROPERTY(VisibleAnywhere, Category = "Viewmodel")
@@ -164,4 +158,10 @@ private:
 
 	FDelegateHandle OnGraphChangedHandle;
 	FDelegateHandle OnUserDefinedPinRenamedHandle;
+
+
+	UPROPERTY()
+	FMemberReference FunctionReference_DEPRECATED;
+	UPROPERTY()
+	TSubclassOf<UK2Node> FunctionNode_DEPRECATED;
 };
diff --git a/Engine/Plugins/Runtime/ModelViewViewModel/Source/ModelViewViewModelEditor/Private/MVVMEditorSubsystem.cpp b/Engine/Plugins/Runtime/ModelViewViewModel/Source/ModelViewViewModelEditor/Private/MVVMEditorSubsystem.cpp
--- a/Engine/Plugins/Runtime/ModelViewViewModel/Source/ModelViewViewModelEditor/Private/MVVMEditorSubsystem.cpp
+++ b/Engine/Plugins/Runtime/ModelViewViewModel/Source/ModelViewViewModelEditor/Private/MVVMEditorSubsystem.cpp
@@ -338,10 +338,10 @@ UFunction* UMVVMEditorSubsystem::GetConversionFunction(const UWidgetBlueprint* W
 {
 	if (UMVVMBlueprintViewConversionFunction* ConversionFunction = Binding.Conversion.GetConversionFunction(bSourceToDestination))
 	{
-		TVariant<const UFunction*, TSubclassOf<UK2Node>> Result = ConversionFunction->GetConversionFunction(WidgetBlueprint);
-		if (Result.IsType<const UFunction*>())
+		FMVVMBlueprintFunctionReference Result = ConversionFunction->GetConversionFunction();
+		if (Result.GetType() == EMVVMBlueprintFunctionReferenceType::Function)
 		{
-			return const_cast<UFunction*>(Result.Get<const UFunction*>());
+			return const_cast<UFunction*>(Result.GetFunction(WidgetBlueprint));
 		}
 	}
 	return nullptr;
@@ -762,10 +762,6 @@ UEdGraph* UMVVMEditorSubsystem::GetConversionFunctionGraph(const UWidgetBlueprin
 
 UK2Node_CallFunction* UMVVMEditorSubsystem::GetConversionFunctionNode(const UWidgetBlueprint* WidgetBlueprint, const FMVVMBlueprintViewBinding& Binding, bool bSourceToDestination) const
 {
-	if (const UMVVMBlueprintViewConversionFunction* Found = Binding.Conversion.GetConversionFunction(bSourceToDestination))
-	{
-		return Cast<UK2Node_CallFunction>(Found->GetWrapperNode());
-	}
 	return nullptr;
 }
 
diff --git a/Engine/Plugins/Runtime/ModelViewViewModel/Source/ModelViewViewModelEditor/Private/Widgets/BindingEntry/SMVVMBindingRow.cpp b/Engine/Plugins/Runtime/ModelViewViewModel/Source/ModelViewViewModelEditor/Private/Widgets/BindingEntry/SMVVMBindingRow.cpp
--- a/Engine/Plugins/Runtime/ModelViewViewModel/Source/ModelViewViewModelEditor/Private/Widgets/BindingEntry/SMVVMBindingRow.cpp
+++ b/Engine/Plugins/Runtime/ModelViewViewModel/Source/ModelViewViewModelEditor/Private/Widgets/BindingEntry/SMVVMBindingRow.cpp
@@ -6,6 +6,7 @@
 #include "Bindings/MVVMBindingHelper.h"
 #include "Framework/MVVMRowHelper.h"
 #include "MVVMBlueprintViewBinding.h"
+#include "MVVMBlueprintViewConversionFunction.h"
 #include "MVVMDeveloperProjectSettings.h"
 #include "MVVMPropertyPath.h"
 
@@ -384,8 +385,13 @@ const UFunction* SBindingRow::GetSelectedConversionFunction(bool bSourceToDest)
 {
 	if (FMVVMBlueprintViewBinding* ViewBinding = GetThisViewBinding())
 	{
-		UMVVMEditorSubsystem* EditorSubsystem = GetEditorSubsystem();
-		return EditorSubsystem->GetConversionFunction(GetBlueprint(), *ViewBinding, bSourceToDest);
+		if (UMVVMBlueprintViewConversionFunction* ConversionFunction = ViewBinding->Conversion.GetConversionFunction(bSourceToDest))
+		{
+			if (ConversionFunction->GetConversionFunction().GetType() == EMVVMBlueprintFunctionReferenceType::Function)
+			{
+				return ConversionFunction->GetConversionFunction().GetFunction(GetBlueprint());
+			}
+		}
 	}
 	return nullptr;
 }
diff --git a/Engine/Plugins/Runtime/ModelViewViewModel/Source/ModelViewViewModelEditor/Private/Widgets/SMVVMFunctionParameter.cpp b/Engine/Plugins/Runtime/ModelViewViewModel/Source/ModelViewViewModelEditor/Private/Widgets/SMVVMFunctionParameter.cpp
--- a/Engine/Plugins/Runtime/ModelViewViewModel/Source/ModelViewViewModelEditor/Private/Widgets/SMVVMFunctionParameter.cpp
+++ b/Engine/Plugins/Runtime/ModelViewViewModel/Source/ModelViewViewModelEditor/Private/Widgets/SMVVMFunctionParameter.cpp
@@ -3,12 +3,15 @@
 #include "SMVVMFunctionParameter.h" 
 #include "EdGraphSchema_K2.h"
 #include "Editor.h"
+#include "MVVMBlueprintFunctionReference.h"
 #include "MVVMBlueprintView.h"
+#include "MVVMBlueprintViewConversionFunction.h"
 #include "MVVMEditorSubsystem.h"
 #include "NodeFactory.h"
 #include "SGraphPin.h"
 #include "Styling/MVVMEditorStyle.h"
 #include "Types/MVVMBindingMode.h"
+
 #include "WidgetBlueprint.h"
 #include "Widgets/Images/SImage.h"
 #include "Widgets/SBoxPanel.h"
@@ -19,7 +22,7 @@
 #include "Widgets/SOverlay.h"
 #include "Widgets/Text/STextBlock.h"
 
-#define LOCTEXT_NAMESPACE "MVVMFieldBinding"
+#define LOCTEXT_NAMESPACE "FunctionParameter"
 
 namespace UE::MVVM
 {
@@ -218,14 +221,25 @@ FFieldSelectionContext SFunctionParameter::GetSelectedSelectionContext() const
 		return Result;
 	}
 
-	const UFunction* ConversionFunction = GEditor->GetEditorSubsystem<UMVVMEditorSubsystem>()->GetConversionFunction(WidgetBlueprintPtr, *Binding, bSourceToDestination);
+	UMVVMBlueprintViewConversionFunction* ConversionFunction = Binding->Conversion.GetConversionFunction(bSourceToDestination);
 	if (ConversionFunction == nullptr)
 	{
 		return Result;
 	}
 
+	FProperty* PinProperty = nullptr;
+	if (ParameterId.GetNames().Num() > 0)
+	{
+		FMVVMBlueprintFunctionReference FunctionReference = ConversionFunction->GetConversionFunction();
+		if (FunctionReference.GetType() == EMVVMBlueprintFunctionReferenceType::Function)
+		{
+			const UFunction* Function = FunctionReference.GetFunction(WidgetBlueprintPtr);
+			PinProperty = Function ? Function->FindPropertyByName(ParameterId.GetNames().Last()) : nullptr;
+		}
+	}
+
 	Result.BindingMode = EMVVMBindingMode::OneWayToDestination;
-	Result.AssignableTo = ParameterId.GetNames().Num() > 0 ? ConversionFunction->FindPropertyByName(ParameterId.GetNames().Last()) : nullptr;
+	Result.AssignableTo = PinProperty;
 	Result.bAllowWidgets = true;
 	Result.bAllowViewModels = true;
 	Result.bAllowConversionFunctions = false;
diff --git a/Engine/Plugins/Runtime/ModelViewViewModel/Source/ModelViewViewModelEditor/Public/MVVMEditorSubsystem.h b/Engine/Plugins/Runtime/ModelViewViewModel/Source/ModelViewViewModelEditor/Public/MVVMEditorSubsystem.h
--- a/Engine/Plugins/Runtime/ModelViewViewModel/Source/ModelViewViewModelEditor/Public/MVVMEditorSubsystem.h
+++ b/Engine/Plugins/Runtime/ModelViewViewModel/Source/ModelViewViewModelEditor/Public/MVVMEditorSubsystem.h
@@ -94,6 +94,7 @@ public:
 	UFUNCTION(BlueprintCallable, Category = "Viewmodel")
 	UEdGraph* GetConversionFunctionGraph(const UWidgetBlueprint* WidgetBlueprint, const FMVVMBlueprintViewBinding& Binding, bool bSourceToDestination) const;
 
+	UE_DEPRECATED(5.4, "GetConversionFunction was moved to MVVMBlueprintViewConversionFunction.")
 	UFUNCTION(BlueprintCallable, Category = "Viewmodel")
 	UFunction* GetConversionFunction(const UWidgetBlueprint* WidgetBlueprint, const FMVVMBlueprintViewBinding& Binding, bool bSourceToDestination) const;
 
