Change the layout to include the notification bell

diff --git a/Engine/Plugins/Runtime/ModelViewViewModel/Source/ModelViewViewModelEditor/Private/Widgets/SMVVMFieldEntry.cpp b/Engine/Plugins/Runtime/ModelViewViewModel/Source/ModelViewViewModelEditor/Private/Widgets/SMVVMFieldEntry.cpp
--- a/Engine/Plugins/Runtime/ModelViewViewModel/Source/ModelViewViewModelEditor/Private/Widgets/SMVVMFieldEntry.cpp
+++ b/Engine/Plugins/Runtime/ModelViewViewModel/Source/ModelViewViewModelEditor/Private/Widgets/SMVVMFieldEntry.cpp
@@ -127,6 +127,8 @@ void SFieldPaths::Construct(const FArguments& InArgs)
 {
 	TextStyle = InArgs._TextStyle;
 
+	HighlightField = InArgs._HighlightField;
+
 	ChildSlot
 	[
 		SAssignNew(FieldBox, SHorizontalBox)
@@ -163,15 +165,29 @@ void SFieldPaths::SetFieldPaths(TArrayView<UE::MVVM::FMVVMConstFieldVariant> InP
 					.Field(InPropertyPath[Index])
 				];
 
+			bool bShowFieldNotify = HighlightField.IsSet() && HighlightField.GetValue().Contains(Index);
+			if (bShowFieldNotify)
+			{
+				FieldBox->AddSlot()
+					.HAlign(HAlign_Fill)
+					.VAlign(VAlign_Center)
+					.AutoWidth()
+					[
+						SNew(SImage)
+						.DesiredSizeOverride(FVector2D(16.0f, 16.0f))
+						.Image(FAppStyle::GetBrush("Kismet.VariableList.FieldNotify"))
+					];
+			}
+
 			FieldBox->AddSlot()
 				.HAlign(HAlign_Fill)
 				.VAlign(VAlign_Center)
-				.Padding(4, 0, 0, 0)
 				.AutoWidth()
 				[
 					SNew(STextBlock)
 					.TextStyle(TextStyle)
 					.Clipping(EWidgetClipping::OnDemand)
+					.SimpleTextMode(true)
 					.Text(Private::GetFieldDisplayName(InPropertyPath[Index]))
 				];
 		}
@@ -185,6 +201,7 @@ void SFieldPaths::SetFieldPaths(TArrayView<UE::MVVM::FMVVMConstFieldVariant> InP
 					SNew(STextBlock)
 					.TextStyle(TextStyle)
 					.Clipping(EWidgetClipping::OnDemand)
+					.SimpleTextMode(true)
 					.Text(LOCTEXT("Invalid", "Invalid Field"))
 				];
 		}
@@ -195,7 +212,6 @@ void SFieldPaths::SetFieldPaths(TArrayView<UE::MVVM::FMVVMConstFieldVariant> InP
 			FieldBox->AddSlot()
 				.HAlign(HAlign_Left)
 				.VAlign(VAlign_Center)
-				.Padding(6, 0)
 				.AutoWidth()
 				[
 					SNew(SImage)
diff --git a/Engine/Plugins/Runtime/ModelViewViewModel/Source/ModelViewViewModelEditor/Private/Widgets/SMVVMFieldEntry.h b/Engine/Plugins/Runtime/ModelViewViewModel/Source/ModelViewViewModelEditor/Private/Widgets/SMVVMFieldEntry.h
--- a/Engine/Plugins/Runtime/ModelViewViewModel/Source/ModelViewViewModelEditor/Private/Widgets/SMVVMFieldEntry.h
+++ b/Engine/Plugins/Runtime/ModelViewViewModel/Source/ModelViewViewModelEditor/Private/Widgets/SMVVMFieldEntry.h
@@ -27,6 +27,7 @@ public:
 		SLATE_STYLE_ARGUMENT(FTextBlockStyle, TextStyle)
 		SLATE_ARGUMENT(TArrayView<UE::MVVM::FMVVMConstFieldVariant>, FieldPaths)
 		SLATE_ARGUMENT(bool, ShowOnlyLast)
+		SLATE_ARGUMENT(TOptional<FInt32Range>, HighlightField)
 	SLATE_END_ARGS()
 
 	void Construct(const FArguments& InArgs);
@@ -35,6 +36,7 @@ public:
 private:
 	const FTextBlockStyle* TextStyle = nullptr;
 	TSharedPtr<SHorizontalBox> FieldBox;
+	TOptional<FInt32Range> HighlightField;
 	bool bShowOnlyLast = false;
 };
 
diff --git a/Engine/Plugins/Runtime/ModelViewViewModel/Source/ModelViewViewModelEditor/Private/Widgets/SMVVMPropertyPath.cpp b/Engine/Plugins/Runtime/ModelViewViewModel/Source/ModelViewViewModelEditor/Private/Widgets/SMVVMPropertyPath.cpp
--- a/Engine/Plugins/Runtime/ModelViewViewModel/Source/ModelViewViewModelEditor/Private/Widgets/SMVVMPropertyPath.cpp
+++ b/Engine/Plugins/Runtime/ModelViewViewModel/Source/ModelViewViewModelEditor/Private/Widgets/SMVVMPropertyPath.cpp
@@ -2,6 +2,7 @@
 
 #include "Widgets/SMVVMPropertyPath.h"
 
+#include "Bindings/MVVMFieldPathHelper.h"
 #include "Blueprint/WidgetTree.h"
 #include "Editor.h"
 #include "MVVMBlueprintView.h"
@@ -9,6 +10,7 @@
 #include "Styling/MVVMEditorStyle.h"
 #include "Types/MVVMBindingSource.h"
 #include "WidgetBlueprint.h"
+
 #include "Widgets/SMVVMFieldEntry.h"
 #include "Widgets/SMVVMFieldIcon.h"
 #include "Widgets/SMVVMSourceEntry.h"
@@ -40,47 +42,86 @@ void SPropertyPath::SetPropertyPath(const FMVVMBlueprintPropertyPath& InProperty
 {
 	FieldBox->ClearChildren();
 
-	bool bHasSource = InPropertyPath.IsValid();
-	if (bShowContext && bHasSource)
+	const UWidgetBlueprint* WidgetBlueprintPtr = WidgetBlueprint.Get();
+	const UClass* ClassContext = WidgetBlueprintPtr->SkeletonGeneratedClass ? WidgetBlueprintPtr->SkeletonGeneratedClass : WidgetBlueprintPtr->GeneratedClass;
+	EMVVMBlueprintFieldPathSource Source = InPropertyPath.GetSource(WidgetBlueprintPtr);
+
+	if (Source != EMVVMBlueprintFieldPathSource::None)
 	{
-		FieldBox->AddSlot()
-			.Padding(8, 0, 0, 0)
-			.HAlign(HAlign_Left)
-			.VAlign(VAlign_Center)
-			.AutoWidth()
-			[
-				SNew(SBindingContextEntry)
-				.TextStyle(TextStyle)
-				.BindingContext(FBindingSource::CreateFromPropertyPath(WidgetBlueprint.Get(), InPropertyPath))
-			];
+		TArray<UE::MVVM::FMVVMConstFieldVariant> AllFields = InPropertyPath.GetCompleteFields(WidgetBlueprintPtr);
+
+		// Find the NotifyFieldId to highlight it. Error are valid.
+		TOptional<FInt32Range> HighlightRange;
+		TValueOrError<UE::MVVM::FieldPathHelper::FParsedNotifyBindingInfo, FText> FieldPathResult = UE::MVVM::FieldPathHelper::GetNotifyBindingInfoFromFieldPath(ClassContext, AllFields);
+		if (FieldPathResult.HasValue())
+		{
+			int32 RangeStart = FieldPathResult.GetValue().ViewModelIndex;
+			HighlightRange = FInt32Range(RangeStart, RangeStart + 1);
+		}
+
+		// Remove the first item to show it as a SBindingContextEntry
+		if ((Source == EMVVMBlueprintFieldPathSource::Widget || Source == EMVVMBlueprintFieldPathSource::ViewModel) && AllFields.Num() > 0)
+		{
+			AllFields.RemoveAt(0);
+		}
+
+
+		bool bHasSource = InPropertyPath.IsValid();
+		if (bShowContext && bHasSource)
+		{
+			FieldBox->AddSlot()
+				.HAlign(HAlign_Left)
+				.VAlign(VAlign_Center)
+				.AutoWidth()
+				[
+					SNew(SBindingContextEntry)
+					.TextStyle(TextStyle)
+					.BindingContext(FBindingSource::CreateFromPropertyPath(WidgetBlueprintPtr, InPropertyPath))
+				];
+
+			if (InPropertyPath.HasPaths())
+			{
+				FieldBox->AddSlot()
+					.HAlign(HAlign_Left)
+					.VAlign(VAlign_Center)
+					.AutoWidth()
+					[
+						SNew(SImage)
+						.Image(FAppStyle::Get().GetBrush("Icons.ChevronRight"))
+					];
+			}
+		}
 
 		if (InPropertyPath.HasPaths())
 		{
 			FieldBox->AddSlot()
-				.Padding(6, 0)
 				.HAlign(HAlign_Left)
 				.VAlign(VAlign_Center)
 				.AutoWidth()
 				[
-					SNew(SImage)
-					.Image(FAppStyle::Get().GetBrush("Icons.ChevronRight"))
+					SNew(SFieldPaths)
+					.TextStyle(TextStyle)
+					.FieldPaths(AllFields)
+					.HighlightField(HighlightRange)
+					.ShowOnlyLast(bShowOnlyLastPath)
 				];
 		}
 	}
-
-	if (InPropertyPath.HasPaths())
+	else
 	{
-		FieldBox->AddSlot()
-			.Padding(0, 0, 8, 0)
-			.HAlign(HAlign_Left)
-			.VAlign(VAlign_Center)
-			.AutoWidth()
-			[
-				SNew(SFieldPaths)
-				.TextStyle(TextStyle)
-				.FieldPaths(InPropertyPath.GetFields(WidgetBlueprint.Get()->SkeletonGeneratedClass))
-				.ShowOnlyLast(bShowOnlyLastPath)
-			];
+		if (InPropertyPath.HasPaths())
+		{
+			FieldBox->AddSlot()
+				.HAlign(HAlign_Left)
+				.VAlign(VAlign_Center)
+				.AutoWidth()
+				[
+					SNew(SFieldPaths)
+					.TextStyle(TextStyle)
+					.FieldPaths(InPropertyPath.GetFields(ClassContext))
+					.ShowOnlyLast(bShowOnlyLastPath)
+				];
+		}
 	}
 }
 
