MVVM: Change the location of the fieldnotify icon in the MVVM editor to help with consistency with the other tools.

diff --git a/Engine/Plugins/Runtime/ModelViewViewModel/Source/ModelViewViewModelEditor/Private/Widgets/SMVVMFieldEntry.cpp b/Engine/Plugins/Runtime/ModelViewViewModel/Source/ModelViewViewModelEditor/Private/Widgets/SMVVMFieldEntry.cpp
--- a/Engine/Plugins/Runtime/ModelViewViewModel/Source/ModelViewViewModelEditor/Private/Widgets/SMVVMFieldEntry.cpp
+++ b/Engine/Plugins/Runtime/ModelViewViewModel/Source/ModelViewViewModelEditor/Private/Widgets/SMVVMFieldEntry.cpp
@@ -156,15 +156,6 @@ void SFieldPaths::SetFieldPaths(TArrayView<UE::MVVM::FMVVMConstFieldVariant> InP
 
 		if (bIsFieldValid && bIsOwnerValid)
 		{
-			FieldBox->AddSlot()
-				.HAlign(HAlign_Left)
-				.VAlign(VAlign_Center)
-				.AutoWidth()
-				[
-					SNew(SFieldIcon)
-					.Field(InPropertyPath[Index])
-				];
-
 			bool bShowFieldNotify = HighlightField.IsSet() && HighlightField.GetValue().Contains(Index);
 			if (bShowFieldNotify)
 			{
@@ -179,10 +170,20 @@ void SFieldPaths::SetFieldPaths(TArrayView<UE::MVVM::FMVVMConstFieldVariant> InP
 					];
 			}
 
+			FieldBox->AddSlot()
+				.HAlign(HAlign_Right)
+				.VAlign(VAlign_Center)
+				.AutoWidth()
+				[
+					SNew(SFieldIcon)
+					.Field(InPropertyPath[Index])
+				];
+
 			FieldBox->AddSlot()
 				.HAlign(HAlign_Fill)
 				.VAlign(VAlign_Center)
 				.AutoWidth()
+				.Padding(4.0f, 0.f, 0.f, 0.f)
 				[
 					SNew(STextBlock)
 					.TextStyle(TextStyle)
