MVVM: Update the shared flag correctly after. Issue introduced with CL 30042195.

diff --git a/Engine/Plugins/Runtime/ModelViewViewModel/Source/ModelViewViewModelBlueprint/Private/MVVMViewBlueprintCompiler.cpp b/Engine/Plugins/Runtime/ModelViewViewModel/Source/ModelViewViewModelBlueprint/Private/MVVMViewBlueprintCompiler.cpp
--- a/Engine/Plugins/Runtime/ModelViewViewModel/Source/ModelViewViewModelBlueprint/Private/MVVMViewBlueprintCompiler.cpp
+++ b/Engine/Plugins/Runtime/ModelViewViewModel/Source/ModelViewViewModelBlueprint/Private/MVVMViewBlueprintCompiler.cpp
@@ -2396,16 +2396,12 @@ void FMVVMViewBlueprintCompiler::CompileBindings(const FCompiledBindingLibraryCo
 		FMVVMViewClass_BindingKey BindingKey = FMVVMViewClass_BindingKey(ViewExtension->Bindings.AddDefaulted());
 		FMVVMViewClass_Binding& NewBinding = ViewExtension->Bindings[BindingKey.GetIndex()];
 		NewBinding.Binding = CompiledBinding ? *CompiledBinding : FMVVMVCompiledBinding();
+		NewBinding.Flags = 0;
 		NewBinding.ExecutionMode = Binding.bOverrideExecutionMode ? Binding.OverrideExecutionMode : (EMVVMExecutionMode)CVarDefaultExecutionMode->GetInt();
 		NewBinding.SourceBitField = 0;
 		NewBinding.EditorId = Binding.BindingId;
 
-		NewBinding.Flags = 0;
-		NewBinding.Flags |= (!ValidBinding->bIsOneTimeBinding) ? (uint8)FMVVMViewClass_Binding::EFlags::OneWay : 0;
-		NewBinding.Flags |= (Binding.bOverrideExecutionMode) ? (uint8)FMVVMViewClass_Binding::EFlags::OverrideExecuteMode : 0;
-		NewBinding.Flags |= (ValidBinding->ReadPaths.Num() > 1) ? (uint8)FMVVMViewClass_Binding::EFlags::Shared : 0;
-		NewBinding.Flags |= (Binding.bEnabled) ? (uint8)FMVVMViewClass_Binding::EFlags::EnabledByDefault : 0;
-
+		// Find the source needed by the binding. Also generate every bindings on that source (that need to register to the FieldNotify).
 		TArray<FMVVMViewClass_SourceKey, TInlineAllocator<16>>  SharedBindings;
 		for (TSharedPtr<FGeneratedReadFieldPathContext> ReadPath : ValidBinding->ReadPaths)
 		{
@@ -2448,6 +2444,7 @@ void FMVVMViewBlueprintCompiler::CompileBindings(const FCompiledBindingLibraryCo
 			NewSourceBinding.BindingKey = BindingKey;
 
 			NewSourceBinding.Flags = 0;
+			// Set the ExecuteAtInitialization but remove it later if it's not the last of the group.
 			NewSourceBinding.Flags |= (bExecuteAtInitialization) ? (uint8)FMVVMViewClass_SourceBinding::EFlags::ExecuteAtInitialization : 0;
 
 			if (CompiledFieldId)
@@ -2462,7 +2459,11 @@ void FMVVMViewBlueprintCompiler::CompileBindings(const FCompiledBindingLibraryCo
 			}
 		}
 
-		NewBinding.Flags |= (SharedBindings.Num() > 0) ? (uint8)FMVVMViewClass_Binding::EFlags::Shared : 0;
+		// Set binding flags.
+		NewBinding.Flags |= (!ValidBinding->bIsOneTimeBinding) ? (uint8)FMVVMViewClass_Binding::EFlags::OneWay : 0;
+		NewBinding.Flags |= (SharedBindings.Num() > 1) ? (uint8)FMVVMViewClass_Binding::EFlags::Shared : 0;
+		NewBinding.Flags |= (Binding.bOverrideExecutionMode) ? (uint8)FMVVMViewClass_Binding::EFlags::OverrideExecuteMode : 0;
+		NewBinding.Flags |= (Binding.bEnabled) ? (uint8)FMVVMViewClass_Binding::EFlags::EnabledByDefault : 0;
 
 		// Only the last binding in the complex conversion.
 		if (SharedBindings.Num() > 0)
