MVVM: Add PinLinedValue. Use as a variant to easily set the propertypath, function or node of a pin.

diff --git a/Engine/Plugins/Runtime/ModelViewViewModel/Source/ModelViewViewModelBlueprint/Private/MVVMBlueprintFunctionReference.cpp b/Engine/Plugins/Runtime/ModelViewViewModel/Source/ModelViewViewModelBlueprint/Private/MVVMBlueprintFunctionReference.cpp
--- a/Engine/Plugins/Runtime/ModelViewViewModel/Source/ModelViewViewModelBlueprint/Private/MVVMBlueprintFunctionReference.cpp
+++ b/Engine/Plugins/Runtime/ModelViewViewModel/Source/ModelViewViewModelBlueprint/Private/MVVMBlueprintFunctionReference.cpp
@@ -48,7 +48,7 @@ FMVVMBlueprintFunctionReference::FMVVMBlueprintFunctionReference(FMemberReferenc
 	, Type(EMVVMBlueprintFunctionReferenceType::Function)
 {}
 
-FMVVMBlueprintFunctionReference::FMVVMBlueprintFunctionReference(const TSubclassOf<UK2Node> InNode)
+FMVVMBlueprintFunctionReference::FMVVMBlueprintFunctionReference(TSubclassOf<UK2Node> InNode)
 	: Node(InNode)
 	, Type(InNode.Get() != nullptr ? EMVVMBlueprintFunctionReferenceType::Node : EMVVMBlueprintFunctionReferenceType::None)
 {
@@ -79,7 +79,7 @@ const UFunction* FMVVMBlueprintFunctionReference::GetFunction(const UClass* Self
 	return FunctionReference.ResolveMember<UFunction>(const_cast<UClass*>(SelfContext), false);
 }
 
-const TSubclassOf<UK2Node> FMVVMBlueprintFunctionReference::GetNode() const
+TSubclassOf<UK2Node> FMVVMBlueprintFunctionReference::GetNode() const
 {
 	return Type == EMVVMBlueprintFunctionReferenceType::Node ? Node : TSubclassOf<UK2Node>();
 }
diff --git a/Engine/Plugins/Runtime/ModelViewViewModel/Source/ModelViewViewModelBlueprint/Private/MVVMBlueprintViewConversionFunction.cpp b/Engine/Plugins/Runtime/ModelViewViewModel/Source/ModelViewViewModelBlueprint/Private/MVVMBlueprintViewConversionFunction.cpp
--- a/Engine/Plugins/Runtime/ModelViewViewModel/Source/ModelViewViewModelBlueprint/Private/MVVMBlueprintViewConversionFunction.cpp
+++ b/Engine/Plugins/Runtime/ModelViewViewModel/Source/ModelViewViewModelBlueprint/Private/MVVMBlueprintViewConversionFunction.cpp
@@ -28,6 +28,9 @@ namespace UE::MVVM::Private
 
 bool UMVVMBlueprintViewConversionFunction::IsValidConversionFunction(const UBlueprint* WidgetBlueprint, const UFunction* Function)
 {
+	check(WidgetBlueprint);
+	check(Function);
+
 	// functions in the widget blueprint can do anything they want, other functions have to be static functions in a BlueprintFunctionLibrary
 	const UClass* FunctionClass = Function->GetOuterUClass();
 
diff --git a/Engine/Plugins/Runtime/ModelViewViewModel/Source/ModelViewViewModelBlueprint/Public/MVVMBlueprintFunctionReference.h b/Engine/Plugins/Runtime/ModelViewViewModel/Source/ModelViewViewModelBlueprint/Public/MVVMBlueprintFunctionReference.h
--- a/Engine/Plugins/Runtime/ModelViewViewModel/Source/ModelViewViewModelBlueprint/Public/MVVMBlueprintFunctionReference.h
+++ b/Engine/Plugins/Runtime/ModelViewViewModel/Source/ModelViewViewModelBlueprint/Public/MVVMBlueprintFunctionReference.h
@@ -33,14 +33,14 @@ public:
 	FMVVMBlueprintFunctionReference() = default;
 	MODELVIEWVIEWMODELBLUEPRINT_API explicit FMVVMBlueprintFunctionReference(const UBlueprint* InContext, const UFunction* InFunction);
 	MODELVIEWVIEWMODELBLUEPRINT_API explicit FMVVMBlueprintFunctionReference(FMemberReference InFunctionReference);
-	MODELVIEWVIEWMODELBLUEPRINT_API explicit FMVVMBlueprintFunctionReference(const TSubclassOf<UK2Node> InNode);
+	MODELVIEWVIEWMODELBLUEPRINT_API explicit FMVVMBlueprintFunctionReference(TSubclassOf<UK2Node> InNode);
 
 	/** If the node is of type Function, resolves the function. Uses the blueprint, if the function is defined on the BP itself. */
 	MODELVIEWVIEWMODELBLUEPRINT_API const UFunction* GetFunction(const UBlueprint* SelfContext) const;
 	/** If the node is of type Function, resolves the function. Uses the blueprint, if the function is defined on the BP itself. */
 	MODELVIEWVIEWMODELBLUEPRINT_API const UFunction* GetFunction(const UClass* SelfContext) const;
 	/** If the node is of type Node, return the node class. */
-	MODELVIEWVIEWMODELBLUEPRINT_API const TSubclassOf<UK2Node> GetNode() const;
+	MODELVIEWVIEWMODELBLUEPRINT_API TSubclassOf<UK2Node> GetNode() const;
 
 	EMVVMBlueprintFunctionReferenceType GetType() const
 	{
diff --git a/Engine/Plugins/Runtime/ModelViewViewModel/Source/ModelViewViewModelEditor/Private/MVVMEditorSubsystem.cpp b/Engine/Plugins/Runtime/ModelViewViewModel/Source/ModelViewViewModelEditor/Private/MVVMEditorSubsystem.cpp
--- a/Engine/Plugins/Runtime/ModelViewViewModel/Source/ModelViewViewModelEditor/Private/MVVMEditorSubsystem.cpp
+++ b/Engine/Plugins/Runtime/ModelViewViewModel/Source/ModelViewViewModelEditor/Private/MVVMEditorSubsystem.cpp
@@ -356,11 +356,34 @@ UEdGraphPin* UMVVMEditorSubsystem::GetConversionFunctionArgumentPin(const UWidge
 	return nullptr;
 }
 
-void UMVVMEditorSubsystem::SetSourceToDestinationConversionFunction(UWidgetBlueprint* WidgetBlueprint, FMVVMBlueprintViewBinding& Binding, const UFunction* ConversionFunction)
+void UMVVMEditorSubsystem::SetSourceToDestinationConversionFunction(UWidgetBlueprint* WidgetBlueprint, FMVVMBlueprintViewBinding& Binding, const UFunction* NewConversionFunction)
+{
+	SetSourceToDestinationConversionFunction(WidgetBlueprint, Binding, FMVVMBlueprintFunctionReference(WidgetBlueprint, NewConversionFunction));
+}
+
+void UMVVMEditorSubsystem::SetSourceToDestinationConversionFunction(UWidgetBlueprint* WidgetBlueprint, FMVVMBlueprintViewBinding& Binding, const FMVVMBlueprintFunctionReference& NewConversionFunction)
 {
 	if (UMVVMBlueprintView* View = GetView(WidgetBlueprint))
 	{
-		if (ConversionFunction == nullptr || IsValidConversionFunction(WidgetBlueprint, ConversionFunction, Binding.SourcePath, Binding.DestinationPath))
+		const UFunction* NewFunction = nullptr;
+		if (NewConversionFunction.GetType() == EMVVMBlueprintFunctionReferenceType::Function)
+		{
+			NewFunction = NewConversionFunction.GetFunction(WidgetBlueprint);
+			if (!IsValidConversionFunction(WidgetBlueprint, NewFunction, Binding.SourcePath, Binding.DestinationPath))
+			{
+				NewFunction = nullptr;
+			}
+		}
+		const UFunction* CurrentFunction = nullptr;
+		if (Binding.Conversion.SourceToDestinationConversion != nullptr && Binding.Conversion.SourceToDestinationConversion->GetConversionFunction().GetType() == EMVVMBlueprintFunctionReferenceType::Function)
+		{
+			CurrentFunction = Binding.Conversion.SourceToDestinationConversion->GetConversionFunction().GetFunction(WidgetBlueprint);
+		}
+
+		bool bBothNull = Binding.Conversion.SourceToDestinationConversion == nullptr && NewConversionFunction.GetType() == EMVVMBlueprintFunctionReferenceType::None;
+		bool bBothSameValue = Binding.Conversion.SourceToDestinationConversion != nullptr && Binding.Conversion.SourceToDestinationConversion->GetConversionFunction() == NewConversionFunction;
+		bool bBothSameFunctionValue = CurrentFunction == NewFunction;
+		if (!bBothSameValue || !bBothNull || !NewFunction)
 		{
 			FScopedTransaction Transaction(LOCTEXT("SetConversionFunction", "Set Conversion Function"));
 
@@ -373,13 +396,20 @@ void UMVVMEditorSubsystem::SetSourceToDestinationConversionFunction(UWidgetBluep
 				Binding.Conversion.SourceToDestinationConversion->RemoveWrapperGraph(WidgetBlueprint);
 				Binding.Conversion.SourceToDestinationConversion = nullptr;
 			}
+			Binding.SourcePath = FMVVMBlueprintPropertyPath();
 
-			if (ConversionFunction != nullptr)
+			if (NewConversionFunction.GetType() == EMVVMBlueprintFunctionReferenceType::Function)
+			{
+				if (NewFunction != nullptr)
+				{
+					Binding.Conversion.SourceToDestinationConversion = NewObject<UMVVMBlueprintViewConversionFunction>(WidgetBlueprint);
+					FName GraphName = UE::MVVM::ConversionFunctionHelper::CreateWrapperName(Binding, true);
+					Binding.Conversion.SourceToDestinationConversion->InitializeFromFunction(WidgetBlueprint, GraphName, NewFunction);
+				}
+			}
+			else
 			{
-				Binding.Conversion.SourceToDestinationConversion = NewObject<UMVVMBlueprintViewConversionFunction>(WidgetBlueprint);
-				FName GraphName = UE::MVVM::ConversionFunctionHelper::CreateWrapperName(Binding, true);
-				Binding.Conversion.SourceToDestinationConversion->InitializeFromFunction(WidgetBlueprint, GraphName, ConversionFunction);
-				Binding.SourcePath = FMVVMBlueprintPropertyPath();
+				check(false); // not supported yet
 			}
 
 			UE::MVVM::Private::OnBindingPostEditChange(View, GET_MEMBER_NAME_CHECKED(FMVVMBlueprintViewBinding, Conversion));
@@ -388,11 +418,34 @@ void UMVVMEditorSubsystem::SetSourceToDestinationConversionFunction(UWidgetBluep
 	}
 }
 
-void UMVVMEditorSubsystem::SetDestinationToSourceConversionFunction(UWidgetBlueprint* WidgetBlueprint, FMVVMBlueprintViewBinding& Binding, const UFunction* ConversionFunction)
+void UMVVMEditorSubsystem::SetDestinationToSourceConversionFunction(UWidgetBlueprint* WidgetBlueprint, FMVVMBlueprintViewBinding& Binding, const UFunction* NewConversionFunction)
+{
+	SetDestinationToSourceConversionFunction(WidgetBlueprint, Binding, FMVVMBlueprintFunctionReference(WidgetBlueprint, NewConversionFunction));
+}
+
+void UMVVMEditorSubsystem::SetDestinationToSourceConversionFunction(UWidgetBlueprint* WidgetBlueprint, FMVVMBlueprintViewBinding& Binding, const FMVVMBlueprintFunctionReference& NewConversionFunction)
 {
 	if (UMVVMBlueprintView* View = GetView(WidgetBlueprint))
 	{
-		if (ConversionFunction == nullptr || IsValidConversionFunction(WidgetBlueprint, ConversionFunction, Binding.DestinationPath, Binding.SourcePath))
+		const UFunction* NewFunction = nullptr;
+		if (NewConversionFunction.GetType() == EMVVMBlueprintFunctionReferenceType::Function)
+		{
+			NewFunction = NewConversionFunction.GetFunction(WidgetBlueprint);
+			if (!IsValidConversionFunction(WidgetBlueprint, NewFunction, Binding.SourcePath, Binding.DestinationPath))
+			{
+				NewFunction = nullptr;
+			}
+		}
+		const UFunction* CurrentFunction = nullptr;
+		if (Binding.Conversion.DestinationToSourceConversion != nullptr && Binding.Conversion.DestinationToSourceConversion->GetConversionFunction().GetType() == EMVVMBlueprintFunctionReferenceType::Function)
+		{
+			CurrentFunction = Binding.Conversion.DestinationToSourceConversion->GetConversionFunction().GetFunction(WidgetBlueprint);
+		}
+
+		bool bBothNull = Binding.Conversion.DestinationToSourceConversion == nullptr && NewConversionFunction.GetType() == EMVVMBlueprintFunctionReferenceType::None;
+		bool bBothSameValue = Binding.Conversion.DestinationToSourceConversion != nullptr && Binding.Conversion.DestinationToSourceConversion->GetConversionFunction() == NewConversionFunction;
+		bool bBothSameFunctionValue = CurrentFunction == NewFunction;
+		if (!bBothSameValue || !bBothNull || !NewFunction)
 		{
 			FScopedTransaction Transaction(LOCTEXT("SetConversionFunction", "Set Conversion Function"));
 
@@ -405,13 +458,20 @@ void UMVVMEditorSubsystem::SetDestinationToSourceConversionFunction(UWidgetBluep
 				Binding.Conversion.DestinationToSourceConversion->RemoveWrapperGraph(WidgetBlueprint);
 				Binding.Conversion.DestinationToSourceConversion = nullptr;
 			}
+			Binding.DestinationPath = FMVVMBlueprintPropertyPath();
 
-			if (ConversionFunction != nullptr)
+			if (NewConversionFunction.GetType() == EMVVMBlueprintFunctionReferenceType::Function)
 			{
-				Binding.Conversion.DestinationToSourceConversion = NewObject<UMVVMBlueprintViewConversionFunction>(WidgetBlueprint);
-				FName GraphName = UE::MVVM::ConversionFunctionHelper::CreateWrapperName(Binding, false);
-				Binding.Conversion.DestinationToSourceConversion->InitializeFromFunction(WidgetBlueprint, GraphName, ConversionFunction);
-				Binding.DestinationPath = FMVVMBlueprintPropertyPath();
+				if (NewFunction != nullptr)
+				{
+					Binding.Conversion.DestinationToSourceConversion = NewObject<UMVVMBlueprintViewConversionFunction>(WidgetBlueprint);
+					FName GraphName = UE::MVVM::ConversionFunctionHelper::CreateWrapperName(Binding, true);
+					Binding.Conversion.DestinationToSourceConversion->InitializeFromFunction(WidgetBlueprint, GraphName, NewFunction);
+				}
+			}
+			else
+			{
+				check(false); // not supported yet
 			}
 
 			UE::MVVM::Private::OnBindingPostEditChange(View, GET_MEMBER_NAME_CHECKED(FMVVMBlueprintViewBinding, Conversion));
@@ -424,34 +484,39 @@ void UMVVMEditorSubsystem::SetDestinationPathForBinding(UWidgetBlueprint* Widget
 {
 	if (UMVVMBlueprintView* View = GetView(WidgetBlueprint))
 	{
-		FScopedTransaction Transaction(LOCTEXT("SetBindingProperty", "Set Binding Property"));
-
-		UE::MVVM::Private::OnBindingPreEditChange(View, GET_MEMBER_NAME_CHECKED(FMVVMBlueprintViewBinding, DestinationPath));
-
+		bool bHasConversion = Binding.Conversion.DestinationToSourceConversion != nullptr;
 		bool bEventSupported = UMVVMBlueprintViewEvent::Supports(WidgetBlueprint, PropertyPath);
-		UMVVMBlueprintViewEvent* Event = nullptr;
-		if (bEventSupported)
-		{
-			Event = AddEvent(WidgetBlueprint);
-		}
 
-		if (Event)
+		if (bEventSupported || bHasConversion || Binding.DestinationPath != PropertyPath)
 		{
-			Event->SetEventPath(PropertyPath);
-			View->RemoveBinding(&Binding);
-		}
-		else
-		{
-			if (Binding.Conversion.DestinationToSourceConversion && PropertyPath.IsValid())
+			FScopedTransaction Transaction(LOCTEXT("SetBindingProperty", "Set Binding Property"));
+
+			UE::MVVM::Private::OnBindingPreEditChange(View, GET_MEMBER_NAME_CHECKED(FMVVMBlueprintViewBinding, DestinationPath));
+
+			UMVVMBlueprintViewEvent* Event = nullptr;
+			if (bEventSupported)
 			{
-				Binding.Conversion.DestinationToSourceConversion->RemoveWrapperGraph(WidgetBlueprint);
-				Binding.Conversion.DestinationToSourceConversion = nullptr;
+				Event = AddEvent(WidgetBlueprint);
+			}
+
+			if (Event)
+			{
+				Event->SetEventPath(PropertyPath);
+				View->RemoveBinding(&Binding);
+			}
+			else
+			{
+				if (Binding.Conversion.DestinationToSourceConversion)
+				{
+					Binding.Conversion.DestinationToSourceConversion->RemoveWrapperGraph(WidgetBlueprint);
+					Binding.Conversion.DestinationToSourceConversion = nullptr;
+				}
+				Binding.DestinationPath = PropertyPath;
 			}
-			Binding.DestinationPath = PropertyPath;
-		}
 
-		UE::MVVM::Private::OnBindingPostEditChange(View, GET_MEMBER_NAME_CHECKED(FMVVMBlueprintViewBinding, DestinationPath));
-		FBlueprintEditorUtils::MarkBlueprintAsModified(WidgetBlueprint);
+			UE::MVVM::Private::OnBindingPostEditChange(View, GET_MEMBER_NAME_CHECKED(FMVVMBlueprintViewBinding, DestinationPath));
+			FBlueprintEditorUtils::MarkBlueprintAsModified(WidgetBlueprint);
+		}
 	}
 }
 
@@ -459,19 +524,23 @@ void UMVVMEditorSubsystem::SetSourcePathForBinding(UWidgetBlueprint* WidgetBluep
 {
 	if (UMVVMBlueprintView* View = GetView(WidgetBlueprint))
 	{
-		FScopedTransaction Transaction(LOCTEXT("SetBindingProperty", "Set Binding Property"));
+		bool bHasConversion = Binding.Conversion.SourceToDestinationConversion != nullptr;
+		if (bHasConversion || Binding.SourcePath != PropertyPath)
+		{
+			FScopedTransaction Transaction(LOCTEXT("SetBindingProperty", "Set Binding Property"));
 
-		UE::MVVM::Private::OnBindingPreEditChange(View, GET_MEMBER_NAME_CHECKED(FMVVMBlueprintViewBinding, SourcePath));
+			UE::MVVM::Private::OnBindingPreEditChange(View, GET_MEMBER_NAME_CHECKED(FMVVMBlueprintViewBinding, SourcePath));
 
-		if (Binding.Conversion.SourceToDestinationConversion && PropertyPath.IsValid())
-		{
-			Binding.Conversion.SourceToDestinationConversion->RemoveWrapperGraph(WidgetBlueprint);
-			Binding.Conversion.SourceToDestinationConversion = nullptr;
-		}
-		Binding.SourcePath = PropertyPath;
+			if (Binding.Conversion.SourceToDestinationConversion)
+			{
+				Binding.Conversion.SourceToDestinationConversion->RemoveWrapperGraph(WidgetBlueprint);
+				Binding.Conversion.SourceToDestinationConversion = nullptr;
+			}
+			Binding.SourcePath = PropertyPath;
 
-		UE::MVVM::Private::OnBindingPostEditChange(View, GET_MEMBER_NAME_CHECKED(FMVVMBlueprintViewBinding, SourcePath));
-		FBlueprintEditorUtils::MarkBlueprintAsModified(WidgetBlueprint);
+			UE::MVVM::Private::OnBindingPostEditChange(View, GET_MEMBER_NAME_CHECKED(FMVVMBlueprintViewBinding, SourcePath));
+			FBlueprintEditorUtils::MarkBlueprintAsModified(WidgetBlueprint);
+		}
 	}
 }
 
@@ -664,6 +733,11 @@ void UMVVMEditorSubsystem::SetCompileForEvent(UMVVMBlueprintViewEvent* Event, bo
 
 bool UMVVMEditorSubsystem::IsValidConversionFunction(const UWidgetBlueprint* WidgetBlueprint, const UFunction* Function, const FMVVMBlueprintPropertyPath& Source, const FMVVMBlueprintPropertyPath& Destination) const
 {
+	if (WidgetBlueprint == nullptr || Function == nullptr)
+	{
+		return false;
+	}
+
 	if (!UMVVMBlueprintViewConversionFunction::IsValidConversionFunction(WidgetBlueprint, Function))
 	{
 		return false;
diff --git a/Engine/Plugins/Runtime/ModelViewViewModel/Source/ModelViewViewModelEditor/Private/Types/MVVMPinLinkedValue.cpp b/Engine/Plugins/Runtime/ModelViewViewModel/Source/ModelViewViewModelEditor/Private/Types/MVVMPinLinkedValue.cpp
new file mode 100644
--- /dev/null
+++ b/Engine/Plugins/Runtime/ModelViewViewModel/Source/ModelViewViewModelEditor/Private/Types/MVVMPinLinkedValue.cpp
@@ -0,0 +1,64 @@
+// Copyright Epic Games, Inc. All Rights Reserved.
+
+#include "Types/MVVMLinkedPinValue.h"
+
+#include UE_INLINE_GENERATED_CPP_BY_NAME(MVVMLinkedPinValue)
+
+FMVVMLinkedPinValue::FMVVMLinkedPinValue(FMVVMBlueprintPropertyPath InPath)
+	: PropertyPath(MoveTemp(InPath))
+{}
+
+FMVVMLinkedPinValue::FMVVMLinkedPinValue(const UBlueprint* InBlueprint, FMVVMBlueprintFunctionReference InConversion)
+{
+	if (InConversion.GetType() == EMVVMBlueprintFunctionReferenceType::Function)
+	{
+		ConversionFunction = InConversion.GetFunction(InBlueprint);
+	}
+	else if (InConversion.GetType() == EMVVMBlueprintFunctionReferenceType::Node)
+	{
+		ConversionNode = InConversion.GetNode();
+	}
+}
+
+FMVVMLinkedPinValue::FMVVMLinkedPinValue(const UFunction* Function)
+	: ConversionFunction(Function)
+{}
+
+FMVVMLinkedPinValue::FMVVMLinkedPinValue(TSubclassOf<UK2Node> Node)
+	: ConversionNode(Node)
+{}
+
+bool FMVVMLinkedPinValue::IsValid() const
+{
+	return ConversionFunction != nullptr || ConversionNode.Get() != nullptr || PropertyPath.IsValid();
+}
+
+bool FMVVMLinkedPinValue::IsPropertyPath() const
+{
+	return PropertyPath.IsValid();
+}
+
+const FMVVMBlueprintPropertyPath& FMVVMLinkedPinValue::GetPropertyPath() const
+{
+	return PropertyPath;
+}
+
+bool FMVVMLinkedPinValue::IsConversionFunction() const
+{
+	return ConversionFunction != nullptr;
+}
+
+const UFunction* FMVVMLinkedPinValue::GetConversionFunction() const
+{
+	return ConversionFunction;
+}
+
+bool FMVVMLinkedPinValue::IsConversionNode() const
+{
+	return ConversionNode.Get() != nullptr;
+}
+
+TSubclassOf<UK2Node> FMVVMLinkedPinValue::GetConversionNode() const
+{
+	return ConversionNode;
+}
diff --git a/Engine/Plugins/Runtime/ModelViewViewModel/Source/ModelViewViewModelEditor/Private/Widgets/BindingEntry/SMVVMBindingRow.cpp b/Engine/Plugins/Runtime/ModelViewViewModel/Source/ModelViewViewModelEditor/Private/Widgets/BindingEntry/SMVVMBindingRow.cpp
--- a/Engine/Plugins/Runtime/ModelViewViewModel/Source/ModelViewViewModelEditor/Private/Widgets/BindingEntry/SMVVMBindingRow.cpp
+++ b/Engine/Plugins/Runtime/ModelViewViewModel/Source/ModelViewViewModelEditor/Private/Widgets/BindingEntry/SMVVMBindingRow.cpp
@@ -119,9 +119,8 @@ TSharedRef<SWidget> SBindingRow::BuildRowWidget()
 				.MinDesiredWidth(150)
 				[
 					SNew(SFieldSelector, GetBlueprint())
-					.OnGetPropertyPath(this, &SBindingRow::GetSelectedPropertyPath, false)
-					.OnGetConversionFunction(this, &SBindingRow::GetSelectedConversionFunction, false)
-					.OnFieldSelectionChanged(this, &SBindingRow::HandleFieldSelectionChanged, false)
+					.OnGetLinkedValue(this, &SBindingRow::GetFieldSelectedValue, false)
+					.OnSelectionChanged(this, &SBindingRow::HandleFieldSelectionChanged, false)
 					.OnGetSelectionContext(this, &SBindingRow::GetSelectedSelectionContext, false)
 					.OnDrop(this, &SBindingRow::HandleFieldSelectorDrop, false)
 					.OnDragEnter(this, &SBindingRow::HandleFieldSelectorDragEnter, false)
@@ -165,9 +164,8 @@ TSharedRef<SWidget> SBindingRow::BuildRowWidget()
 				.MinDesiredWidth(150.0f)
 				[
 					SNew(SFieldSelector, GetBlueprint())
-					.OnGetPropertyPath(this, &SBindingRow::GetSelectedPropertyPath, true)
-					.OnGetConversionFunction(this, &SBindingRow::GetSelectedConversionFunction, true)
-					.OnFieldSelectionChanged(this, &SBindingRow::HandleFieldSelectionChanged, true)
+					.OnGetLinkedValue(this, &SBindingRow::GetFieldSelectedValue, true)
+					.OnSelectionChanged(this, &SBindingRow::HandleFieldSelectionChanged, true)
 					.OnGetSelectionContext(this, &SBindingRow::GetSelectedSelectionContext, true)
 					.OnDrop(this, &SBindingRow::HandleFieldSelectorDrop, true)
 					.OnDragEnter(this, &SBindingRow::HandleFieldSelectorDragEnter, true)
@@ -372,31 +370,31 @@ TArray<FBindingSource> SBindingRow::GetAvailableViewModels() const
 	return EditorSubsystem->GetAllViewModels(GetBlueprint());
 }
 
-FMVVMBlueprintPropertyPath SBindingRow::GetSelectedPropertyPath(bool bSource) const
+FMVVMLinkedPinValue SBindingRow::GetFieldSelectedValue(bool bSourceToDest) const
 {
 	if (FMVVMBlueprintViewBinding* ViewBinding = GetThisViewBinding())
 	{
-		return bSource ? ViewBinding->SourcePath : ViewBinding->DestinationPath;
-	}
-	return FMVVMBlueprintPropertyPath();
-}
 
-const UFunction* SBindingRow::GetSelectedConversionFunction(bool bSourceToDest) const
-{
-	if (FMVVMBlueprintViewBinding* ViewBinding = GetThisViewBinding())
-	{
 		if (UMVVMBlueprintViewConversionFunction* ConversionFunction = ViewBinding->Conversion.GetConversionFunction(bSourceToDest))
 		{
 			if (ConversionFunction->GetConversionFunction().GetType() == EMVVMBlueprintFunctionReferenceType::Function)
 			{
-				return ConversionFunction->GetConversionFunction().GetFunction(GetBlueprint());
+				return FMVVMLinkedPinValue(ConversionFunction->GetConversionFunction().GetFunction(GetBlueprint()));
+			}
+			else if (ConversionFunction->GetConversionFunction().GetType() == EMVVMBlueprintFunctionReferenceType::Function)
+			{
+				return FMVVMLinkedPinValue(ConversionFunction->GetConversionFunction().GetNode());
 			}
 		}
+		else
+		{
+			return FMVVMLinkedPinValue(bSourceToDest ? ViewBinding->SourcePath : ViewBinding->DestinationPath);
+		}
 	}
-	return nullptr;
+	return FMVVMLinkedPinValue();
 }
 
-void SBindingRow::HandleFieldSelectionChanged(FMVVMBlueprintPropertyPath SelectedField, const UFunction* Function, bool bSource)
+void SBindingRow::HandleFieldSelectionChanged(FMVVMLinkedPinValue Value, bool bSource)
 {
 	UWidgetBlueprint* WidgetBlueprint = GetBlueprint();
 	FMVVMBlueprintViewBinding* ViewBinding = GetThisViewBinding();
@@ -405,18 +403,40 @@ void SBindingRow::HandleFieldSelectionChanged(FMVVMBlueprintPropertyPath Selecte
 		UMVVMEditorSubsystem* Subsystem = GetEditorSubsystem();
 		if (bSource)
 		{
-			Subsystem->SetSourceToDestinationConversionFunction(WidgetBlueprint, *ViewBinding, Function);
-			if (ViewBinding->SourcePath != SelectedField)
+			if (Value.IsPropertyPath())
 			{
-				Subsystem->SetSourcePathForBinding(WidgetBlueprint, *ViewBinding, SelectedField);
+				Subsystem->SetSourcePathForBinding(WidgetBlueprint, *ViewBinding, Value.GetPropertyPath());
+			}
+			else if (Value.IsConversionFunction())
+			{
+				Subsystem->SetSourceToDestinationConversionFunction(WidgetBlueprint, *ViewBinding, FMVVMBlueprintFunctionReference(WidgetBlueprint, Value.GetConversionFunction()));
+			}
+			else if (Value.IsConversionNode())
+			{
+				Subsystem->SetSourceToDestinationConversionFunction(WidgetBlueprint, *ViewBinding, FMVVMBlueprintFunctionReference(Value.GetConversionNode()));
+			}
+			else
+			{
+				Subsystem->SetSourcePathForBinding(WidgetBlueprint, *ViewBinding, FMVVMBlueprintPropertyPath());
 			}
 		}
 		else
 		{
-			Subsystem->SetDestinationToSourceConversionFunction(WidgetBlueprint, *ViewBinding, Function);
-			if (ViewBinding->DestinationPath != SelectedField)
+			if (Value.IsPropertyPath())
+			{
+				Subsystem->SetDestinationPathForBinding(WidgetBlueprint, *ViewBinding, Value.GetPropertyPath());
+			}
+			else if (Value.IsConversionFunction())
+			{
+				Subsystem->SetDestinationToSourceConversionFunction(WidgetBlueprint, *ViewBinding, FMVVMBlueprintFunctionReference(WidgetBlueprint, Value.GetConversionFunction()));
+			}
+			else if (Value.IsConversionNode())
+			{
+				Subsystem->SetDestinationToSourceConversionFunction(WidgetBlueprint, *ViewBinding, FMVVMBlueprintFunctionReference(Value.GetConversionNode()));
+			}
+			else
 			{
-				Subsystem->SetDestinationPathForBinding(WidgetBlueprint, *ViewBinding, SelectedField);
+				Subsystem->SetDestinationPathForBinding(WidgetBlueprint, *ViewBinding, FMVVMBlueprintPropertyPath());
 			}
 		}
 	}
@@ -494,12 +514,10 @@ FReply SBindingRow::HandleFieldSelectorDrop(const FGeometry& MyGeometry, const F
 	UMVVMEditorSubsystem* Subsystem = GetEditorSubsystem();
 	if (bSource)
 	{
-		Subsystem->SetSourceToDestinationConversionFunction(WidgetBlueprint, *ViewBinding, nullptr);
 		Subsystem->SetSourcePathForBinding(WidgetBlueprint, *ViewBinding, PropertyPath.GetValue());
 	}
 	else
 	{
-		Subsystem->SetDestinationToSourceConversionFunction(WidgetBlueprint, *ViewBinding, nullptr);
 		Subsystem->SetDestinationPathForBinding(WidgetBlueprint, *ViewBinding, PropertyPath.GetValue());
 	}
 	return FReply::Handled();
diff --git a/Engine/Plugins/Runtime/ModelViewViewModel/Source/ModelViewViewModelEditor/Private/Widgets/BindingEntry/SMVVMBindingRow.h b/Engine/Plugins/Runtime/ModelViewViewModel/Source/ModelViewViewModelEditor/Private/Widgets/BindingEntry/SMVVMBindingRow.h
--- a/Engine/Plugins/Runtime/ModelViewViewModel/Source/ModelViewViewModelEditor/Private/Widgets/BindingEntry/SMVVMBindingRow.h
+++ b/Engine/Plugins/Runtime/ModelViewViewModel/Source/ModelViewViewModelEditor/Private/Widgets/BindingEntry/SMVVMBindingRow.h
@@ -51,11 +51,9 @@ private:
 
 	TArray<FBindingSource> GetAvailableViewModels() const;
 
-	FMVVMBlueprintPropertyPath GetSelectedPropertyPath(bool bSource) const;
+	FMVVMLinkedPinValue GetFieldSelectedValue(bool bSource) const;
 
-	const UFunction* GetSelectedConversionFunction(bool bSourceToDest) const;
-
-	void HandleFieldSelectionChanged(FMVVMBlueprintPropertyPath SelectedField, const UFunction* Function, bool bSource);
+	void HandleFieldSelectionChanged(FMVVMLinkedPinValue Value, bool bSource);
 
 	FFieldSelectionContext GetSelectedSelectionContext(bool bSource) const;
 
diff --git a/Engine/Plugins/Runtime/ModelViewViewModel/Source/ModelViewViewModelEditor/Private/Widgets/BindingEntry/SMVVMEventRow.cpp b/Engine/Plugins/Runtime/ModelViewViewModel/Source/ModelViewViewModelEditor/Private/Widgets/BindingEntry/SMVVMEventRow.cpp
--- a/Engine/Plugins/Runtime/ModelViewViewModel/Source/ModelViewViewModelEditor/Private/Widgets/BindingEntry/SMVVMEventRow.cpp
+++ b/Engine/Plugins/Runtime/ModelViewViewModel/Source/ModelViewViewModelEditor/Private/Widgets/BindingEntry/SMVVMEventRow.cpp
@@ -83,9 +83,8 @@ TSharedRef<SWidget> SEventRow::BuildRowWidget()
 				.MinDesiredWidth(150.0f)
 				[
 					SNew(SFieldSelector, GetBlueprint())
-					.OnGetPropertyPath(this, &SEventRow::GetSelectedPropertyPath, true)
-					.OnGetConversionFunction(this, &SEventRow::GetSelectedConversionFunction, true)
-					.OnFieldSelectionChanged(this, &SEventRow::HandleFieldSelectionChanged, true)
+					.OnGetLinkedValue(this, &SEventRow::GetFieldSelectedValue, true)
+					.OnSelectionChanged(this, &SEventRow::HandleFieldSelectionChanged, true)
 					.OnGetSelectionContext(this, &SEventRow::GetSelectedSelectionContext, true)
 					.OnDrop(this, &SEventRow::HandleFieldSelectorDrop, true)
 					.OnDragEnter(this, &SEventRow::HandleFieldSelectorDragEnter, true)
@@ -120,9 +119,8 @@ TSharedRef<SWidget> SEventRow::BuildRowWidget()
 				.MinDesiredWidth(150.0f)
 				[
 					SNew(SFieldSelector, GetBlueprint())
-					.OnGetPropertyPath(this, &SEventRow::GetSelectedPropertyPath, false)
-					.OnGetConversionFunction(this, &SEventRow::GetSelectedConversionFunction, false)
-					.OnFieldSelectionChanged(this, &SEventRow::HandleFieldSelectionChanged, false)
+					.OnGetLinkedValue(this, &SEventRow::GetFieldSelectedValue, false)
+					.OnSelectionChanged(this, &SEventRow::HandleFieldSelectionChanged, false)
 					.OnGetSelectionContext(this, &SEventRow::GetSelectedSelectionContext, false)
 					.OnDrop(this, &SEventRow::HandleFieldSelectorDrop, false)
 					.OnDragEnter(this, &SEventRow::HandleFieldSelectorDragEnter, false)
@@ -310,21 +308,16 @@ void SEventRow::OnIsEventCompileChanged(ECheckBoxState NewState)
 	}
 }
 
-FMVVMBlueprintPropertyPath SEventRow::GetSelectedPropertyPath(bool bEvent) const
+FMVVMLinkedPinValue SEventRow::GetFieldSelectedValue(bool bEvent) const
 {
 	if (const UMVVMBlueprintViewEvent* Event = GetEvent())
 	{
-		return bEvent ? Event->GetEventPath() : Event->GetDestinationPath();
+		return bEvent ? FMVVMLinkedPinValue(Event->GetEventPath()) : FMVVMLinkedPinValue(Event->GetDestinationPath());
 	}
-	return FMVVMBlueprintPropertyPath();
+	return FMVVMLinkedPinValue();
 }
 
-const UFunction* SEventRow::GetSelectedConversionFunction(bool bSourceToDest) const
-{
-	return nullptr;
-}
-
-void SEventRow::HandleFieldSelectionChanged(FMVVMBlueprintPropertyPath SelectedField, const UFunction* Function, bool bEvent)
+void SEventRow::HandleFieldSelectionChanged(FMVVMLinkedPinValue Value, bool bEvent)
 {
 	UWidgetBlueprint* WidgetBlueprint = GetBlueprint();
 	UMVVMBlueprintViewEvent* Event = GetEvent();
@@ -333,17 +326,11 @@ void SEventRow::HandleFieldSelectionChanged(FMVVMBlueprintPropertyPath SelectedF
 		UMVVMEditorSubsystem* Subsystem = GetEditorSubsystem();
 		if (bEvent)
 		{
-			if (Event->GetEventPath() != SelectedField)
-			{
-				Subsystem->SetEventPath(Event, SelectedField);
-			}
+			Subsystem->SetEventPath(Event, Value.IsPropertyPath() ? Value.GetPropertyPath() : FMVVMBlueprintPropertyPath());
 		}
 		else
 		{
-			if (Event->GetDestinationPath() != SelectedField)
-			{
-				Subsystem->SetEventDestinationPath(Event, SelectedField);
-			}
+			Subsystem->SetEventDestinationPath(Event, Value.IsPropertyPath() ? Value.GetPropertyPath() : FMVVMBlueprintPropertyPath());
 		}
 	}
 }
diff --git a/Engine/Plugins/Runtime/ModelViewViewModel/Source/ModelViewViewModelEditor/Private/Widgets/BindingEntry/SMVVMEventRow.h b/Engine/Plugins/Runtime/ModelViewViewModel/Source/ModelViewViewModelEditor/Private/Widgets/BindingEntry/SMVVMEventRow.h
--- a/Engine/Plugins/Runtime/ModelViewViewModel/Source/ModelViewViewModelEditor/Private/Widgets/BindingEntry/SMVVMEventRow.h
+++ b/Engine/Plugins/Runtime/ModelViewViewModel/Source/ModelViewViewModelEditor/Private/Widgets/BindingEntry/SMVVMEventRow.h
@@ -51,11 +51,9 @@ private:
 
 	void OnIsEventCompileChanged(ECheckBoxState NewState);
 
-	FMVVMBlueprintPropertyPath GetSelectedPropertyPath(bool bSource) const;
+	FMVVMLinkedPinValue GetFieldSelectedValue(bool bSource) const;
 
-	const UFunction* GetSelectedConversionFunction(bool bSourceToDest) const;
-
-	void HandleFieldSelectionChanged(FMVVMBlueprintPropertyPath SelectedField, const UFunction* Function, bool bSource);
+	void HandleFieldSelectionChanged(FMVVMLinkedPinValue Value, bool bSource);
 
 	FFieldSelectionContext GetSelectedSelectionContext(bool bSource) const;
 
diff --git a/Engine/Plugins/Runtime/ModelViewViewModel/Source/ModelViewViewModelEditor/Private/Widgets/SMVVMCachedViewBindingConversionFunction.cpp b/Engine/Plugins/Runtime/ModelViewViewModel/Source/ModelViewViewModelEditor/Private/Widgets/SMVVMCachedViewBindingConversionFunction.cpp
--- a/Engine/Plugins/Runtime/ModelViewViewModel/Source/ModelViewViewModelEditor/Private/Widgets/SMVVMCachedViewBindingConversionFunction.cpp
+++ b/Engine/Plugins/Runtime/ModelViewViewModel/Source/ModelViewViewModelEditor/Private/Widgets/SMVVMCachedViewBindingConversionFunction.cpp
@@ -11,7 +11,7 @@
 #include "Widgets/Images/SImage.h"
 #include "Widgets/Text/STextBlock.h"
 
-#define LOCTEXT_NAMESPACE "SCachedViewBindingConverisonFunction"
+#define LOCTEXT_NAMESPACE "SCachedViewBindingConversionFunction"
 
 namespace UE::MVVM
 {
@@ -26,13 +26,23 @@ void SCachedViewBindingConversionFunction::Construct(const FArguments& InArgs, c
 		SNew(SHorizontalBox)
 		.ToolTipText_Lambda([this]() 
 		{
-			const UFunction* FoundFunction = OnGetConversionFunction.Execute();
-			return FoundFunction != nullptr ? FoundFunction->GetToolTipText() : FText::GetEmpty();
+			FConversionFunctionVariant ConversionFunction = OnGetConversionFunction.Execute();
+			if (ConversionFunction.IsType<const UFunction*>())
+			{
+				const UFunction* FoundFunction = ConversionFunction.Get<const UFunction*>();
+				return FoundFunction != nullptr ? FoundFunction->GetToolTipText() : FText::GetEmpty();
+			}
+			else if (ConversionFunction.IsType<TSubclassOf<UK2Node>>())
+			{
+				TSubclassOf<UK2Node> FoundFunction = ConversionFunction.Get<TSubclassOf<UK2Node>>();
+				return FoundFunction.Get() != nullptr ? FoundFunction->GetToolTipText() : FText::GetEmpty();
+			}
+			return FText::GetEmpty();
 		})
 		+ SHorizontalBox::Slot()
 		.HAlign(HAlign_Left)
 		.VAlign(VAlign_Center)
-		.Padding(0, 0, 4.0f, 0)
+		.Padding(0.0f, 0.0f, 4.0f, 0.0f)
 		.AutoWidth()
 		[
 			SNew(SImage)
@@ -46,8 +56,18 @@ void SCachedViewBindingConversionFunction::Construct(const FArguments& InArgs, c
 			SNew(STextBlock)
 			.Text_Lambda([this]() 
 			{ 
-				const UFunction* FoundFunction = OnGetConversionFunction.Execute();
-				return FoundFunction != nullptr ? FoundFunction->GetDisplayNameText() : FText::GetEmpty();
+				FConversionFunctionVariant ConversionFunction = OnGetConversionFunction.Execute();
+				if (ConversionFunction.IsType<const UFunction*>())
+				{
+					const UFunction* FoundFunction = ConversionFunction.Get<const UFunction*>();
+					return FoundFunction != nullptr ? FoundFunction->GetDisplayNameText() : FText::GetEmpty();
+				}
+				else if (ConversionFunction.IsType<TSubclassOf<UK2Node>>())
+				{
+					TSubclassOf<UK2Node> FoundFunction = ConversionFunction.Get<TSubclassOf<UK2Node>>();
+					return FoundFunction.Get() != nullptr ? FoundFunction->GetDisplayNameText() : FText::GetEmpty();
+				}
+				return FText::GetEmpty();
 			})
 		]
 	];
diff --git a/Engine/Plugins/Runtime/ModelViewViewModel/Source/ModelViewViewModelEditor/Private/Widgets/SMVVMCachedViewBindingConversionFunction.h b/Engine/Plugins/Runtime/ModelViewViewModel/Source/ModelViewViewModelEditor/Private/Widgets/SMVVMCachedViewBindingConversionFunction.h
--- a/Engine/Plugins/Runtime/ModelViewViewModel/Source/ModelViewViewModelEditor/Private/Widgets/SMVVMCachedViewBindingConversionFunction.h
+++ b/Engine/Plugins/Runtime/ModelViewViewModel/Source/ModelViewViewModelEditor/Private/Widgets/SMVVMCachedViewBindingConversionFunction.h
@@ -2,14 +2,20 @@
 
 #pragma once
 
-#include "Widgets/SMVVMCachedViewBindingPropertyPath.h"
+#include "K2Node.h"
+#include "Misc/TVariant.h"
+#include "Templates/SubclassOf.h"
 
-class SHorizontalBox;
+#include "Styling/CoreStyle.h"
+#include "Styling/SlateTypes.h"
+#include "Styling/SlateWidgetStyleAsset.h"
+#include "Widgets/SCompoundWidget.h"
+
+class UFunction;
 class UWidgetBlueprint;
 
 namespace UE::MVVM
 {
-DECLARE_DELEGATE_RetVal(const UFunction*, FOnGetConversionFunction);
 
 class SCachedViewBindingConversionFunction : public SCompoundWidget
 {
@@ -17,6 +23,8 @@ private:
 	using Super = SCompoundWidget;
 
 public:
+	using FConversionFunctionVariant = TVariant<const UFunction*, TSubclassOf<UK2Node>, FEmptyVariantState>;
+	DECLARE_DELEGATE_RetVal(FConversionFunctionVariant, FOnGetConversionFunction);
 
 	SLATE_BEGIN_ARGS(SCachedViewBindingConversionFunction)
 		: _TextStyle(&FCoreStyle::Get().GetWidgetStyle<FTextBlockStyle>("NormalText"))
diff --git a/Engine/Plugins/Runtime/ModelViewViewModel/Source/ModelViewViewModelEditor/Private/Widgets/SMVVMEventParameter.cpp b/Engine/Plugins/Runtime/ModelViewViewModel/Source/ModelViewViewModelEditor/Private/Widgets/SMVVMEventParameter.cpp
--- a/Engine/Plugins/Runtime/ModelViewViewModel/Source/ModelViewViewModelEditor/Private/Widgets/SMVVMEventParameter.cpp
+++ b/Engine/Plugins/Runtime/ModelViewViewModel/Source/ModelViewViewModelEditor/Private/Widgets/SMVVMEventParameter.cpp
@@ -67,8 +67,7 @@ void SEventParameter::Construct(const FArguments& InArgs, UWidgetBlueprint* InWi
 		}
 	}
 
-	FMVVMBlueprintPropertyPath Path = OnGetSelectedField();
-	bDefaultValueVisible = !Path.IsValid();
+	bDefaultValueVisible = !OnGetSelectedField().IsValid();
 
 	TSharedPtr<SHorizontalBox> HBox;
 
@@ -91,8 +90,8 @@ void SEventParameter::Construct(const FArguments& InArgs, UWidgetBlueprint* InWi
 				+ SWidgetSwitcher::Slot()
 				[
 					SNew(SFieldSelector, WidgetBlueprint.Get())
-					.OnGetPropertyPath(this, &SEventParameter::OnGetSelectedField)
-					.OnFieldSelectionChanged(this, &SEventParameter::HandleFieldSelectionChanged)
+					.OnGetLinkedValue(this, &SEventParameter::OnGetSelectedField)
+					.OnSelectionChanged(this, &SEventParameter::HandleFieldSelectionChanged)
 					.OnGetSelectionContext(this, &SEventParameter::GetSelectedSelectionContext)
 				]
 			]
@@ -143,7 +142,7 @@ void SEventParameter::OnBindArgumentChecked(ECheckBoxState Checked)
 	if (bDefaultValueVisible)
 	{
 		PreviousSelectedField = OnGetSelectedField();
-		SetSelectedField(FMVVMBlueprintPropertyPath());
+		SetSelectedField(FMVVMLinkedPinValue());
 	}
 	else
 	{
@@ -151,27 +150,27 @@ void SEventParameter::OnBindArgumentChecked(ECheckBoxState Checked)
 	}
 }
 
-FMVVMBlueprintPropertyPath SEventParameter::OnGetSelectedField() const
+FMVVMLinkedPinValue SEventParameter::OnGetSelectedField() const
 {
 	if (const UMVVMBlueprintViewEvent* EventPtr = ViewEvent.Get())
 	{
-		return EventPtr->GetPinPath(ParameterId);
+		return FMVVMLinkedPinValue(EventPtr->GetPinPath(ParameterId));
 	}
-	return FMVVMBlueprintPropertyPath();
+	return FMVVMLinkedPinValue();
 }
 
-void SEventParameter::SetSelectedField(const FMVVMBlueprintPropertyPath& Path)
+void SEventParameter::SetSelectedField(const FMVVMLinkedPinValue& Path)
 {
 	if (UMVVMBlueprintViewEvent* EventPtr = ViewEvent.Get())
 	{
 		const UMVVMEditorSubsystem* Subsystem = GEditor->GetEditorSubsystem<UMVVMEditorSubsystem>();
-		Subsystem->SetEventArgumentPath(EventPtr, ParameterId, Path);
+		Subsystem->SetEventArgumentPath(EventPtr, ParameterId, Path.IsPropertyPath() ? Path.GetPropertyPath(): FMVVMBlueprintPropertyPath());
 	}
 }
 
-void SEventParameter::HandleFieldSelectionChanged(FMVVMBlueprintPropertyPath SelectedField, const UFunction* Function)
+void SEventParameter::HandleFieldSelectionChanged(FMVVMLinkedPinValue Value)
 {
-	SetSelectedField(SelectedField);
+	SetSelectedField(Value);
 }
 
 FFieldSelectionContext SEventParameter::GetSelectedSelectionContext() const
diff --git a/Engine/Plugins/Runtime/ModelViewViewModel/Source/ModelViewViewModelEditor/Private/Widgets/SMVVMEventParameter.h b/Engine/Plugins/Runtime/ModelViewViewModel/Source/ModelViewViewModelEditor/Private/Widgets/SMVVMEventParameter.h
--- a/Engine/Plugins/Runtime/ModelViewViewModel/Source/ModelViewViewModelEditor/Private/Widgets/SMVVMEventParameter.h
+++ b/Engine/Plugins/Runtime/ModelViewViewModel/Source/ModelViewViewModelEditor/Private/Widgets/SMVVMEventParameter.h
@@ -4,6 +4,7 @@
 
 #include "MVVMBlueprintPin.h"
 #include "MVVMPropertyPath.h"
+#include "Types/MVVMLinkedPinValue.h"
 #include "Widgets/SCompoundWidget.h"
 #include "Widgets/SMVVMFieldSelectorMenu.h"
 
@@ -35,10 +36,10 @@ private:
 	ECheckBoxState OnGetIsBindArgumentChecked() const;
 	void OnBindArgumentChecked(ECheckBoxState Checked);
 
-	FMVVMBlueprintPropertyPath OnGetSelectedField() const;
-	void SetSelectedField(const FMVVMBlueprintPropertyPath& Path);
+	FMVVMLinkedPinValue OnGetSelectedField() const;
+	void SetSelectedField(const FMVVMLinkedPinValue& Path);
 
-	void HandleFieldSelectionChanged(FMVVMBlueprintPropertyPath SelectedField, const UFunction* Function);
+	void HandleFieldSelectionChanged(FMVVMLinkedPinValue Value);
 	FFieldSelectionContext GetSelectedSelectionContext() const;
 
 	int32 GetCurrentWidgetIndex() const;
@@ -50,7 +51,7 @@ private:
 	/** This reference is just to keep the default value widget alive. */
 	TSharedPtr<SGraphPin> GraphPin;
 
-	FMVVMBlueprintPropertyPath PreviousSelectedField;
+	FMVVMLinkedPinValue PreviousSelectedField;
 
 	bool bAllowDefault = true;
 	bool bDefaultValueVisible = true;
diff --git a/Engine/Plugins/Runtime/ModelViewViewModel/Source/ModelViewViewModelEditor/Private/Widgets/SMVVMFieldSelector.cpp b/Engine/Plugins/Runtime/ModelViewViewModel/Source/ModelViewViewModelEditor/Private/Widgets/SMVVMFieldSelector.cpp
--- a/Engine/Plugins/Runtime/ModelViewViewModel/Source/ModelViewViewModelEditor/Private/Widgets/SMVVMFieldSelector.cpp
+++ b/Engine/Plugins/Runtime/ModelViewViewModel/Source/ModelViewViewModelEditor/Private/Widgets/SMVVMFieldSelector.cpp
@@ -33,17 +33,16 @@ void SFieldSelector::Construct(const FArguments& InArgs, const UWidgetBlueprint*
 	check(InWidgetBlueprint);
 	
 	TextStyle = InArgs._TextStyle;
-	OnGetPropertyPath = InArgs._OnGetPropertyPath;
-	OnGetConversionFunction = InArgs._OnGetConversionFunction;
+	OnGetLinkedValue = InArgs._OnGetLinkedValue;
+	OnSelectionChanged = InArgs._OnSelectionChanged;
 	OnGetSelectionContext = InArgs._OnGetSelectionContext;
-	OnFieldSelectionChanged = InArgs._OnFieldSelectionChanged;
 	OnDragEnterEvent = InArgs._OnDragEnter;
 	OnDropEvent = InArgs._OnDrop;
 
 	ChildSlot
 	[
 		SNew(SBox)
-		.MinDesiredWidth(200)
+		.MinDesiredWidth(200.0f)
 		[
 			SAssignNew(ComboButton, SComboButton)
 			.ComboButtonStyle(FMVVMEditorStyle::Get(), "FieldSelector.ComboButton")
@@ -63,7 +62,7 @@ void SFieldSelector::Construct(const FArguments& InArgs, const UWidgetBlueprint*
 					SNew(SCachedViewBindingPropertyPath, WidgetBlueprint.Get())
 					.TextStyle(TextStyle)
 					.ShowContext(InArgs._ShowContext)
-					.OnGetPropertyPath(InArgs._OnGetPropertyPath)
+					.OnGetPropertyPath(this, &SFieldSelector::HandleGetPropertyPath)
 				]
 
 				//1-Conversion Function
@@ -74,7 +73,7 @@ void SFieldSelector::Construct(const FArguments& InArgs, const UWidgetBlueprint*
 				[
 					SNew(SCachedViewBindingConversionFunction, WidgetBlueprint.Get())
 					.TextStyle(TextStyle)
-					.OnGetConversionFunction(InArgs._OnGetConversionFunction)
+					.OnGetConversionFunction(this, &SFieldSelector::HandleGetConversionFunction)
 				]
 
 				//2-Nothing selected.
@@ -95,16 +94,14 @@ void SFieldSelector::Construct(const FArguments& InArgs, const UWidgetBlueprint*
 
 int32 SFieldSelector::GetCurrentDisplayIndex() const
 {
-	if (OnGetConversionFunction.IsBound())
+	if (OnGetLinkedValue.IsBound())
 	{
-		if (OnGetConversionFunction.Execute() != nullptr)
+		FMVVMLinkedPinValue LinkedValue = OnGetLinkedValue.Execute();
+		if (LinkedValue.IsConversionFunction() || LinkedValue.IsConversionNode())
 		{
 			return 1;
 		}
-	}
-	if (OnGetPropertyPath.IsBound())
-	{
-		if (OnGetPropertyPath.Execute().IsValid())
+		else if (LinkedValue.IsPropertyPath())
 		{
 			return 0;
 		}
@@ -120,15 +117,10 @@ TSharedRef<SWidget> SFieldSelector::HandleGetMenuContent()
 		return SNullWidget::NullWidget;
 	}
 
-	const UFunction* CurrentSelectedFunction = nullptr;
-	if (OnGetConversionFunction.IsBound())
+	TOptional<FMVVMLinkedPinValue> CurrentSelected;
+	if (OnGetLinkedValue.IsBound())
 	{
-		CurrentSelectedFunction = OnGetConversionFunction.Execute();
-	}
-	TOptional<FMVVMBlueprintPropertyPath> CurrentSelectedPropertyPath;
-	if (OnGetPropertyPath.IsBound())
-	{
-		CurrentSelectedPropertyPath = OnGetPropertyPath.Execute();
+		CurrentSelected = OnGetLinkedValue.Execute();
 	}
 
 	FFieldSelectionContext SelectionContext;
@@ -138,11 +130,10 @@ TSharedRef<SWidget> SFieldSelector::HandleGetMenuContent()
 	}
 
 	TSharedRef<SFieldSelectorMenu> Menu = SNew(SFieldSelectorMenu, WidgetBlueprintPtr)
-		.OnFieldSelectionChanged(this, &SFieldSelector::HandleFieldSelectionChanged)
+		.CurrentSelected(CurrentSelected)
+		.OnSelectionChanged(this, &SFieldSelector::HandleFieldSelectionChanged)
 		.OnMenuCloseRequested(this, &SFieldSelector::HandleMenuClosed)
 		.SelectionContext(SelectionContext)
-		.CurrentPropertyPathSelected(CurrentSelectedPropertyPath)
-		.CurrentFunctionSelected(CurrentSelectedFunction)
 		;
 
 	ComboButton->SetMenuContentWidgetToFocus(Menu->GetWidgetToFocus());
@@ -150,17 +141,52 @@ TSharedRef<SWidget> SFieldSelector::HandleGetMenuContent()
 	return Menu;
 }
 
-void SFieldSelector::HandleFieldSelectionChanged(FMVVMBlueprintPropertyPath PropertyPath, const UFunction* Function)
+void SFieldSelector::HandleFieldSelectionChanged(FMVVMLinkedPinValue LinkedValue)
 {
 	if (ComboButton.IsValid())
 	{
 		ComboButton->SetIsOpen(false);
 	}
 
-	if (OnFieldSelectionChanged.IsBound())
+	if (OnSelectionChanged.IsBound())
+	{
+		OnSelectionChanged.Execute(LinkedValue);
+	}
+}
+
+FMVVMBlueprintPropertyPath SFieldSelector::HandleGetPropertyPath() const
+{
+	if (OnGetLinkedValue.IsBound())
+	{
+		FMVVMLinkedPinValue LinkedValue = OnGetLinkedValue.Execute();
+		if (LinkedValue.IsPropertyPath())
+		{
+			return LinkedValue.GetPropertyPath();
+		}
+	}
+	return FMVVMBlueprintPropertyPath();
+}
+
+TVariant<const UFunction*, TSubclassOf<UK2Node>, FEmptyVariantState> SFieldSelector::HandleGetConversionFunction() const
+{
+	using FReturnType = TVariant<const UFunction*, TSubclassOf<UK2Node>, FEmptyVariantState>;
+	if (OnGetLinkedValue.IsBound())
 	{
-		OnFieldSelectionChanged.Execute(PropertyPath, Function);
+		FMVVMLinkedPinValue LinkedValue = OnGetLinkedValue.Execute();
+		if (ensure(LinkedValue.IsConversionFunction() || LinkedValue.IsConversionNode()))
+		{
+			if (LinkedValue.IsConversionFunction())
+			{
+				return FReturnType(TInPlaceType<const UFunction*>(), LinkedValue.GetConversionFunction());
+			}
+			else
+			{
+				check(LinkedValue.IsConversionNode());
+				return FReturnType(TInPlaceType<TSubclassOf<UK2Node>>(), LinkedValue.GetConversionNode());
+			}
+		}
 	}
+	return FReturnType(TInPlaceType<FEmptyVariantState>());
 }
 
 void SFieldSelector::OnDragEnter(const FGeometry& MyGeometry, const FDragDropEvent& DragDropEvent)
diff --git a/Engine/Plugins/Runtime/ModelViewViewModel/Source/ModelViewViewModelEditor/Private/Widgets/SMVVMFieldSelector.h b/Engine/Plugins/Runtime/ModelViewViewModel/Source/ModelViewViewModelEditor/Private/Widgets/SMVVMFieldSelector.h
--- a/Engine/Plugins/Runtime/ModelViewViewModel/Source/ModelViewViewModelEditor/Private/Widgets/SMVVMFieldSelector.h
+++ b/Engine/Plugins/Runtime/ModelViewViewModel/Source/ModelViewViewModelEditor/Private/Widgets/SMVVMFieldSelector.h
@@ -2,7 +2,6 @@
 
 #pragma once
 
-#include "MVVMPropertyPath.h"
 #include "Styling/CoreStyle.h"
 #include "Styling/SlateTypes.h"
 #include "Styling/SlateWidgetStyleAsset.h"
@@ -21,6 +20,7 @@ namespace UE::MVVM
 class SFieldSelector : public SCompoundWidget
 {
 public:
+	DECLARE_DELEGATE_RetVal(FMVVMLinkedPinValue, FOnGetLinkedPinValue);
 	DECLARE_DELEGATE_RetVal(FFieldSelectionContext, FOnGetSelectionContext);
 	DECLARE_DELEGATE_RetVal_TwoParams(FReply, FOnDrop, const FGeometry&, const FDragDropEvent&);
 	DECLARE_DELEGATE_TwoParams(FOnDragEnter, const FGeometry&, const FDragDropEvent&);
@@ -31,9 +31,8 @@ public:
 		}
 		SLATE_STYLE_ARGUMENT(FTextBlockStyle, TextStyle)
 		SLATE_ARGUMENT_DEFAULT(bool, ShowContext) = true;
-		SLATE_EVENT(FOnGetPropertyPath, OnGetPropertyPath)
-		SLATE_EVENT(FOnGetConversionFunction, OnGetConversionFunction)
-		SLATE_EVENT(FOnFieldSelectionChanged, OnFieldSelectionChanged)
+		SLATE_EVENT(FOnGetLinkedPinValue, OnGetLinkedValue)
+		SLATE_EVENT(FOnLinkedValueSelectionChanged, OnSelectionChanged)
 		SLATE_EVENT(FOnGetSelectionContext, OnGetSelectionContext)
 		SLATE_EVENT(FOnDrop, OnDrop)
 		SLATE_EVENT(FOnDragEnter, OnDragEnter)
@@ -51,7 +50,9 @@ private:
 	TSharedRef<SWidget> CreateSourcePanel();
 	TSharedRef<SWidget> HandleGetMenuContent();
 
-	void HandleFieldSelectionChanged(FMVVMBlueprintPropertyPath, const UFunction*);
+	void HandleFieldSelectionChanged(FMVVMLinkedPinValue NewValue);
+	FMVVMBlueprintPropertyPath HandleGetPropertyPath() const;
+	TVariant<const UFunction*, TSubclassOf<UK2Node>, FEmptyVariantState> HandleGetConversionFunction() const;
 	void HandleMenuClosed();
 
 private:
@@ -60,10 +61,9 @@ private:
 
 	TWeakObjectPtr<const UWidgetBlueprint> WidgetBlueprint;
 	const FTextBlockStyle* TextStyle = nullptr;
-	FOnGetPropertyPath OnGetPropertyPath;
-	FOnGetConversionFunction OnGetConversionFunction;
+	FOnGetLinkedPinValue OnGetLinkedValue;
+	FOnLinkedValueSelectionChanged OnSelectionChanged;
 	FOnGetSelectionContext OnGetSelectionContext;
-	FOnFieldSelectionChanged OnFieldSelectionChanged;
 	FOnDrop OnDropEvent;
 	FOnDragEnter OnDragEnterEvent;
 }; 
diff --git a/Engine/Plugins/Runtime/ModelViewViewModel/Source/ModelViewViewModelEditor/Private/Widgets/SMVVMFieldSelectorMenu.cpp b/Engine/Plugins/Runtime/ModelViewViewModel/Source/ModelViewViewModelEditor/Private/Widgets/SMVVMFieldSelectorMenu.cpp
--- a/Engine/Plugins/Runtime/ModelViewViewModel/Source/ModelViewViewModelEditor/Private/Widgets/SMVVMFieldSelectorMenu.cpp
+++ b/Engine/Plugins/Runtime/ModelViewViewModel/Source/ModelViewViewModelEditor/Private/Widgets/SMVVMFieldSelectorMenu.cpp
@@ -61,12 +61,12 @@ void SFieldSelectorMenu::Construct(const FArguments& InArgs, const UWidgetBluepr
 	WidgetBlueprint = InWidgetBlueprint;
 	check(InWidgetBlueprint);
 
-	OnFieldSelectionChanged = InArgs._OnFieldSelectionChanged;
+	OnSelectionChanged = InArgs._OnSelectionChanged;
 	OnMenuCloseRequested = InArgs._OnMenuCloseRequested;
 	SelectionContext = InArgs._SelectionContext;
 
 	bIsMenuInitialized = false;
-	bIsClearEnabled = (InArgs._CurrentPropertyPathSelected.IsSet() && InArgs._CurrentPropertyPathSelected.GetValue().IsValid()) || InArgs._CurrentFunctionSelected != nullptr;
+	bIsClearEnabled = InArgs._CurrentSelected.IsSet() && InArgs._CurrentSelected.GetValue().IsValid();
 
 	// If we're showing conversion functions, we don't want to set the AssignableTo property of SSourceBindingList, because then it will only show exact matches, 
 	// and since we're also showing conversion functions we know that's not what the user wants.
@@ -78,7 +78,7 @@ void SFieldSelectorMenu::Construct(const FArguments& InArgs, const UWidgetBluepr
 		AssignableToProperty = SelectionContext.AssignableTo;
 	}
 
-	if (SelectionContext.bAllowConversionFunctions)
+	if (bShowConversionFunctions)
 	{
 		GenerateConversionFunctionItems();
 	}
@@ -235,19 +235,37 @@ bool SFieldSelectorMenu::IsSelectEnabled() const
 
 FReply SFieldSelectorMenu::HandleSelectClicked()
 {
-	FMVVMBlueprintPropertyPath PropertyPath;
-	if (BindingList.IsValid())
+	if (OnSelectionChanged.IsBound())
 	{
-		PropertyPath = BindingList->GetSelectedProperty();
-	}
+		bool bSet = false;
+		FMVVMLinkedPinValue Value;
+		if (ConversionFunctionList.IsValid())
+		{
+			TArray<const UFunction*> Selection;
+			if (ConversionFunctionList.IsValid())
+			{
+				Selection = ConversionFunctionList->GetSelectedItems();
+			}
+			if (Selection.Num() > 0)
+			{
+				Value = FMVVMLinkedPinValue(Selection[0]);
+				bSet = true;
+			}
+		}
 
-	TArray<const UFunction*> Selection;
-	if (ConversionFunctionList.IsValid())
-	{
-		Selection = ConversionFunctionList->GetSelectedItems();
-	}
+		if (BindingList.IsValid())
+		{
+			FMVVMBlueprintPropertyPath SelectedProperty = BindingList->GetSelectedProperty();
+			if (SelectedProperty.IsValid())
+			{
+				ensure(bSet == false);
+				Value = FMVVMLinkedPinValue(MoveTemp(SelectedProperty));
+				bSet = true;
+			}
+		}
 
-	OnFieldSelectionChanged.ExecuteIfBound(PropertyPath, Selection.Num() > 0 ? Selection[0] : nullptr);
+		OnSelectionChanged.Execute(MoveTemp(Value));
+	}
 
 	return FReply::Handled();
 }
@@ -267,7 +285,7 @@ FReply SFieldSelectorMenu::HandleClearClicked()
 		Source.SetSourceTo(NewProperty);
 	}
 
-	OnFieldSelectionChanged.ExecuteIfBound(NewProperty, nullptr);
+	OnSelectionChanged.ExecuteIfBound(FMVVMLinkedPinValue(MoveTemp(NewProperty)));
 
 	return FReply::Handled();
 }
@@ -281,12 +299,12 @@ FReply SFieldSelectorMenu::HandleCancelClicked()
 
 void SFieldSelectorMenu::SetPropertyPathSelection(const FMVVMBlueprintPropertyPath& SelectedPath)
 {
-	OnFieldSelectionChanged.ExecuteIfBound(SelectedPath, nullptr);
+	OnSelectionChanged.ExecuteIfBound(FMVVMLinkedPinValue(SelectedPath));
 }
 
 void SFieldSelectorMenu::SetConversionFunctionSelection(const UFunction* SelectedFunction)
 {
-	OnFieldSelectionChanged.ExecuteIfBound(FMVVMBlueprintPropertyPath(), SelectedFunction);
+	OnSelectionChanged.ExecuteIfBound(FMVVMLinkedPinValue(SelectedFunction));
 }
 
 TSharedPtr<SFieldSelectorMenu::FConversionFunctionItem> SFieldSelectorMenu::FindConversionFunctionCategory(const TArray<TSharedPtr<FConversionFunctionItem>>& Items, TArrayView<FString> CategoryNameParts) const
@@ -770,7 +788,13 @@ TSharedRef<SWidget> SFieldSelectorMenu::CreateBindingContextPanel(const FArgumen
 
 	const bool bUseFixedSource = SelectionContext.FixedBindingSource.IsSet() && SelectionContext.FixedBindingSource.GetValue().IsValid();
 	const EMVVMBlueprintFieldPathSource FixedFieldPathSource = bUseFixedSource ? SelectionContext.FixedBindingSource.GetValue().GetSource() : EMVVMBlueprintFieldPathSource::None;
-	const EMVVMBlueprintFieldPathSource CurrentPropertyFieldPathSource = InArgs._CurrentPropertyPathSelected.IsSet() ? InArgs._CurrentPropertyPathSelected.GetValue().GetSource(WidgetBlueprint.Get()) : EMVVMBlueprintFieldPathSource::None;
+	EMVVMBlueprintFieldPathSource CurrentPropertyFieldPathSource = EMVVMBlueprintFieldPathSource::None;
+
+	FMVVMLinkedPinValue CurrentFieldSelectedLinkedValue = InArgs._CurrentSelected.IsSet() ? InArgs._CurrentSelected.GetValue() : FMVVMLinkedPinValue();
+	if (InArgs._CurrentSelected.IsSet() && CurrentFieldSelectedLinkedValue.IsPropertyPath())
+	{
+		CurrentPropertyFieldPathSource = CurrentFieldSelectedLinkedValue.GetPropertyPath().GetSource(WidgetBlueprint.Get());
+	}
 
 	if (SelectionContext.bAllowViewModels)
 	{
@@ -796,7 +820,9 @@ TSharedRef<SWidget> SFieldSelectorMenu::CreateBindingContextPanel(const FArgumen
 		{
 			for (const FBindingSource& Source : FilteredViewModelSources)
 			{
-				if (Source.GetViewModelId() == InArgs._CurrentPropertyPathSelected.GetValue().GetViewModelId())
+				check(InArgs._CurrentSelected.IsSet());
+				check(CurrentFieldSelectedLinkedValue.IsPropertyPath());
+				if (Source.GetViewModelId() == CurrentFieldSelectedLinkedValue.GetPropertyPath().GetViewModelId())
 				{
 					SelectedSource = Source;
 				}
@@ -836,7 +862,9 @@ TSharedRef<SWidget> SFieldSelectorMenu::CreateBindingContextPanel(const FArgumen
 
 		if (CurrentPropertyFieldPathSource == EMVVMBlueprintFieldPathSource::Widget)
 		{
-			WidgetList->SetSelectedWidget(InArgs._CurrentPropertyPathSelected.GetValue().GetWidgetName());
+			check(InArgs._CurrentSelected.IsSet());
+			check(CurrentFieldSelectedLinkedValue.IsPropertyPath());
+			WidgetList->SetSelectedWidget(CurrentFieldSelectedLinkedValue.GetPropertyPath().GetWidgetName());
 		}
 		else if (CurrentPropertyFieldPathSource == EMVVMBlueprintFieldPathSource::SelfContext)
 		{
@@ -861,15 +889,23 @@ TSharedRef<SWidget> SFieldSelectorMenu::CreateBindingContextPanel(const FArgumen
 			.OnSelectionChanged(this, &SFieldSelectorMenu::HandleConversionFunctionCategorySelected)
 			.OnGetChildren(this, &SFieldSelectorMenu::HandleGetConversionFunctionCategoryChildren);
 
-		if (InArgs._CurrentFunctionSelected != nullptr)
+		if (CurrentFieldSelectedLinkedValue.IsConversionFunction())
 		{
-			TSharedPtr<FConversionFunctionItem> FunctionItem = ExpandFunctionCategoryTreeToItem(InArgs._CurrentFunctionSelected);
-			if (FunctionItem)
+			const UFunction* ConversionFunction = CurrentFieldSelectedLinkedValue.GetConversionFunction();
+			if (ConversionFunction)
 			{
-				TGuardValue<bool> TmpGuard(bIsMenuInitialized, true);
-				ConversionFunctionCategoryTree->SetItemSelection(FunctionItem, true);
+				TSharedPtr<FConversionFunctionItem> FunctionItem = ExpandFunctionCategoryTreeToItem(ConversionFunction);
+				if (FunctionItem)
+				{
+					TGuardValue<bool> TmpGuard(bIsMenuInitialized, true);
+					ConversionFunctionCategoryTree->SetItemSelection(FunctionItem, true);
+				}
 			}
 		}
+		else if (CurrentFieldSelectedLinkedValue.IsConversionNode())
+		{
+			check(false); // not supported yet
+		}
 
 		StackedSourcePicker->AddSlot()
 			.AutoHeight()
@@ -889,20 +925,22 @@ TSharedRef<SWidget> SFieldSelectorMenu::CreateBindingListPanel(const FArguments&
 		.FieldVisibilityFlags(Private::GetFieldVisibilityFlags(SelectionContext.BindingMode, SelectionContext.bReadable, SelectionContext.bWritable))
 		.AssignableTo(AssignableToProperty);
 
+	const bool bValidSelectedProperty = InArgs._CurrentSelected.IsSet() && InArgs._CurrentSelected.GetValue().IsPropertyPath();
+
 	if (SelectionContext.FixedBindingSource.IsSet())
 	{
 		// Single fixed source, don't show the separate source panel.
 		BindingList->AddSource(SelectionContext.FixedBindingSource.GetValue());
 	}
-	else if (InArgs._CurrentPropertyPathSelected.IsSet())
+	else if (bValidSelectedProperty)
 	{
-		FBindingSource Source = FBindingSource::CreateFromPropertyPath(WidgetBlueprint.Get(), InArgs._CurrentPropertyPathSelected.GetValue());
+		FBindingSource Source = FBindingSource::CreateFromPropertyPath(WidgetBlueprint.Get(), InArgs._CurrentSelected.GetValue().GetPropertyPath());
 		BindingList->AddSource(Source);
 	}
 
-	if (InArgs._CurrentPropertyPathSelected.IsSet())
+	if (bValidSelectedProperty)
 	{
-		BindingList->SetSelectedProperty(InArgs._CurrentPropertyPathSelected.GetValue());
+		BindingList->SetSelectedProperty(InArgs._CurrentSelected.GetValue().GetPropertyPath());
 	}
 
 
@@ -914,6 +952,8 @@ TSharedRef<SWidget> SFieldSelectorMenu::CreateBindingListPanel(const FArguments&
 
 	if (SelectionContext.bAllowConversionFunctions)
 	{
+		const bool bValidSelectedFunction = InArgs._CurrentSelected.IsSet() && (InArgs._CurrentSelected.GetValue().IsConversionFunction() || InArgs._CurrentSelected.GetValue().IsConversionNode());
+
 		BindingListVBox->AddSlot()
 			.AutoHeight()
 			[
@@ -924,9 +964,10 @@ TSharedRef<SWidget> SFieldSelectorMenu::CreateBindingListPanel(const FArguments&
 				.OnGenerateRow(this, &SFieldSelectorMenu::HandleGenerateConversionFunctionRow)
 			];
 
-		if (InArgs._CurrentFunctionSelected != nullptr)
+		if (bValidSelectedFunction)
 		{
-			ConversionFunctionList->SetItemSelection(InArgs._CurrentFunctionSelected, true);
+			check(InArgs._CurrentSelected.GetValue().IsConversionFunction()); // node not implemented yet
+			ConversionFunctionList->SetItemSelection(InArgs._CurrentSelected.GetValue().GetConversionFunction(), true);
 		}
 	}
 
diff --git a/Engine/Plugins/Runtime/ModelViewViewModel/Source/ModelViewViewModelEditor/Private/Widgets/SMVVMFieldSelectorMenu.h b/Engine/Plugins/Runtime/ModelViewViewModel/Source/ModelViewViewModelEditor/Private/Widgets/SMVVMFieldSelectorMenu.h
--- a/Engine/Plugins/Runtime/ModelViewViewModel/Source/ModelViewViewModelEditor/Private/Widgets/SMVVMFieldSelectorMenu.h
+++ b/Engine/Plugins/Runtime/ModelViewViewModel/Source/ModelViewViewModelEditor/Private/Widgets/SMVVMFieldSelectorMenu.h
@@ -2,11 +2,12 @@
 
 #pragma once
 
-#include "MVVMPropertyPath.h"
-#include "Styling/CoreStyle.h"
-#include "Styling/SlateTypes.h"
 #include "Types/MVVMBindingMode.h"
 #include "Types/MVVMBindingSource.h"
+#include "Types/MVVMLinkedPinValue.h"
+
+#include "Styling/CoreStyle.h"
+#include "Styling/SlateTypes.h"
 #include "Widgets/SCompoundWidget.h"
 
 namespace ESelectInfo { enum Type : int; }
@@ -35,22 +36,18 @@ struct FFieldSelectionContext
 	bool bWritable = true;
 };
 
-DECLARE_DELEGATE_TwoParams(FOnFieldSelectionChanged, FMVVMBlueprintPropertyPath, const UFunction*);
+DECLARE_DELEGATE_OneParam(FOnLinkedValueSelectionChanged, FMVVMLinkedPinValue);
 
 class SFieldSelectorMenu : public SCompoundWidget
 {
 public:
 	DECLARE_DELEGATE_RetVal(FFieldSelectionContext, FOnGetFieldSelectionContext);
 
-	SLATE_BEGIN_ARGS(SFieldSelectorMenu)
-	{
-	}
-		SLATE_EVENT(FOnFieldSelectionChanged, OnFieldSelectionChanged)
+	SLATE_BEGIN_ARGS(SFieldSelectorMenu){}
+		SLATE_ARGUMENT(TOptional<FMVVMLinkedPinValue>, CurrentSelected)
+		SLATE_EVENT(FOnLinkedValueSelectionChanged, OnSelectionChanged)
 		SLATE_EVENT(FSimpleDelegate, OnMenuCloseRequested)
-
 		SLATE_ARGUMENT(FFieldSelectionContext, SelectionContext)
-		SLATE_ARGUMENT(TOptional<FMVVMBlueprintPropertyPath>, CurrentPropertyPathSelected)
-		SLATE_ARGUMENT_DEFAULT(const UFunction*, CurrentFunctionSelected) = nullptr;
 	SLATE_END_ARGS()
 
 	void Construct(const FArguments& InArgs, const UWidgetBlueprint* InWidgetBlueprint);
@@ -111,7 +108,7 @@ private:
 
 private:
 	TWeakObjectPtr<const UWidgetBlueprint> WidgetBlueprint;
-	FOnFieldSelectionChanged OnFieldSelectionChanged;
+	FOnLinkedValueSelectionChanged OnSelectionChanged;
 	FSimpleDelegate OnMenuCloseRequested;
 	FFieldSelectionContext SelectionContext;
 
diff --git a/Engine/Plugins/Runtime/ModelViewViewModel/Source/ModelViewViewModelEditor/Private/Widgets/SMVVMFunctionParameter.cpp b/Engine/Plugins/Runtime/ModelViewViewModel/Source/ModelViewViewModelEditor/Private/Widgets/SMVVMFunctionParameter.cpp
--- a/Engine/Plugins/Runtime/ModelViewViewModel/Source/ModelViewViewModelEditor/Private/Widgets/SMVVMFunctionParameter.cpp
+++ b/Engine/Plugins/Runtime/ModelViewViewModel/Source/ModelViewViewModelEditor/Private/Widgets/SMVVMFunctionParameter.cpp
@@ -79,8 +79,7 @@ void SFunctionParameter::Construct(const FArguments& InArgs, UWidgetBlueprint* I
 			];
 	}
 
-	FMVVMBlueprintPropertyPath Path = OnGetSelectedField();
-	bDefaultValueVisible = !Path.IsValid();
+	bDefaultValueVisible = !OnGetSelectedField().IsValid();
 
 	bool bFromViewModel = UE::MVVM::IsForwardBinding(Binding->BindingType);
 	TSharedPtr<SHorizontalBox> HBox;
@@ -104,8 +103,8 @@ void SFunctionParameter::Construct(const FArguments& InArgs, UWidgetBlueprint* I
 				+ SWidgetSwitcher::Slot()
 				[
 					SNew(SFieldSelector, WidgetBlueprint.Get())
-					.OnGetPropertyPath(this, &SFunctionParameter::OnGetSelectedField)
-					.OnFieldSelectionChanged(this, &SFunctionParameter::HandleFieldSelectionChanged)
+					.OnGetLinkedValue(this, &SFunctionParameter::OnGetSelectedField)
+					.OnSelectionChanged(this, &SFunctionParameter::HandleFieldSelectionChanged)
 					.OnGetSelectionContext(this, &SFunctionParameter::GetSelectedSelectionContext)
 				]
 			]
@@ -158,7 +157,7 @@ void SFunctionParameter::OnBindArgumentChecked(ECheckBoxState Checked)
 	if (bDefaultValueVisible)
 	{
 		PreviousSelectedField = OnGetSelectedField();
-		SetSelectedField(FMVVMBlueprintPropertyPath());
+		SetSelectedField(FMVVMLinkedPinValue());
 	}
 	else
 	{
@@ -166,7 +165,7 @@ void SFunctionParameter::OnBindArgumentChecked(ECheckBoxState Checked)
 	}
 }
 
-FMVVMBlueprintPropertyPath SFunctionParameter::OnGetSelectedField() const
+FMVVMLinkedPinValue SFunctionParameter::OnGetSelectedField() const
 {
 	if (const UWidgetBlueprint* WidgetBlueprintPtr = WidgetBlueprint.Get())
 	{
@@ -175,14 +174,14 @@ FMVVMBlueprintPropertyPath SFunctionParameter::OnGetSelectedField() const
 		{
 			if (const FMVVMBlueprintViewBinding* Binding = View->GetBinding(BindingId))
 			{
-				return Subsystem->GetPathForConversionFunctionArgument(WidgetBlueprint.Get(), *Binding, ParameterId, bSourceToDestination);
+				return FMVVMLinkedPinValue(Subsystem->GetPathForConversionFunctionArgument(WidgetBlueprint.Get(), *Binding, ParameterId, bSourceToDestination));
 			}
 		}
 	}
-	return FMVVMBlueprintPropertyPath();
+	return FMVVMLinkedPinValue();
 }
 
-void SFunctionParameter::SetSelectedField(const FMVVMBlueprintPropertyPath& Path)
+void SFunctionParameter::SetSelectedField(const FMVVMLinkedPinValue& Value)
 {
 	if (const UWidgetBlueprint* WidgetBlueprintPtr = WidgetBlueprint.Get())
 	{
@@ -191,15 +190,15 @@ void SFunctionParameter::SetSelectedField(const FMVVMBlueprintPropertyPath& Path
 		{
 			if (FMVVMBlueprintViewBinding* Binding = View->GetBinding(BindingId))
 			{
-				Subsystem->SetPathForConversionFunctionArgument(WidgetBlueprint.Get(), *Binding, ParameterId, Path, bSourceToDestination);
+				Subsystem->SetPathForConversionFunctionArgument(WidgetBlueprint.Get(), *Binding, ParameterId, Value.IsPropertyPath() ? Value.GetPropertyPath() : FMVVMBlueprintPropertyPath(), bSourceToDestination);
 			}
 		}
 	}
 }
 
-void SFunctionParameter::HandleFieldSelectionChanged(FMVVMBlueprintPropertyPath SelectedField, const UFunction* Function)
+void SFunctionParameter::HandleFieldSelectionChanged(FMVVMLinkedPinValue Value)
 {
-	SetSelectedField(SelectedField);
+	SetSelectedField(Value);
 }
 
 FFieldSelectionContext SFunctionParameter::GetSelectedSelectionContext() const
diff --git a/Engine/Plugins/Runtime/ModelViewViewModel/Source/ModelViewViewModelEditor/Private/Widgets/SMVVMFunctionParameter.h b/Engine/Plugins/Runtime/ModelViewViewModel/Source/ModelViewViewModelEditor/Private/Widgets/SMVVMFunctionParameter.h
--- a/Engine/Plugins/Runtime/ModelViewViewModel/Source/ModelViewViewModelEditor/Private/Widgets/SMVVMFunctionParameter.h
+++ b/Engine/Plugins/Runtime/ModelViewViewModel/Source/ModelViewViewModelEditor/Private/Widgets/SMVVMFunctionParameter.h
@@ -4,6 +4,7 @@
 
 #include "MVVMBlueprintPin.h"
 #include "MVVMPropertyPath.h"
+#include "Types/MVVMLinkedPinValue.h"
 #include "Widgets/SCompoundWidget.h"
 #include "Widgets/SMVVMFieldSelectorMenu.h"
 
@@ -37,10 +38,10 @@ private:
 	ECheckBoxState OnGetIsBindArgumentChecked() const;
 	void OnBindArgumentChecked(ECheckBoxState Checked);
 
-	FMVVMBlueprintPropertyPath OnGetSelectedField() const;
-	void SetSelectedField(const FMVVMBlueprintPropertyPath& Path);
+	FMVVMLinkedPinValue OnGetSelectedField() const;
+	void SetSelectedField(const FMVVMLinkedPinValue& Path);
 
-	void HandleFieldSelectionChanged(FMVVMBlueprintPropertyPath SelectedField, const UFunction* Function);
+	void HandleFieldSelectionChanged(FMVVMLinkedPinValue Value);
 	FFieldSelectionContext GetSelectedSelectionContext() const;
 
 	int32 GetCurrentWidgetIndex() const;
@@ -52,7 +53,7 @@ private:
 	/** This reference is just to keep the default value widget alive. */
 	TSharedPtr<SGraphPin> GraphPin;
 
-	FMVVMBlueprintPropertyPath PreviousSelectedField;
+	FMVVMLinkedPinValue PreviousSelectedField;
 
 	bool bSourceToDestination = true;
 	bool bAllowDefault = true;
diff --git a/Engine/Plugins/Runtime/ModelViewViewModel/Source/ModelViewViewModelEditor/Public/MVVMEditorSubsystem.h b/Engine/Plugins/Runtime/ModelViewViewModel/Source/ModelViewViewModelEditor/Public/MVVMEditorSubsystem.h
--- a/Engine/Plugins/Runtime/ModelViewViewModel/Source/ModelViewViewModelEditor/Public/MVVMEditorSubsystem.h
+++ b/Engine/Plugins/Runtime/ModelViewViewModel/Source/ModelViewViewModelEditor/Public/MVVMEditorSubsystem.h
@@ -14,6 +14,7 @@ enum class EMVVMBindingMode : uint8;
 enum class EMVVMExecutionMode : uint8;
 namespace UE::MVVM { struct FBindingSource; }
 struct FMVVMAvailableBinding;
+struct FMVVMBlueprintFunctionReference;
 struct FMVVMBlueprintPropertyPath;
 struct FMVVMBlueprintViewBinding;
 template <typename T> class TSubclassOf;
@@ -69,8 +70,12 @@ public:
 	UFUNCTION(BlueprintCallable, Category = "Viewmodel")
 	TArray<FMVVMAvailableBinding> GetChildViewModels(TSubclassOf<UObject> Class, TSubclassOf<UObject> Accessor);
 
+	UE_DEPRECATED(5.4, "SetSourceToDestinationConversionFunction with a UFunction is deprecated.")
 	void SetSourceToDestinationConversionFunction(UWidgetBlueprint* WidgetBlueprint, FMVVMBlueprintViewBinding& Binding, const UFunction* ConversionFunction);
+	void SetSourceToDestinationConversionFunction(UWidgetBlueprint* WidgetBlueprint, FMVVMBlueprintViewBinding& Binding, const FMVVMBlueprintFunctionReference& ConversionFunction);
+	UE_DEPRECATED(5.4, "SetDestinationToSourceConversionFunction  with a UFunction is deprecated.")
 	void SetDestinationToSourceConversionFunction(UWidgetBlueprint* WidgetBlueprint, FMVVMBlueprintViewBinding& Binding, const UFunction* ConversionFunction);
+	void SetDestinationToSourceConversionFunction(UWidgetBlueprint* WidgetBlueprint, FMVVMBlueprintViewBinding& Binding, const FMVVMBlueprintFunctionReference& ConversionFunction);
 	void SetDestinationPathForBinding(UWidgetBlueprint* WidgetBlueprint, FMVVMBlueprintViewBinding& Binding, FMVVMBlueprintPropertyPath Field);
 	void SetSourcePathForBinding(UWidgetBlueprint* WidgetBlueprint, FMVVMBlueprintViewBinding& Binding, FMVVMBlueprintPropertyPath Field);
 	void OverrideExecutionModeForBinding(UWidgetBlueprint* WidgetBlueprint, FMVVMBlueprintViewBinding& Binding, EMVVMExecutionMode Mode);
diff --git a/Engine/Plugins/Runtime/ModelViewViewModel/Source/ModelViewViewModelEditor/Public/Types/MVVMLinkedPinValue.h b/Engine/Plugins/Runtime/ModelViewViewModel/Source/ModelViewViewModelEditor/Public/Types/MVVMLinkedPinValue.h
new file mode 100644
--- /dev/null
+++ b/Engine/Plugins/Runtime/ModelViewViewModel/Source/ModelViewViewModelEditor/Public/Types/MVVMLinkedPinValue.h
@@ -0,0 +1,53 @@
+// Copyright Epic Games, Inc. All Rights Reserved.
+
+#pragma once
+
+#include "MVVMBlueprintFunctionReference.h"
+#include "MVVMPropertyPath.h"
+#include "Templates/SubclassOf.h"
+
+#include "MVVMLinkedPinValue.generated.h"
+
+class UFunction;
+class UK2Node;
+
+/** */
+USTRUCT(BlueprintType)
+struct MODELVIEWVIEWMODELEDITOR_API FMVVMLinkedPinValue
+{
+public:
+	GENERATED_BODY()
+
+	FMVVMLinkedPinValue() = default;
+	explicit FMVVMLinkedPinValue(FMVVMBlueprintPropertyPath InPath);
+	explicit FMVVMLinkedPinValue(const UBlueprint* InBlueprint, FMVVMBlueprintFunctionReference InConversion);
+	explicit FMVVMLinkedPinValue(const UFunction* Function);
+	explicit FMVVMLinkedPinValue(TSubclassOf<UK2Node> Node);
+
+public:
+	bool IsPropertyPath() const;
+	const FMVVMBlueprintPropertyPath& GetPropertyPath() const;
+
+	bool IsConversionFunction() const;
+	const UFunction* GetConversionFunction() const;
+
+	bool IsConversionNode() const;
+	TSubclassOf<UK2Node> GetConversionNode() const;
+
+public:
+	bool IsValid() const;
+	bool operator== (const FMVVMLinkedPinValue& Other) const
+	{
+		return ConversionNode == Other.ConversionNode && ConversionFunction == Other.ConversionFunction && PropertyPath == Other.PropertyPath;
+	}
+
+protected:
+	UPROPERTY(BlueprintReadWrite, Category = "MVVM")
+	FMVVMBlueprintPropertyPath PropertyPath;
+
+	UPROPERTY(BlueprintReadWrite, Category = "MVVM")
+	const UFunction* ConversionFunction = nullptr;
+
+	UPROPERTY(BlueprintReadWrite, Category = "MVVM")
+	TSubclassOf<UK2Node> ConversionNode;
+};
