MVVM: In the editor, the "show generated graph" button now only show the generated graph and does not duplicate or add it to the blueprint.
Fix build issue "Use of if (A) {...} else if (A){...}"

diff --git a/Engine/Plugins/Runtime/ModelViewViewModel/Source/ModelViewViewModelEditor/Private/Framework/MVVMRowHelper.cpp b/Engine/Plugins/Runtime/ModelViewViewModel/Source/ModelViewViewModelEditor/Private/Framework/MVVMRowHelper.cpp
--- a/Engine/Plugins/Runtime/ModelViewViewModel/Source/ModelViewViewModelEditor/Private/Framework/MVVMRowHelper.cpp
+++ b/Engine/Plugins/Runtime/ModelViewViewModel/Source/ModelViewViewModelEditor/Private/Framework/MVVMRowHelper.cpp
@@ -8,6 +8,7 @@
 #include "DragAndDrop/DecoratedDragDropOp.h"
 #include "Editor/EditorEngine.h"
 #include "EdGraphSchema_K2.h"
+#include "EdGraphUtilities.h"
 #include "Kismet2/BlueprintEditorUtils.h"
 #include "Misc/Guid.h"
 #include "ScopedTransaction.h"
@@ -105,34 +106,17 @@ void FRowHelper::DeleteEntries(const UWidgetBlueprint* WidgetBlueprint, UMVVMBlu
 	}
 }
 
-void FRowHelper::DuplicateBlueprintGraph(FBlueprintEditor* BlueprintEditor, UWidgetBlueprint* WidgetBlueprint, UMVVMBlueprintView* BlueprintView, TArrayView<const TSharedPtr<FBindingEntry>> Entries)
+void FRowHelper::ShowBlueprintGraph(FBlueprintEditor* BlueprintEditor, UWidgetBlueprint* WidgetBlueprint, UMVVMBlueprintView* BlueprintView, TArrayView<const TSharedPtr<FBindingEntry>> Entries)
 {
 	auto DuplicateGraph = [WidgetBlueprint, BlueprintEditor](UEdGraph* Graph)
 		{
-			if (Graph)
+			if (Graph && BlueprintEditor)
 			{
-				const FScopedTransaction Transaction(LOCTEXT("DuplicateGraph", "Duplicate Graph"));
-				WidgetBlueprint->Modify();
-
 				UEdGraph* DuplicatedGraph = GetDefault<UEdGraphSchema_K2>()->DuplicateGraph(Graph);
 				check(DuplicatedGraph);
+				DuplicatedGraph->SetFlags(RF_Transient);
 
-				for (UEdGraphNode* EdGraphNode : DuplicatedGraph->Nodes)
-				{
-					if (EdGraphNode)
-					{
-						EdGraphNode->CreateNewGuid();
-					}
-				}
-
-				DuplicatedGraph->Modify();
-				WidgetBlueprint->FunctionGraphs.Add(DuplicatedGraph);
-
-				FBlueprintEditorUtils::MarkBlueprintAsStructurallyModified(WidgetBlueprint);
-				if (BlueprintEditor)
-				{
-					BlueprintEditor->OpenDocument(DuplicatedGraph, FDocumentTracker::OpenNewDocument);
-				}
+				BlueprintEditor->OpenDocument(DuplicatedGraph, FDocumentTracker::OpenNewDocument);
 			}
 		};
 
diff --git a/Engine/Plugins/Runtime/ModelViewViewModel/Source/ModelViewViewModelEditor/Private/Framework/MVVMRowHelper.h b/Engine/Plugins/Runtime/ModelViewViewModel/Source/ModelViewViewModelEditor/Private/Framework/MVVMRowHelper.h
--- a/Engine/Plugins/Runtime/ModelViewViewModel/Source/ModelViewViewModelEditor/Private/Framework/MVVMRowHelper.h
+++ b/Engine/Plugins/Runtime/ModelViewViewModel/Source/ModelViewViewModelEditor/Private/Framework/MVVMRowHelper.h
@@ -23,7 +23,7 @@ struct FRowHelper
 	static void GatherAllChildBindings(UMVVMBlueprintView* BlueprintView, const TConstArrayView<TSharedPtr<FBindingEntry>> Entries, TArray<const FMVVMBlueprintViewBinding*>& OutBindings, TArray<UMVVMBlueprintViewEvent*>& OutEvents);
 
 	static void DeleteEntries(const UWidgetBlueprint* WidgetBlueprint, UMVVMBlueprintView* BlueprintView, TArrayView<const TSharedPtr<FBindingEntry>> Selection);
-	static void DuplicateBlueprintGraph(FBlueprintEditor* Editor, UWidgetBlueprint* WidgetBlueprint, UMVVMBlueprintView* BlueprintView, TArrayView<const TSharedPtr<FBindingEntry>> Selection);
+	static void ShowBlueprintGraph(FBlueprintEditor* Editor, UWidgetBlueprint* WidgetBlueprint, UMVVMBlueprintView* BlueprintView, TArrayView<const TSharedPtr<FBindingEntry>> Selection);
 
 	static TOptional<FMVVMBlueprintPropertyPath> DropFieldSelector(const UWidgetBlueprint* WidgetBlueprint, const FDragDropEvent& DragDropEvent);
 	static void DragEnterFieldSelector(const UWidgetBlueprint* WidgetBlueprint, const FDragDropEvent& DragDropEvent);
diff --git a/Engine/Plugins/Runtime/ModelViewViewModel/Source/ModelViewViewModelEditor/Private/Widgets/BindingEntry/SMVVMBindingRow.cpp b/Engine/Plugins/Runtime/ModelViewViewModel/Source/ModelViewViewModelEditor/Private/Widgets/BindingEntry/SMVVMBindingRow.cpp
--- a/Engine/Plugins/Runtime/ModelViewViewModel/Source/ModelViewViewModelEditor/Private/Widgets/BindingEntry/SMVVMBindingRow.cpp
+++ b/Engine/Plugins/Runtime/ModelViewViewModel/Source/ModelViewViewModelEditor/Private/Widgets/BindingEntry/SMVVMBindingRow.cpp
@@ -381,7 +381,7 @@ FMVVMLinkedPinValue SBindingRow::GetFieldSelectedValue(bool bSourceToDest) const
 			{
 				return FMVVMLinkedPinValue(ConversionFunction->GetConversionFunction().GetFunction(GetBlueprint()));
 			}
-			else if (ConversionFunction->GetConversionFunction().GetType() == EMVVMBlueprintFunctionReferenceType::Function)
+			else if (ConversionFunction->GetConversionFunction().GetType() == EMVVMBlueprintFunctionReferenceType::Node)
 			{
 				return FMVVMLinkedPinValue(ConversionFunction->GetConversionFunction().GetNode());
 			}
@@ -770,18 +770,16 @@ TSharedRef<SWidget> SBindingRow::HandleContextMenu() const
 		FMVVMBlueprintViewBinding* ViewBinding = GetThisViewBinding();
 		if (GetDefault<UMVVMDeveloperProjectSettings>()->bShowDeveloperGenerateGraphSettings)
 		{
-			bool bCanDuplicateGraph = ViewBinding && (ViewBinding->Conversion.GetConversionFunction(true) != nullptr || ViewBinding->Conversion.GetConversionFunction(false) != nullptr);
-			//CanDuplicateGraph(GraphToDuplicate)
-
-			FUIAction DuplicateAction;
-			DuplicateAction.ExecuteAction = FExecuteAction::CreateSP(this, &SBindingRow::HandleDuplicateGraph);
-			DuplicateAction.CanExecuteAction = FCanExecuteAction::CreateLambda([bCanDuplicateGraph]() { return bCanDuplicateGraph; });
-			MenuBuilder.AddMenuEntry(LOCTEXT("DuplicateGraph", "Copy binding graph")
-				, LOCTEXT("DuplicateGraphTooltip", "Add a copy of the binding graph to Blueprint."
-					" The copied graph will not bu used."
+			bool bCanShowGraph = ViewBinding && (ViewBinding->Conversion.GetConversionFunction(true) != nullptr || ViewBinding->Conversion.GetConversionFunction(false) != nullptr);
+
+			FUIAction ShowGraphAction;
+			ShowGraphAction.ExecuteAction = FExecuteAction::CreateSP(this, &SBindingRow::HandleShowBlueprintGraph);
+			ShowGraphAction.CanExecuteAction = FCanExecuteAction::CreateLambda([bCanShowGraph]() { return bCanShowGraph; });
+			MenuBuilder.AddMenuEntry(LOCTEXT("ShowGraph", "Show binding graph")
+				, LOCTEXT("ShowGraphTooltip", "Show the Blueprint graph that represent the binding."
 					" The graph is always generated but may not be visible to the user.")
 				, FSlateIcon(FAppStyle::GetAppStyleSetName(), "Icons.Duplicate")
-				, DuplicateAction);
+				, ShowGraphAction);
 		}
 
 		MenuBuilder.EndSection();
@@ -790,10 +788,10 @@ TSharedRef<SWidget> SBindingRow::HandleContextMenu() const
 	return MenuBuilder.MakeWidget();
 }
 
-void SBindingRow::HandleDuplicateGraph() const
+void SBindingRow::HandleShowBlueprintGraph() const
 {
 	TSharedPtr<FBindingEntry> Entry = GetEntry();
-	BindingEntry::FRowHelper::DuplicateBlueprintGraph(GetBlueprintEditor().Get(), GetBlueprint(), GetBlueprintView(), MakeArrayView(&Entry, 1));
+	BindingEntry::FRowHelper::ShowBlueprintGraph(GetBlueprintEditor().Get(), GetBlueprint(), GetBlueprintView(), MakeArrayView(&Entry, 1));
 }
 
 } // namespace
diff --git a/Engine/Plugins/Runtime/ModelViewViewModel/Source/ModelViewViewModelEditor/Private/Widgets/BindingEntry/SMVVMBindingRow.h b/Engine/Plugins/Runtime/ModelViewViewModel/Source/ModelViewViewModelEditor/Private/Widgets/BindingEntry/SMVVMBindingRow.h
--- a/Engine/Plugins/Runtime/ModelViewViewModel/Source/ModelViewViewModelEditor/Private/Widgets/BindingEntry/SMVVMBindingRow.h
+++ b/Engine/Plugins/Runtime/ModelViewViewModel/Source/ModelViewViewModelEditor/Private/Widgets/BindingEntry/SMVVMBindingRow.h
@@ -92,7 +92,7 @@ private:
 	void OnBindingModeSelectionChanged(FName ValueName, ESelectInfo::Type);
 
 	TSharedRef<SWidget> HandleContextMenu() const;
-	void HandleDuplicateGraph() const;
+	void HandleShowBlueprintGraph() const;
 
 private:
 	TArray<TSharedPtr<FText>> ErrorItems;
diff --git a/Engine/Plugins/Runtime/ModelViewViewModel/Source/ModelViewViewModelEditor/Private/Widgets/BindingEntry/SMVVMEventRow.cpp b/Engine/Plugins/Runtime/ModelViewViewModel/Source/ModelViewViewModelEditor/Private/Widgets/BindingEntry/SMVVMEventRow.cpp
--- a/Engine/Plugins/Runtime/ModelViewViewModel/Source/ModelViewViewModelEditor/Private/Widgets/BindingEntry/SMVVMEventRow.cpp
+++ b/Engine/Plugins/Runtime/ModelViewViewModel/Source/ModelViewViewModelEditor/Private/Widgets/BindingEntry/SMVVMEventRow.cpp
@@ -399,20 +399,18 @@ TSharedRef<SWidget> SEventRow::HandleContextMenu() const
 		MenuBuilder.BeginSection("Developer", LOCTEXT("Developer", "Developer"));
 
 		UMVVMBlueprintViewEvent* Event = GetEvent();
-		if (GetDefault<UMVVMDeveloperProjectSettings>()->bShowDeveloperGenerateGraphSettings)		
+		if (GetDefault<UMVVMDeveloperProjectSettings>()->bShowDeveloperGenerateGraphSettings)
 		{
-			bool bCanDuplicateGraph = Event && Event->GetWrapperGraph();
-			//CanDuplicateGraph(GraphToDuplicate)
-
-			FUIAction DuplicateAction;
-			DuplicateAction.ExecuteAction = FExecuteAction::CreateSP(this, &SEventRow::HandleDuplicateGraph);
-			DuplicateAction.CanExecuteAction = FCanExecuteAction::CreateLambda([bCanDuplicateGraph]() { return bCanDuplicateGraph; });
-			MenuBuilder.AddMenuEntry(LOCTEXT("DuplicateGraph", "Copy event graph")
-				, LOCTEXT("DuplicateGraphTooltip", "Add a copy of the event graph to Blueprint."
-					" The copied graph will not bu used."
+			bool bCanShowGraph = Event && Event->GetWrapperGraph();
+
+			FUIAction ShowGraphAction;
+			ShowGraphAction.ExecuteAction = FExecuteAction::CreateSP(this, &SEventRow::HandleShowBlueprintGraph);
+			ShowGraphAction.CanExecuteAction = FCanExecuteAction::CreateLambda([bCanShowGraph]() { return bCanShowGraph; });
+			MenuBuilder.AddMenuEntry(LOCTEXT("ShowGraph", "Show event graph")
+				, LOCTEXT("ShowGraphTooltip", "Show the Blueprint graph that represent the event."
 					" The graph is always generated but may not be visible to the user.")
 				, FSlateIcon(FAppStyle::GetAppStyleSetName(), "Icons.Duplicate")
-				, DuplicateAction);
+				, ShowGraphAction);
 		}
 
 		MenuBuilder.EndSection();
@@ -421,10 +419,10 @@ TSharedRef<SWidget> SEventRow::HandleContextMenu() const
 	return MenuBuilder.MakeWidget();
 }
 
-void SEventRow::HandleDuplicateGraph() const
+void SEventRow::HandleShowBlueprintGraph() const
 {
 	TSharedPtr<FBindingEntry> Entry = GetEntry();
-	BindingEntry::FRowHelper::DuplicateBlueprintGraph(GetBlueprintEditor().Get(), GetBlueprint(), GetBlueprintView(), MakeArrayView(&Entry, 1));
+	BindingEntry::FRowHelper::ShowBlueprintGraph(GetBlueprintEditor().Get(), GetBlueprint(), GetBlueprintView(), MakeArrayView(&Entry, 1));
 }
 
 } // namespace
diff --git a/Engine/Plugins/Runtime/ModelViewViewModel/Source/ModelViewViewModelEditor/Private/Widgets/BindingEntry/SMVVMEventRow.h b/Engine/Plugins/Runtime/ModelViewViewModel/Source/ModelViewViewModelEditor/Private/Widgets/BindingEntry/SMVVMEventRow.h
--- a/Engine/Plugins/Runtime/ModelViewViewModel/Source/ModelViewViewModelEditor/Private/Widgets/BindingEntry/SMVVMEventRow.h
+++ b/Engine/Plugins/Runtime/ModelViewViewModel/Source/ModelViewViewModelEditor/Private/Widgets/BindingEntry/SMVVMEventRow.h
@@ -62,7 +62,7 @@ private:
 	void HandleFieldSelectorDragEnter(const FGeometry& MyGeometry, const FDragDropEvent& DragDropEvent, bool bSource);
 
 	TSharedRef<SWidget> HandleContextMenu() const;
-	void HandleDuplicateGraph() const;
+	void HandleShowBlueprintGraph() const;
 
 private:
 	TArray<TSharedPtr<FText>> ErrorItems;
