MVVM: Show binding blueprint graph doesn't copy the graph. It now bring it into focus or show it into a new tab if it's a temporary graph.

diff --git a/Engine/Plugins/Runtime/ModelViewViewModel/Source/ModelViewViewModelEditor/Private/Framework/MVVMRowHelper.cpp b/Engine/Plugins/Runtime/ModelViewViewModel/Source/ModelViewViewModelEditor/Private/Framework/MVVMRowHelper.cpp
--- a/Engine/Plugins/Runtime/ModelViewViewModel/Source/ModelViewViewModelEditor/Private/Framework/MVVMRowHelper.cpp
+++ b/Engine/Plugins/Runtime/ModelViewViewModel/Source/ModelViewViewModelEditor/Private/Framework/MVVMRowHelper.cpp
@@ -10,6 +10,7 @@
 #include "EdGraphSchema_K2.h"
 #include "EdGraphUtilities.h"
 #include "Kismet2/BlueprintEditorUtils.h"
+#include "Kismet2/KismetEditorUtilities.h"
 #include "Misc/Guid.h"
 #include "ScopedTransaction.h"
 #include "Widgets/ViewModelFieldDragDropOp.h"
@@ -108,15 +109,18 @@ void FRowHelper::DeleteEntries(const UWidgetBlueprint* WidgetBlueprint, UMVVMBlu
 
 void FRowHelper::ShowBlueprintGraph(FBlueprintEditor* BlueprintEditor, UWidgetBlueprint* WidgetBlueprint, UMVVMBlueprintView* BlueprintView, TArrayView<const TSharedPtr<FBindingEntry>> Entries)
 {
-	auto DuplicateGraph = [WidgetBlueprint, BlueprintEditor](UEdGraph* Graph)
+	auto ShowGraph = [WidgetBlueprint, BlueprintEditor](UEdGraph* Graph)
 		{
 			if (Graph && BlueprintEditor)
 			{
-				UEdGraph* DuplicatedGraph = GetDefault<UEdGraphSchema_K2>()->DuplicateGraph(Graph);
-				check(DuplicatedGraph);
-				DuplicatedGraph->SetFlags(RF_Transient);
-
-				BlueprintEditor->OpenDocument(DuplicatedGraph, FDocumentTracker::OpenNewDocument);
+				if (Graph->HasAnyFlags(RF_Transient))
+				{
+					BlueprintEditor->OpenDocument(Graph, FDocumentTracker::OpenNewDocument);
+				}
+				else
+				{
+					FKismetEditorUtilities::BringKismetToFocusAttentionOnObject(Graph);
+				}
 			}
 		};
 
@@ -131,11 +135,11 @@ void FRowHelper::ShowBlueprintGraph(FBlueprintEditor* BlueprintEditor, UWidgetBl
 				UMVVMBlueprintViewConversionFunction* ConversionFunctionB = ViewBinding->Conversion.GetConversionFunction(false);
 				if (ConversionFunctionA)
 				{
-					DuplicateGraph(ConversionFunctionA->GetWrapperGraph());
+					ShowGraph(ConversionFunctionA->GetWrapperGraph());
 				}
 				if (ConversionFunctionB)
 				{
-					DuplicateGraph(ConversionFunctionB->GetWrapperGraph());
+					ShowGraph(ConversionFunctionB->GetWrapperGraph());
 				}
 			}
 		}
@@ -143,7 +147,7 @@ void FRowHelper::ShowBlueprintGraph(FBlueprintEditor* BlueprintEditor, UWidgetBl
 		{
 			if (UMVVMBlueprintViewEvent* Event = Entry->GetEvent())
 			{
-				DuplicateGraph(Event->GetWrapperGraph());
+				ShowGraph(Event->GetWrapperGraph());
 			}
 		}
 	}
diff --git a/Engine/Plugins/Runtime/ModelViewViewModel/Source/ModelViewViewModelEditor/Private/Widgets/BindingEntry/SMVVMBindingRow.cpp b/Engine/Plugins/Runtime/ModelViewViewModel/Source/ModelViewViewModelEditor/Private/Widgets/BindingEntry/SMVVMBindingRow.cpp
--- a/Engine/Plugins/Runtime/ModelViewViewModel/Source/ModelViewViewModelEditor/Private/Widgets/BindingEntry/SMVVMBindingRow.cpp
+++ b/Engine/Plugins/Runtime/ModelViewViewModel/Source/ModelViewViewModelEditor/Private/Widgets/BindingEntry/SMVVMBindingRow.cpp
@@ -778,7 +778,7 @@ TSharedRef<SWidget> SBindingRow::HandleContextMenu() const
 			MenuBuilder.AddMenuEntry(LOCTEXT("ShowGraph", "Show binding graph")
 				, LOCTEXT("ShowGraphTooltip", "Show the Blueprint graph that represent the binding."
 					" The graph is always generated but may not be visible to the user.")
-				, FSlateIcon(FAppStyle::GetAppStyleSetName(), "Icons.Duplicate")
+				, FSlateIcon(FAppStyle::GetAppStyleSetName(), "BlueprintEditor.FindInBlueprints.MenuIcon")
 				, ShowGraphAction);
 		}
 
diff --git a/Engine/Plugins/Runtime/ModelViewViewModel/Source/ModelViewViewModelEditor/Private/Widgets/BindingEntry/SMVVMEventRow.cpp b/Engine/Plugins/Runtime/ModelViewViewModel/Source/ModelViewViewModelEditor/Private/Widgets/BindingEntry/SMVVMEventRow.cpp
--- a/Engine/Plugins/Runtime/ModelViewViewModel/Source/ModelViewViewModelEditor/Private/Widgets/BindingEntry/SMVVMEventRow.cpp
+++ b/Engine/Plugins/Runtime/ModelViewViewModel/Source/ModelViewViewModelEditor/Private/Widgets/BindingEntry/SMVVMEventRow.cpp
@@ -409,7 +409,7 @@ TSharedRef<SWidget> SEventRow::HandleContextMenu() const
 			MenuBuilder.AddMenuEntry(LOCTEXT("ShowGraph", "Show event graph")
 				, LOCTEXT("ShowGraphTooltip", "Show the Blueprint graph that represent the event."
 					" The graph is always generated but may not be visible to the user.")
-				, FSlateIcon(FAppStyle::GetAppStyleSetName(), "Icons.Duplicate")
+				, FSlateIcon(FAppStyle::GetAppStyleSetName(), "BlueprintEditor.FindInBlueprints.MenuIcon")
 				, ShowGraphAction);
 		}
 
