MVVM: Fix issue where the binding list was registering to the modification callback too many times.

diff --git a/Engine/Plugins/Runtime/ModelViewViewModel/Source/ModelViewViewModelEditor/Private/Widgets/SMVVMViewBindingListView.cpp b/Engine/Plugins/Runtime/ModelViewViewModel/Source/ModelViewViewModelEditor/Private/Widgets/SMVVMViewBindingListView.cpp
--- a/Engine/Plugins/Runtime/ModelViewViewModel/Source/ModelViewViewModelEditor/Private/Widgets/SMVVMViewBindingListView.cpp
+++ b/Engine/Plugins/Runtime/ModelViewViewModel/Source/ModelViewViewModelEditor/Private/Widgets/SMVVMViewBindingListView.cpp
@@ -206,6 +206,30 @@ void SBindingsList::GetChildrenOfEntry(TSharedPtr<FBindingEntry> Entry, TArray<T
 	OutChildren.Append(Entry->GetFilteredChildren());
 }
 
+template<typename TEntryValueType>
+void SBindingsList::RegisterWrapperGraphModified(TEntryValueType* EntryValue, TSharedPtr<FBindingEntry> BindingEntry)
+{
+	{
+		const FObjectKey ObjectKey = FObjectKey(EntryValue);
+		TPair<TWeakPtr<FBindingEntry>, FDelegateHandle>* FoundWrapperGraphModifiedPtr = WrapperGraphModifiedDelegates.Find(ObjectKey);
+		if (FoundWrapperGraphModifiedPtr)
+		{
+			TSharedPtr<FBindingEntry> FoundWrapperGraphModified = FoundWrapperGraphModifiedPtr->Get<0>().Pin();
+			if (FoundWrapperGraphModified != BindingEntry)
+			{
+				EntryValue->OnWrapperGraphModified.Remove(FoundWrapperGraphModifiedPtr->Get<1>());
+				WrapperGraphModifiedDelegates.Remove(ObjectKey);
+				FoundWrapperGraphModifiedPtr = nullptr;
+			}
+		}
+		if (FoundWrapperGraphModifiedPtr == nullptr)
+		{
+			FDelegateHandle DelegateHandle = EntryValue->OnWrapperGraphModified.AddSP(this, &SBindingsList::HandleRefreshChildren, ObjectKey);
+			WrapperGraphModifiedDelegates.Add(ObjectKey, { TWeakPtr<FBindingEntry>(BindingEntry), DelegateHandle });
+		}
+	}
+}
+
 void SBindingsList::Refresh()
 {
 	struct FPreviousGroup
@@ -339,7 +363,12 @@ void SBindingsList::Refresh()
 			// Create/Find entries for conversion function parameters
 			if (UMVVMBlueprintViewConversionFunction* ConversionFunction = Binding.Conversion.GetConversionFunction(UE::MVVM::IsForwardBinding(Binding.BindingType)))
 			{
+				// Register to any modifications made in the graph
+				RegisterWrapperGraphModified(ConversionFunction, BindingEntry);
+
+				// Make sure the graph is up to date
 				ConversionFunction->GetOrCreateWrapperGraph(MVVMExtensionPtr->GetWidgetBlueprint());
+
 				for (const FMVVMBlueprintPin& Pin : ConversionFunction->GetPins())
 				{
 					UEdGraphPin* GraphPin = ConversionFunction->GetOrCreateGraphPin(MVVMExtensionPtr->GetWidgetBlueprint(), Pin.GetId());
@@ -366,7 +395,6 @@ void SBindingsList::Refresh()
 					{
 						ArgumentEntry = MakeShared<FBindingEntry>();
 						ArgumentEntry->SetBindingParameter(Binding.BindingId, Pin.GetId());
-						ConversionFunction->OnWrapperGraphModified.AddSP(this, &SBindingsList::ForceRefresh);
 
 						NewEntries.Add(ArgumentEntry);
 					}
@@ -377,6 +405,7 @@ void SBindingsList::Refresh()
 
 		for (UMVVMBlueprintViewEvent* Event : BlueprintView->GetEvents())
 		{
+			// Make sure the graph is up to date
 			Event->GetOrCreateWrapperGraph();
 
 			FName GroupName;
@@ -426,6 +455,9 @@ void SBindingsList::Refresh()
 				GroupEntry->AddChild(EventEntry);
 			}
 
+			// Register to any modifications made by the graph
+			RegisterWrapperGraphModified(Event, EventEntry);
+
 			// Create/Find entries for function parameters
 			for (const FMVVMBlueprintPin& Pin : Event->GetPins())
 			{
@@ -453,7 +485,6 @@ void SBindingsList::Refresh()
 				{
 					ArgumentEntry = MakeShared<FBindingEntry>();
 					ArgumentEntry->SetEventParameter(Event, Pin.GetId());
-					Event->OnWrapperGraphModified.AddSP(this, &SBindingsList::ForceRefresh);
 
 					NewEntries.Add(ArgumentEntry);
 				}
@@ -481,6 +512,20 @@ void SBindingsList::ForceRefresh()
 	Refresh();
 }
 
+void SBindingsList::HandleRefreshChildren(FObjectKey ObjectHolder)
+{
+	TPair<TWeakPtr<FBindingEntry>, FDelegateHandle>* Found = WrapperGraphModifiedDelegates.Find(ObjectHolder);
+	if (Found)
+	{
+		TSharedPtr<FBindingEntry> FoundWrapperGraphModified = Found->Get<0>().Pin();
+		if (FoundWrapperGraphModified)
+		{
+			FoundWrapperGraphModified->ResetChildren();
+			Refresh();
+		}
+	}
+}
+
 TSharedRef<ITableRow> SBindingsList::GenerateEntryRow(TSharedPtr<FBindingEntry> Entry, const TSharedRef<STableViewBase>& OwnerTable) const
 {
 	TSharedPtr<ITableRow> Row;
diff --git a/Engine/Plugins/Runtime/ModelViewViewModel/Source/ModelViewViewModelEditor/Private/Widgets/SMVVMViewBindingListView.h b/Engine/Plugins/Runtime/ModelViewViewModel/Source/ModelViewViewModelEditor/Private/Widgets/SMVVMViewBindingListView.h
--- a/Engine/Plugins/Runtime/ModelViewViewModel/Source/ModelViewViewModelEditor/Private/Widgets/SMVVMViewBindingListView.h
+++ b/Engine/Plugins/Runtime/ModelViewViewModel/Source/ModelViewViewModelEditor/Private/Widgets/SMVVMViewBindingListView.h
@@ -2,7 +2,7 @@
 
 #pragma once
 
-
+#include "UObject/ObjectKey.h"
 #include "Widgets/SCompoundWidget.h"
 
 namespace ESelectInfo { enum Type : int; }
@@ -48,10 +48,14 @@ private:
 	virtual FReply OnKeyDown(const FGeometry& MyGeometry, const FKeyEvent& InKeyEvent) override;
 
 	void ForceRefresh();
+	template<typename TEntryValueType>
+	void RegisterWrapperGraphModified(TEntryValueType* EntryValue, TSharedPtr<FBindingEntry> GroupEntry);
+	void HandleRefreshChildren(FObjectKey EntryKey);
 
 private:
 	TWeakPtr<SBindingsPanel> BindingPanel;
 	TSharedPtr<STreeView<TSharedPtr<FBindingEntry>>> TreeView;
+	TMap<FObjectKey, TPair<TWeakPtr<FBindingEntry>, FDelegateHandle>> WrapperGraphModifiedDelegates;
 	TArray<TSharedPtr<FBindingEntry>> AllRootGroups;
 	TArray<TSharedPtr<FBindingEntry>> FilteredRootGroups;
 	TWeakObjectPtr<UMVVMWidgetBlueprintExtension_View> MVVMExtension;
