MVVM: Enable K2Node as conversion functions.

diff --git a/Engine/Plugins/Runtime/ModelViewViewModel/Source/ModelViewViewModelBlueprint/Private/Bindings/MVVMConversionFunctionHelper.cpp b/Engine/Plugins/Runtime/ModelViewViewModel/Source/ModelViewViewModelBlueprint/Private/Bindings/MVVMConversionFunctionHelper.cpp
--- a/Engine/Plugins/Runtime/ModelViewViewModel/Source/ModelViewViewModelBlueprint/Private/Bindings/MVVMConversionFunctionHelper.cpp
+++ b/Engine/Plugins/Runtime/ModelViewViewModel/Source/ModelViewViewModelBlueprint/Private/Bindings/MVVMConversionFunctionHelper.cpp
@@ -11,7 +11,6 @@
 #include "MVVMWidgetBlueprintExtension_View.h"
 #include "UObject/MetaData.h"
 
-#include "EdGraphSchema_K2.h"
 #include "K2Node_BreakStruct.h"
 #include "K2Node_CallFunction.h"
 #include "K2Node_FunctionEntry.h"
@@ -20,6 +19,7 @@
 #include "K2Node_VariableGet.h"
 #include "K2Node_VariableSet.h"
 #include "Kismet2/BlueprintEditorUtils.h"
+#include "MVVMConversionFunctionGraphSchema.h"
 
 #define LOCTEXT_NAMESPACE "MVVMConversionFunctionHelper"
 
@@ -85,7 +85,7 @@ namespace Private
 	FCreateGraphResult CreateGraph(UBlueprint* Blueprint, FName GraphName, const UFunction* FunctionEntryDefinition, bool bIsConst, bool bIsEditable, bool bAddToBlueprint)
 	{
 		FName UniqueFunctionName = FBlueprintEditorUtils::FindUniqueKismetName(Blueprint, GraphName.ToString());
-		UEdGraph* FunctionGraph = FBlueprintEditorUtils::CreateNewGraph(Blueprint, UniqueFunctionName, UEdGraph::StaticClass(), UEdGraphSchema_K2::StaticClass());
+		UEdGraph* FunctionGraph = FBlueprintEditorUtils::CreateNewGraph(Blueprint, UniqueFunctionName, UEdGraph::StaticClass(), UMVVMConversionFunctionGraphSchema::StaticClass());
 		FunctionGraph->bEditable = bIsEditable;
 		if (bAddToBlueprint)
 		{
@@ -96,7 +96,7 @@ namespace Private
 			FunctionGraph->SetFlags(RF_Transient);
 		}
 
-		const UEdGraphSchema_K2* Schema = CastChecked<UEdGraphSchema_K2>(FunctionGraph->GetSchema());
+		const UEdGraphSchema_K2* Schema = GetDefault<UMVVMConversionFunctionGraphSchema>();
 		Schema->MarkFunctionEntryAsEditable(FunctionGraph, bIsEditable);
 		Schema->CreateDefaultNodesForGraph(*FunctionGraph);
 
@@ -342,7 +342,7 @@ namespace Private
 			{
 				return Pin->Direction == EGPD_Input
 					&& Pin->PinName != UEdGraphSchema_K2::PN_Execute
-					&& GetDefault<UEdGraphSchema_K2>()->ArePinsCompatible(PreviousDataPin, Pin, Context);
+					&& GetDefault<UMVVMConversionFunctionGraphSchema>()->ArePinsCompatible(PreviousDataPin, Pin, Context);
 			};
 
 		NumberOfFieldExcludingThePropertyPathSource = FMath::Clamp(NumberOfFieldExcludingThePropertyPathSource, 0, PropertyPath.GetFieldPaths().Num());
@@ -534,7 +534,7 @@ namespace Private
 			{
 				if (CanNewConnections(Pin, PreviousDataPin, PreviousClass))
 				{
-					Pin->MakeLinkTo(PreviousDataPin);
+					GetDefault<UMVVMConversionFunctionGraphSchema>()->TryCreateConnection(Pin, PreviousDataPin);
 				}
 			}
 
@@ -660,6 +660,7 @@ namespace Private
 	void LinkAllNodes(UEdGraph* FunctionGraph, UK2Node_FunctionEntry* FunctionEntry, UEdGraphNode* Wrapper, UK2Node_FunctionResult* FunctionResult)
 	{
 		check(FunctionGraph && FunctionEntry && Wrapper && FunctionResult);
+		const UEdGraphSchema* GraphSchema = GetDefault<UMVVMConversionFunctionGraphSchema>();
 
 		UEdGraphPin* ThenPin = FunctionEntry->FindPinChecked(UEdGraphSchema_K2::PN_Then);
 		check(ThenPin);
@@ -678,7 +679,7 @@ namespace Private
 					if (UEdGraphPin* ExecPin = CallFunction->FindPin(UEdGraphSchema_K2::PN_Execute))
 					{
 						ThenPin->BreakAllPinLinks();
-						ThenPin->MakeLinkTo(ExecPin);
+						GraphSchema->TryCreateConnection(ThenPin, ExecPin);
 						ThenPin = CallFunction->FindPinChecked(UEdGraphSchema_K2::PN_Then);
 					}
 				}
@@ -688,12 +689,12 @@ namespace Private
 		// Make pin to the conversion node or to the return node
 		if (UEdGraphPin* CallFunctionExecPin = Wrapper->FindPin(UEdGraphSchema_K2::PN_Execute))
 		{
-			ThenPin->MakeLinkTo(CallFunctionExecPin);
+			GraphSchema->TryCreateConnection(ThenPin, CallFunctionExecPin);
 		}
 		else
 		{
 			UEdGraphPin* FunctionResultExecPin = FunctionResult->GetExecPin();
-			ThenPin->MakeLinkTo(FunctionResultExecPin);
+			GraphSchema->TryCreateConnection(ThenPin, FunctionResultExecPin);
 		}
 	}
 } //namespace
@@ -739,7 +740,7 @@ TValueOrError<void, FText> CanCreateSetterGraph(UBlueprint* Blueprint, const FMV
 	return MakeValue();
 }
 
-TValueOrError<FCreateSetterGraphResult, FText> CreateSetterGraph(UBlueprint* Blueprint, FName GraphName, const UFunction* Signature, const FMVVMBlueprintPropertyPath& PropertyPath, bool bIsConst, bool bTransient)
+TValueOrError<FCreateGraphResult, FText> CreateSetterGraph(UBlueprint* Blueprint, FName GraphName, const UFunction* Signature, const FMVVMBlueprintPropertyPath& PropertyPath, bool bIsConst, bool bTransient)
 {
 	TArray<UE::MVVM::FMVVMConstFieldVariant> Fields = PropertyPath.GetCompleteFields(Blueprint);
 	TValueOrError<Private::FCanSetterGraphResult, FText> CanCreateSetterGraphResult = Private::CanCreateSetterGraph(Blueprint, Fields);
@@ -766,10 +767,12 @@ TValueOrError<FCreateSetterGraphResult, FText> CreateSetterGraph(UBlueprint* Blu
 		return MakeError(LOCTEXT("SetterGraph_CreateGraphFail", "Can create the graph object."));
 	}
 
-	FCreateSetterGraphResult Result;
+	FCreateGraphResult Result;
 	Result.NewGraph = CreateGraphInternalResult.FunctionGraph;
 	Result.WrappedNode = nullptr;
 
+	const UMVVMConversionFunctionGraphSchema* GraphSchema = GetDefault<UMVVMConversionFunctionGraphSchema>();
+
 	// Create the wrapper setter node
 	{
 		int32 WrapperFieldIndex = CanCreateSetterGraphResult.GetValue().bSplitPin ? CanCreateSetterGraphResult.GetValue().LocalContainerPathIndex : (CanCreateSetterGraphResult.GetValue().LocalContainerPathIndex + 1);
@@ -810,12 +813,12 @@ TValueOrError<FCreateSetterGraphResult, FText> CreateSetterGraph(UBlueprint* Blu
 
 		if (Result.WrappedNode->IsNodePure())
 		{
-			FunctionEntryThenPin->MakeLinkTo(FunctionResultExecPin);
+			GraphSchema->TryCreateConnection(FunctionEntryThenPin, FunctionResultExecPin);
 		}
 		else
 		{
-			FunctionEntryThenPin->MakeLinkTo(Result.WrappedNode->GetExecPin());
-			Result.WrappedNode->GetThenPin()->MakeLinkTo(FunctionResultExecPin);
+			GraphSchema->TryCreateConnection(FunctionEntryThenPin, Result.WrappedNode->GetExecPin());
+			GraphSchema->TryCreateConnection(Result.WrappedNode->GetThenPin(), FunctionResultExecPin);
 
 			Result.WrappedNode->NodePosY = 0;
 		}
@@ -841,7 +844,7 @@ TValueOrError<FCreateSetterGraphResult, FText> CreateSetterGraph(UBlueprint* Blu
 		int32 WrapperPinIndex = PropertyPathPinsResult.GetValue().Num() - 1 + DeltaIndex;
 		UEdGraphPin* ToLinkPin = PropertyPathPinsResult.GetValue()[WrapperPinIndex];
 
-		ToLinkPin->MakeLinkTo(WrapperSelfPin);
+		GraphSchema->TryCreateConnection(ToLinkPin, WrapperSelfPin);
 
 		SplitPin = CanCreateSetterGraphResult.GetValue().bSplitPin ? PropertyPathPinsResult.GetValue().Last() : nullptr;
 	}
@@ -881,7 +884,7 @@ TValueOrError<FCreateSetterGraphResult, FText> CreateSetterGraph(UBlueprint* Blu
 				{
 					if (OutPins[SubPinIndex]->GetFName() != PinName)
 					{
-						OutPins[SubPinIndex]->MakeLinkTo(InPins[SubPinIndex]);
+						GraphSchema->TryCreateConnection(OutPins[SubPinIndex], InPins[SubPinIndex]);
 						InPins[SubPinIndex]->bHidden = true;
 					}
 				}
@@ -894,15 +897,15 @@ TValueOrError<FCreateSetterGraphResult, FText> CreateSetterGraph(UBlueprint* Blu
 				UEdGraphPin** FoundOutPinPtr = OutPins.FindByPredicate([PinName](const UEdGraphPin* Other) { return Other->GetFName() == PinName && Other->Direction == EGPD_Output && !Other->bHidden; });
 				if (FoundInPinPtr == nullptr
 					|| FoundOutPinPtr == nullptr
-					|| !GetDefault<UEdGraphSchema_K2>()->CanSplitStructPin(**FoundInPinPtr)
-					|| !GetDefault<UEdGraphSchema_K2>()->CanSplitStructPin(**FoundOutPinPtr))
+					|| !GraphSchema->CanSplitStructPin(**FoundInPinPtr)
+					|| !GraphSchema->CanSplitStructPin(**FoundOutPinPtr))
 				{
 					RemoveGraph();
 					return MakeError(FText::Format(LOCTEXT("SetterGraph_CantSplitPin", "The pin {0} can't be split."), FText::FromName(NewWrapperField.GetName())));
 				}
 
-				GetDefault<UEdGraphSchema_K2>()->SplitPin(*FoundInPinPtr, false);
-				GetDefault<UEdGraphSchema_K2>()->SplitPin(*FoundOutPinPtr, false);
+				GraphSchema->SplitPin(*FoundInPinPtr, false);
+				GraphSchema->SplitPin(*FoundOutPinPtr, false);
 
 				InPins = (*FoundInPinPtr)->SubPins;
 				OutPins = (*FoundOutPinPtr)->SubPins;
@@ -913,20 +916,20 @@ TValueOrError<FCreateSetterGraphResult, FText> CreateSetterGraph(UBlueprint* Blu
 	return MakeValue(MoveTemp(Result));
 }
 
-TPair<UEdGraph*, UK2Node*> CreateGraph(UBlueprint* Blueprint, FName GraphName, const UFunction* Signature, const UFunction* FunctionToWrap, bool bIsConst, bool bTransient)
+FCreateGraphResult CreateGraph(UBlueprint* Blueprint, FName GraphName, const UFunction* Signature, const UFunction* FunctionToWrap, bool bIsConst, bool bTransient)
 {
 	bool bIsEditable = false;
 	bool bAddToBlueprint = !bTransient;
 
 	Private::FCreateGraphResult NewGraph = Private::CreateGraph(Blueprint, GraphName, Signature, bIsConst, bIsEditable, bAddToBlueprint);
-
-	const FProperty* ReturnProperty = UE::MVVM::BindingHelper::GetReturnProperty(FunctionToWrap);
+	const UMVVMConversionFunctionGraphSchema* GraphSchema = GetDefault<UMVVMConversionFunctionGraphSchema>();
 
 	// create return value pin
+	const FProperty* ReturnProperty = UE::MVVM::BindingHelper::GetReturnProperty(FunctionToWrap);
 	if (ReturnProperty)
 	{
 		TSharedPtr<FUserPinInfo> PinInfo = MakeShared<FUserPinInfo>();
-		GetDefault<UEdGraphSchema_K2>()->ConvertPropertyToPinType(ReturnProperty, PinInfo->PinType);
+		GraphSchema->ConvertPropertyToPinType(ReturnProperty, PinInfo->PinType);
 		PinInfo->PinName = ReturnProperty->GetFName();
 		PinInfo->DesiredPinDirection = EGPD_Input;
 		NewGraph.FunctionResult->UserDefinedPins.Add(PinInfo);
@@ -953,14 +956,14 @@ TPair<UEdGraph*, UK2Node*> CreateGraph(UBlueprint* Blueprint, FName GraphName, c
 			UEdGraphPin* CallFunctionExecPin = CallFunctionNode->GetExecPin();
 			UEdGraphPin* CallFunctionThenPin = CallFunctionNode->GetThenPin();
 
-			FunctionEntryThenPin->MakeLinkTo(CallFunctionExecPin);
-			CallFunctionThenPin->MakeLinkTo(FunctionResultExecPin);
+			GraphSchema->TryCreateConnection(FunctionEntryThenPin, CallFunctionExecPin);
+			GraphSchema->TryCreateConnection(CallFunctionThenPin, FunctionResultExecPin);
 
 			CallFunctionNode->NodePosY = 0;
 		}
 		else
 		{
-			FunctionEntryThenPin->MakeLinkTo(FunctionResultExecPin);
+			GraphSchema->TryCreateConnection(FunctionEntryThenPin, FunctionResultExecPin);
 			CallFunctionNode->NodePosY = 100;
 		}
 	}
@@ -970,19 +973,20 @@ TPair<UEdGraph*, UK2Node*> CreateGraph(UBlueprint* Blueprint, FName GraphName, c
 		UEdGraphPin* FunctionReturnPin = CallFunctionNode->FindPin(ReturnProperty->GetName(), EGPD_Output);
 		UEdGraphPin* FunctionResultPin = NewGraph.FunctionResult->FindPin(ReturnProperty->GetFName(), EGPD_Input);
 		check(FunctionResultPin && FunctionReturnPin);
-		FunctionReturnPin->MakeLinkTo(FunctionResultPin);
+		GraphSchema->TryCreateConnection(FunctionReturnPin, FunctionResultPin);
 	}
 
 	return { NewGraph.FunctionGraph , CallFunctionNode };
 }
 
 
-TPair<UEdGraph*, UK2Node*> CreateGraph(UBlueprint* Blueprint, FName GraphName, const UFunction* Signature, TSubclassOf<UK2Node> NodeType, bool bIsConst, bool bTransient, TFunctionRef<void(UK2Node*)> InitNodeCallback)
+FCreateGraphResult CreateGraph(UBlueprint* Blueprint, FName GraphName, const UFunction* Signature, TSubclassOf<UK2Node> NodeType, bool bIsConst, bool bTransient, TFunctionRef<void(UK2Node*)> InitNodeCallback)
 {
 	bool bIsEditable = false;
 	bool bAddToBlueprint = !bTransient;
 
 	Private::FCreateGraphResult NewGraph = Private::CreateGraph(Blueprint, GraphName, Signature, bIsConst, bIsEditable, bAddToBlueprint);
+	const UEdGraphSchema* GraphSchema = GetDefault<UMVVMConversionFunctionGraphSchema>();
 
 	UK2Node* CallFunctionNode = nullptr;
 	{
@@ -994,6 +998,18 @@ TPair<UEdGraph*, UK2Node*> CreateGraph(UBlueprint* Blueprint, FName GraphName, c
 		Private::MarkAsConversionFunction(CallFunctionNode, NewGraph.FunctionGraph);
 	}
 
+	// Create return value pin
+	UEdGraphPin* CallFunctionOutputPin = Private::FindNewOutputPin(CallFunctionNode);
+	if (CallFunctionOutputPin)
+	{
+		TSharedPtr<FUserPinInfo> PinInfo = MakeShared<FUserPinInfo>();
+		PinInfo->PinType = CallFunctionOutputPin->PinType;
+		PinInfo->PinName = CallFunctionOutputPin->GetFName();
+		PinInfo->DesiredPinDirection = EGPD_Input;
+		NewGraph.FunctionResult->UserDefinedPins.Add(PinInfo);
+		NewGraph.FunctionResult->ReconstructNode();
+	}
+
 	// Make link Entry -> CallFunction || Entry -> Return
 	{
 		UEdGraphPin* FunctionEntryThenPin = NewGraph.FunctionEntry->GetThenPin();
@@ -1004,18 +1020,26 @@ TPair<UEdGraph*, UK2Node*> CreateGraph(UBlueprint* Blueprint, FName GraphName, c
 			UEdGraphPin* CallFunctionExecPin = CallFunctionNode->GetExecPin();
 			UEdGraphPin* CallFunctionThenPin = CallFunctionNode->GetThenPin();
 
-			FunctionEntryThenPin->MakeLinkTo(CallFunctionExecPin);
-			CallFunctionThenPin->MakeLinkTo(FunctionResultExecPin);
+			GraphSchema->TryCreateConnection(FunctionEntryThenPin, CallFunctionExecPin);
+			GraphSchema->TryCreateConnection(CallFunctionThenPin, FunctionResultExecPin);
 
 			CallFunctionNode->NodePosY = 0;
 		}
 		else
 		{
-			FunctionEntryThenPin->MakeLinkTo(FunctionResultExecPin);
+			GraphSchema->TryCreateConnection(FunctionEntryThenPin, FunctionResultExecPin);
 			CallFunctionNode->NodePosY = 100;
 		}
 	}
 
+	if (CallFunctionOutputPin)
+	{
+		UEdGraphPin* FunctionResultPin = NewGraph.FunctionResult->FindPin(CallFunctionOutputPin->GetFName(), EGPD_Input);
+		check(FunctionResultPin);
+
+		GraphSchema->TryCreateConnection(CallFunctionOutputPin, FunctionResultPin);
+	}
+
 	return { NewGraph.FunctionGraph , CallFunctionNode };
 }
 
@@ -1073,7 +1097,7 @@ void SetPropertyPathForPin(const UBlueprint* Blueprint, const FMVVMBlueprintProp
 
 	UEdGraphNode* ConversionNode = PathPin->GetOwningNode();
 	UEdGraph* FunctionGraph = ConversionNode ? ConversionNode->GetGraph() : nullptr;
-	const UEdGraphSchema* Schema = FunctionGraph ? FunctionGraph->GetSchema() : nullptr;
+	const UEdGraphSchema* Schema = GetDefault<UMVVMConversionFunctionGraphSchema>();
 
 	UK2Node_FunctionEntry* ConverionFunctionEntry = ConversionNode ? Private::FindFunctionEntry(FunctionGraph) : nullptr;
 	UK2Node_FunctionResult* ConverionFunctionResult = ConversionNode ? Private::FindFunctionResult(FunctionGraph) : nullptr;
@@ -1105,7 +1129,7 @@ void SetPropertyPathForPin(const UBlueprint* Blueprint, const FMVVMBlueprintProp
 		}
 
 		// Link the last data pin to the Conversation Function Pin
-		BuildPropertyPathResult.GetValue().Last()->MakeLinkTo(PathPin);
+		Schema->TryCreateConnection(BuildPropertyPathResult.GetValue().Last(), PathPin);
 	}
 
 	// Link Then / Exec pin
@@ -1181,6 +1205,30 @@ TArray<FName> FindPinId(const UEdGraphPin* GraphPin)
 	return Result;
 }
 
+TArray<UEdGraphPin*> FindInputPins(const UK2Node* Node)
+{
+	TArray<UEdGraphPin*> Result;
+	if (Node == nullptr)
+	{
+		return Result;
+	}
+
+	Result.Reserve(Node->Pins.Num());
+	for (UEdGraphPin* GraphPin : Node->Pins)
+	{
+		if (IsInputPin(GraphPin))
+		{
+			Result.Add(GraphPin);
+		}
+	}
+	return Result;
+}
+
+UEdGraphPin* FindOutputPin(const UK2Node* Node)
+{
+	return Node? Private::FindNewOutputPin(Node) : nullptr;
+}
+
 } // UE::MVVM::ConversionFunctionHelper
 
 #undef LOCTEXT_NAMESPACE
diff --git a/Engine/Plugins/Runtime/ModelViewViewModel/Source/ModelViewViewModelBlueprint/Private/MVVMBlueprintFunctionReference.cpp b/Engine/Plugins/Runtime/ModelViewViewModel/Source/ModelViewViewModelBlueprint/Private/MVVMBlueprintFunctionReference.cpp
--- a/Engine/Plugins/Runtime/ModelViewViewModel/Source/ModelViewViewModelBlueprint/Private/MVVMBlueprintFunctionReference.cpp
+++ b/Engine/Plugins/Runtime/ModelViewViewModel/Source/ModelViewViewModelBlueprint/Private/MVVMBlueprintFunctionReference.cpp
@@ -84,6 +84,16 @@ TSubclassOf<UK2Node> FMVVMBlueprintFunctionReference::GetNode() const
 	return Type == EMVVMBlueprintFunctionReferenceType::Node ? Node : TSubclassOf<UK2Node>();
 }
 
+bool FMVVMBlueprintFunctionReference::IsValid(const UBlueprint* SelfContext) const
+{
+	return GetFunction(SelfContext) != nullptr || GetNode().Get() != nullptr;
+}
+
+bool FMVVMBlueprintFunctionReference::IsValid(const UClass* SelfContext) const
+{
+	return GetFunction(SelfContext) != nullptr || GetNode().Get() != nullptr;
+}
+
 bool FMVVMBlueprintFunctionReference::operator==(const FMVVMBlueprintFunctionReference& Other) const
 {
 	if (Type != Other.Type)
diff --git a/Engine/Plugins/Runtime/ModelViewViewModel/Source/ModelViewViewModelBlueprint/Private/MVVMBlueprintPin.cpp b/Engine/Plugins/Runtime/ModelViewViewModel/Source/ModelViewViewModelBlueprint/Private/MVVMBlueprintPin.cpp
--- a/Engine/Plugins/Runtime/ModelViewViewModel/Source/ModelViewViewModelBlueprint/Private/MVVMBlueprintPin.cpp
+++ b/Engine/Plugins/Runtime/ModelViewViewModel/Source/ModelViewViewModelBlueprint/Private/MVVMBlueprintPin.cpp
@@ -6,6 +6,7 @@
 #include "Bindings/MVVMConversionFunctionHelper.h"
 #include "EdGraph/EdGraphPin.h"
 #include "EdGraphSchema_K2.h"
+#include "MVVMConversionFunctionGraphSchema.h"
 
 #include UE_INLINE_GENERATED_CPP_BY_NAME(MVVMBlueprintPin)
 
@@ -208,9 +209,9 @@ void FMVVMBlueprintPin::CopyTo(const UBlueprint* Blueprint, UEdGraphNode* Node)
 				GetDefault<UEdGraphSchema_K2>()->SplitPin(GraphPin, false);
 			}
 
-			GraphPin->DefaultObject = DefaultObject;
-			GraphPin->DefaultValue = DefaultString;
-			GraphPin->DefaultTextValue = DefaultText;
+			GetDefault<UMVVMConversionFunctionGraphSchema>()->TrySetDefaultObject(*GraphPin, DefaultObject, false);
+			GetDefault<UMVVMConversionFunctionGraphSchema>()->TrySetDefaultValue(*GraphPin, DefaultString, false);
+			GetDefault<UMVVMConversionFunctionGraphSchema>()->TrySetDefaultText(*GraphPin, DefaultText, false);
 			UE::MVVM::ConversionFunctionHelper::SetPropertyPathForPin(Blueprint, Path, GraphPin);
 		}
 	}
diff --git a/Engine/Plugins/Runtime/ModelViewViewModel/Source/ModelViewViewModelBlueprint/Private/MVVMBlueprintViewConversionFunction.cpp b/Engine/Plugins/Runtime/ModelViewViewModel/Source/ModelViewViewModelBlueprint/Private/MVVMBlueprintViewConversionFunction.cpp
--- a/Engine/Plugins/Runtime/ModelViewViewModel/Source/ModelViewViewModelBlueprint/Private/MVVMBlueprintViewConversionFunction.cpp
+++ b/Engine/Plugins/Runtime/ModelViewViewModel/Source/ModelViewViewModelBlueprint/Private/MVVMBlueprintViewConversionFunction.cpp
@@ -63,23 +63,41 @@ bool UMVVMBlueprintViewConversionFunction::IsValidConversionFunction(const UBlue
 	return GetDefault<UMVVMDeveloperProjectSettings>()->IsConversionFunctionAllowed(WidgetBlueprint, Function);
 }
 
-bool UMVVMBlueprintViewConversionFunction::IsValid(const UBlueprint* SelfContext) const
+bool UMVVMBlueprintViewConversionFunction::IsValidConversionNode(const UBlueprint* WidgetBlueprint, const TSubclassOf<UK2Node> Function)
 {
-	if (ConversionFunction.GetType() != EMVVMBlueprintFunctionReferenceType::Function)
+	check(WidgetBlueprint);
+	check(Function.Get());
+
+	UEdGraphPin* ReturnPin = UE::MVVM::ConversionFunctionHelper::FindOutputPin(Function.GetDefaultObject());
+	if (ReturnPin == nullptr)
 	{
-		return false; // not valid for now
+		return false;
 	}
-	
-	UClass* SelfClass = SelfContext->SkeletonGeneratedClass ? SelfContext->SkeletonGeneratedClass : SelfContext->GeneratedClass;
-	const UFunction* Function = ConversionFunction.GetFunction(SelfContext);
-	if (Function)
+
+	TArray<UEdGraphPin*> ArgumentsResult = UE::MVVM::ConversionFunctionHelper::FindInputPins(Function.GetDefaultObject());
+	if (ArgumentsResult.Num() == 0)
 	{
-		if (!GetDefault<UMVVMDeveloperProjectSettings>()->IsConversionFunctionAllowed(SelfContext, Function))
+		return false;
+	}
+
+	return GetDefault<UMVVMDeveloperProjectSettings>()->IsConversionFunctionAllowed(WidgetBlueprint, Function);
+}
+
+bool UMVVMBlueprintViewConversionFunction::IsValid(const UBlueprint* SelfContext) const
+{
+	if (ConversionFunction.GetType() == EMVVMBlueprintFunctionReferenceType::Function)
+	{
+		UClass* SelfClass = SelfContext->SkeletonGeneratedClass ? SelfContext->SkeletonGeneratedClass : SelfContext->GeneratedClass;
+		const UFunction* Function = ConversionFunction.GetFunction(SelfContext);
+		if (Function)
 		{
-			return false;
+			return IsValidConversionFunction(SelfContext, Function);
 		}
-
-		return IsValidConversionFunction(SelfContext, Function);
+	}
+	else if (ConversionFunction.GetType() == EMVVMBlueprintFunctionReferenceType::Node)
+	{
+		TSubclassOf<UK2Node> Node = ConversionFunction.GetNode();
+		return Node.Get() && GetDefault<UMVVMDeveloperProjectSettings>()->IsConversionFunctionAllowed(SelfContext, Node);
 	}
 	return false;
 }
@@ -92,21 +110,21 @@ bool UMVVMBlueprintViewConversionFunction::NeedsWrapperGraph(const UBlueprint* S
 
 bool UMVVMBlueprintViewConversionFunction::NeedsWrapperGraphInternal(const UClass* SkeletalSelfContext) const
 {
-	if (ConversionFunction.GetType() != EMVVMBlueprintFunctionReferenceType::Function)
+	if (ConversionFunction.GetType() == EMVVMBlueprintFunctionReferenceType::Function)
 	{
-		return true; // not valid for now
+		const UFunction* Function = ConversionFunction.GetFunction(SkeletalSelfContext);
+		if (ensure(Function))
+		{
+			return SavedPins.Num() > 1 || !UE::MVVM::BindingHelper::IsValidForSimpleRuntimeConversion(Function);
+		}
 	}
-
-
-	const UFunction* Function = ConversionFunction.GetFunction(SkeletalSelfContext);
-	if (ensure(Function))
+	else if (ConversionFunction.GetType() == EMVVMBlueprintFunctionReferenceType::Node)
 	{
-		return SavedPins.Num() > 1 || !UE::MVVM::BindingHelper::IsValidForSimpleRuntimeConversion(Function);
+		return true;
 	}
 	return false;
 }
 
-/** The wrapper Graph is generated on domains and is not saved. */
 bool UMVVMBlueprintViewConversionFunction::IsWrapperGraphTransient() const
 {
 	return bWrapperGraphTransient;
@@ -121,6 +139,10 @@ const UFunction* UMVVMBlueprintViewConversionFunction::GetCompiledFunction(const
 		return CompiledFunction.ResolveMember<UFunction>(const_cast<UClass*>(SelfContext));
 	}
 
+	//NeedsWrapperGraphInternal should return true.
+	check(ConversionFunction.GetType() != EMVVMBlueprintFunctionReferenceType::Node);
+
+	// If simple conversion function.
 	if (ConversionFunction.GetType() == EMVVMBlueprintFunctionReferenceType::Function)
 	{
 		return ConversionFunction.GetFunction(SelfContext);
@@ -152,6 +174,21 @@ void UMVVMBlueprintViewConversionFunction::Reset()
 	SetCachedWrapperGraph(nullptr, nullptr, nullptr);
 }
 
+void UMVVMBlueprintViewConversionFunction::Initialize(UBlueprint* InContext, FName InGraphName, FMVVMBlueprintFunctionReference InFunction)
+{
+	Reset();
+
+	if (ensure(InFunction.IsValid(InContext) && !InGraphName.IsNone()))
+	{
+		ConversionFunction = InFunction;
+		check(GraphName.IsNone()); // the name needs to be set before a GetOrCreateWrapperGraph
+		GraphName = InGraphName;
+		bWrapperGraphTransient = !GetDefault<UMVVMDeveloperProjectSettings>()->bAllowConversionFunctionGeneratedGraphInEditor;
+		GetOrCreateWrapperGraphInternal(InContext);
+		SavePinValues(InContext);
+	}
+}
+
 void UMVVMBlueprintViewConversionFunction::InitializeFromFunction(UBlueprint* InContext, FName InGraphName, const UFunction* InFunction)
 {
 	Reset();
@@ -162,7 +199,7 @@ void UMVVMBlueprintViewConversionFunction::InitializeFromFunction(UBlueprint* In
 		check(GraphName.IsNone()); // the name needs to be set before a GetOrCreateWrapperGraph
 		GraphName = InGraphName;
 		bWrapperGraphTransient = !GetDefault<UMVVMDeveloperProjectSettings>()->bAllowConversionFunctionGeneratedGraphInEditor;
-		GetOrCreateWrapperGraphInternal(InContext, InFunction);
+		GetOrCreateWrapperGraphInternal(InContext);
 		SavePinValues(InContext);
 	}
 }
@@ -298,6 +335,11 @@ UEdGraph* UMVVMBlueprintViewConversionFunction::GetOrCreateIntermediateWrapperGr
 		return CachedWrapperGraph;
 	}
 
+	if (ConversionFunction.GetType() == EMVVMBlueprintFunctionReferenceType::None)
+	{
+		return nullptr;
+	}
+
 	TObjectPtr<UEdGraph>* FoundGraph = !GraphName.IsNone() ? Context.Blueprint->FunctionGraphs.FindByPredicate([GraphName = GetWrapperGraphName()](const UEdGraph* Other) { return Other->GetFName() == GraphName; }) : nullptr;
 	if (FoundGraph)
 	{
@@ -309,14 +351,8 @@ UEdGraph* UMVVMBlueprintViewConversionFunction::GetOrCreateIntermediateWrapperGr
 	}
 	else if (IsValid(Context.Blueprint))
 	{
-		if (ConversionFunction.GetType() != EMVVMBlueprintFunctionReferenceType::Function)
-		{
-			ensure(false); // not supported for now
-			return nullptr;
-		}
-
 		CreateWrapperGraphName();
-		return GetOrCreateWrapperGraphInternal(Context, ConversionFunction.GetFunction(Context.Blueprint));
+		return GetOrCreateWrapperGraphInternal(Context);
 	}
 	return nullptr;
 }
@@ -329,6 +365,11 @@ UEdGraph* UMVVMBlueprintViewConversionFunction::GetOrCreateWrapperGraph(UBluepri
 	{
 		return CachedWrapperGraph;
 	}
+	
+	if (ConversionFunction.GetType() == EMVVMBlueprintFunctionReferenceType::None)
+	{
+		return nullptr;
+	}
 
 	TObjectPtr<UEdGraph>* FoundGraph = Blueprint->FunctionGraphs.FindByPredicate([GraphName = GetWrapperGraphName()](const UEdGraph* Other) { return Other->GetFName() == GraphName; });
 	if (FoundGraph)
@@ -339,14 +380,8 @@ UEdGraph* UMVVMBlueprintViewConversionFunction::GetOrCreateWrapperGraph(UBluepri
 	}
 	else if (IsValid(Blueprint))
 	{
-		if (ConversionFunction.GetType() != EMVVMBlueprintFunctionReferenceType::Function)
-		{
-			ensure(false); // not supported for now
-			return nullptr;
-		}
-
 		CreateWrapperGraphName();
-		return GetOrCreateWrapperGraphInternal(Blueprint, ConversionFunction.GetFunction(Blueprint));
+		return GetOrCreateWrapperGraphInternal(Blueprint);
 	}
 	return nullptr;
 }
@@ -357,18 +392,29 @@ UEdGraphPin* UMVVMBlueprintViewConversionFunction::GetOrCreateGraphPin(UBlueprin
 	return CachedWrapperGraph ? UE::MVVM::ConversionFunctionHelper::FindPin(CachedWrapperGraph, PinId.GetNames()) : nullptr;
 }
 
-UEdGraph* UMVVMBlueprintViewConversionFunction::GetOrCreateWrapperGraphInternal(FKismetCompilerContext& Context, const UFunction* InFunction)
+UEdGraph* UMVVMBlueprintViewConversionFunction::GetOrCreateWrapperGraphInternal(FKismetCompilerContext& Context)
 {
-	return GetOrCreateWrapperGraphInternal(Context.Blueprint, InFunction);
+	return GetOrCreateWrapperGraphInternal(Context.Blueprint);
 }
 
-UEdGraph* UMVVMBlueprintViewConversionFunction::GetOrCreateWrapperGraphInternal(UBlueprint* Blueprint, const UFunction* InFunction)
+UEdGraph* UMVVMBlueprintViewConversionFunction::GetOrCreateWrapperGraphInternal(UBlueprint* Blueprint)
 {
 	check(Blueprint);
-	check(InFunction);
 	bool bConst = true;
-	TPair<UEdGraph*, UK2Node*> Result = UE::MVVM::ConversionFunctionHelper::CreateGraph(Blueprint, GraphName, nullptr, InFunction, bConst, bWrapperGraphTransient);
-	const_cast<UMVVMBlueprintViewConversionFunction*>(this)->SetCachedWrapperGraph(Blueprint, Result.Get<0>(), Result.Get<1>());
+	UE::MVVM::ConversionFunctionHelper::FCreateGraphResult Result;
+	if (ConversionFunction.GetType() == EMVVMBlueprintFunctionReferenceType::Function)
+	{
+		const UFunction* Function = ConversionFunction.GetFunction(Blueprint);
+		check(Function);
+		Result = UE::MVVM::ConversionFunctionHelper::CreateGraph(Blueprint, GraphName, nullptr, Function, bConst, bWrapperGraphTransient);
+	}
+	else if (ConversionFunction.GetType() == EMVVMBlueprintFunctionReferenceType::Node)
+	{
+		TSubclassOf<UK2Node> Node = ConversionFunction.GetNode();
+		check(Node.Get());
+		Result = UE::MVVM::ConversionFunctionHelper::CreateGraph(Blueprint, GraphName, nullptr, Node, bConst, bWrapperGraphTransient, [](UK2Node*){});
+	}
+	const_cast<UMVVMBlueprintViewConversionFunction*>(this)->SetCachedWrapperGraph(Blueprint, Result.NewGraph, Result.WrappedNode);
 	LoadPinValuesInternal(Blueprint);
 	return CachedWrapperGraph;
 }
@@ -386,6 +432,7 @@ void UMVVMBlueprintViewConversionFunction::RemoveWrapperGraph(UBlueprint* Bluepr
 void UMVVMBlueprintViewConversionFunction::SetGraphPin(UBlueprint* Blueprint, const FMVVMBlueprintPinId& PinId, const FMVVMBlueprintPropertyPath& Path)
 {
 	UEdGraphPin* GraphPin = GetOrCreateGraphPin(Blueprint, PinId);
+	check(GraphPin);
 
 	// Set the value and make the blueprint as dirty before creating the pin.
 	//A property may not be created yet and the skeletal needs to be recreated.
@@ -403,10 +450,13 @@ void UMVVMBlueprintViewConversionFunction::SetGraphPin(UBlueprint* Blueprint, co
 
 void UMVVMBlueprintViewConversionFunction::SavePinValues(UBlueprint* Blueprint)
 {
-	SavedPins.Empty();
-	if (CachedWrapperNode)
+	if (!bLoadingPins) // While loading pins value, the node can trigger a notify that would then trigger a save.
 	{
-		SavedPins = FMVVMBlueprintPin::CreateFromNode(Blueprint, CachedWrapperNode);
+		SavedPins.Empty();
+		if (CachedWrapperNode)
+		{
+			SavedPins = FMVVMBlueprintPin::CreateFromNode(Blueprint, CachedWrapperNode);
+		}
 	}
 }
 
@@ -434,6 +484,7 @@ bool UMVVMBlueprintViewConversionFunction::HasOrphanedPin() const
 
 void UMVVMBlueprintViewConversionFunction::LoadPinValuesInternal(UBlueprint* Blueprint)
 {
+	TGuardValue<bool> Tmp (bLoadingPins, true);
 	if (CachedWrapperNode)
 	{
 		TArray<FMVVMBlueprintPin> MissingPins = FMVVMBlueprintPin::CopyAndReturnMissingPins(Blueprint, CachedWrapperNode, SavedPins);
diff --git a/Engine/Plugins/Runtime/ModelViewViewModel/Source/ModelViewViewModelBlueprint/Private/MVVMBlueprintViewEvent.cpp b/Engine/Plugins/Runtime/ModelViewViewModel/Source/ModelViewViewModelBlueprint/Private/MVVMBlueprintViewEvent.cpp
--- a/Engine/Plugins/Runtime/ModelViewViewModel/Source/ModelViewViewModelBlueprint/Private/MVVMBlueprintViewEvent.cpp
+++ b/Engine/Plugins/Runtime/ModelViewViewModel/Source/ModelViewViewModelBlueprint/Private/MVVMBlueprintViewEvent.cpp
@@ -121,11 +121,14 @@ UEdGraphPin* UMVVMBlueprintViewEvent::GetOrCreateGraphPin(const FMVVMBlueprintPi
 
 void UMVVMBlueprintViewEvent::SavePinValues()
 {
-	SavedPins.Empty();
-	if (CachedWrapperNode)
+	if (!bLoadingPins) // While loading pins value, the node can trigger a notify that would then trigger a save.
 	{
-		UWidgetBlueprint* Blueprint = GetWidgetBlueprintInternal();
-		SavedPins = FMVVMBlueprintPin::CreateFromNode(Blueprint, CachedWrapperNode);
+		SavedPins.Empty();
+		if (CachedWrapperNode)
+		{
+			UWidgetBlueprint* Blueprint = GetWidgetBlueprintInternal();
+			SavedPins = FMVVMBlueprintPin::CreateFromNode(Blueprint, CachedWrapperNode);
+		}
 	}
 }
 
@@ -243,7 +246,7 @@ UEdGraph* UMVVMBlueprintViewEvent::CreateWrapperGraphInternal()
 	UWidgetBlueprint* WidgetBlueprint = GetWidgetBlueprintInternal();
 	bool bIsConst = false;
 	bool bTransient = true;
-	TValueOrError<UE::MVVM::ConversionFunctionHelper::FCreateSetterGraphResult, FText> CreateSetterGraphResult = UE::MVVM::ConversionFunctionHelper::CreateSetterGraph(WidgetBlueprint, GraphName, DelegateSignature, DestinationPath, bIsConst, bTransient);
+	TValueOrError<UE::MVVM::ConversionFunctionHelper::FCreateGraphResult, FText> CreateSetterGraphResult = UE::MVVM::ConversionFunctionHelper::CreateSetterGraph(WidgetBlueprint, GraphName, DelegateSignature, DestinationPath, bIsConst, bTransient);
 	if (CreateSetterGraphResult.HasError())
 	{
 		SetCachedWrapperGraphInternal(nullptr, nullptr);
@@ -260,6 +263,7 @@ UEdGraph* UMVVMBlueprintViewEvent::CreateWrapperGraphInternal()
 
 void UMVVMBlueprintViewEvent::LoadPinValuesInternal()
 {
+	TGuardValue<bool> Tmp(bLoadingPins, true);
 	if (CachedWrapperNode)
 	{
 		TArray<FMVVMBlueprintPin> MissingPins = FMVVMBlueprintPin::CopyAndReturnMissingPins(GetWidgetBlueprintInternal(), CachedWrapperNode, SavedPins);
diff --git a/Engine/Plugins/Runtime/ModelViewViewModel/Source/ModelViewViewModelBlueprint/Private/MVVMConversionFunctionGraphSchema.h b/Engine/Plugins/Runtime/ModelViewViewModel/Source/ModelViewViewModelBlueprint/Private/MVVMConversionFunctionGraphSchema.h
new file mode 100644
--- /dev/null
+++ b/Engine/Plugins/Runtime/ModelViewViewModel/Source/ModelViewViewModelBlueprint/Private/MVVMConversionFunctionGraphSchema.h
@@ -0,0 +1,28 @@
+// Copyright Epic Games, Inc. All Rights Reserved.
+
+#pragma once
+
+#include "EdGraphSchema_K2.h"
+
+#include "MVVMConversionFunctionGraphSchema.generated.h"
+
+
+/**
+ *
+ */
+UCLASS()
+class UMVVMConversionFunctionGraphSchema : public UEdGraphSchema_K2
+{
+	GENERATED_BODY()
+
+public:
+	virtual const FPinConnectionResponse CanCreateConnection(const UEdGraphPin* A, const UEdGraphPin* B) const override
+	{
+		FPinConnectionResponse ConnectionResponse = Super::CanCreateConnection(A, B);
+		if (ConnectionResponse.Response == CONNECT_RESPONSE_MAKE_WITH_CONVERSION_NODE || ConnectionResponse.Response == CONNECT_RESPONSE_MAKE_WITH_PROMOTION)
+		{
+			ConnectionResponse.Response = CONNECT_RESPONSE_DISALLOW;
+		}
+		return ConnectionResponse;
+	}
+};
\ No newline at end of file
diff --git a/Engine/Plugins/Runtime/ModelViewViewModel/Source/ModelViewViewModelBlueprint/Private/MVVMDeveloperProjectSettings.cpp b/Engine/Plugins/Runtime/ModelViewViewModel/Source/ModelViewViewModelBlueprint/Private/MVVMDeveloperProjectSettings.cpp
--- a/Engine/Plugins/Runtime/ModelViewViewModel/Source/ModelViewViewModelBlueprint/Private/MVVMDeveloperProjectSettings.cpp
+++ b/Engine/Plugins/Runtime/ModelViewViewModel/Source/ModelViewViewModelBlueprint/Private/MVVMDeveloperProjectSettings.cpp
@@ -157,12 +157,6 @@ bool UMVVMDeveloperProjectSettings::IsFunctionAllowed(const UBlueprint* Generati
 
 bool UMVVMDeveloperProjectSettings::IsConversionFunctionAllowed(const UBlueprint* GeneratingFor, const UFunction* Function) const
 {
-	static FName NAME_ComplexConversionFunction = TEXT("MVVMComplexConversionFunction");
-	if (Function->HasMetaData(NAME_ComplexConversionFunction))
-	{
-		return true;
-	}
-
 	if (ConversionFunctionFilter == EMVVMDeveloperConversionFunctionFilterType::BlueprintActionRegistry)
 	{
 		return IsFunctionAllowed(GeneratingFor, Function->GetOwnerClass(), Function);
@@ -196,6 +190,32 @@ bool UMVVMDeveloperProjectSettings::IsConversionFunctionAllowed(const UBlueprint
 	}
 }
 
+bool UMVVMDeveloperProjectSettings::IsConversionFunctionAllowed(const UBlueprint* Context, const TSubclassOf<UK2Node> Function) const
+{
+	if (ConversionFunctionFilter == EMVVMDeveloperConversionFunctionFilterType::BlueprintActionRegistry)
+	{
+		return !Function.Get()->HasAnyClassFlags(CLASS_Abstract | CLASS_Deprecated | CLASS_NewerVersionExists);
+	}
+	else
+	{
+		check(ConversionFunctionFilter == EMVVMDeveloperConversionFunctionFilterType::AllowedList);
+
+		TStringBuilder<512> FunctionClassPath;
+		Function.Get()->GetPathName(nullptr, FunctionClassPath);
+		TStringBuilder<512> AllowedClassPath;
+		for (const FSoftClassPath& SoftClass : AllowedClassForConversionFunctions)
+		{
+			SoftClass.ToString(AllowedClassPath);
+			if (AllowedClassPath.ToView() == FunctionClassPath.ToView())
+			{
+				return true;
+			}
+			AllowedClassPath.Reset();
+		}
+		return false;
+	}
+}
+
 TArray<const UClass*> UMVVMDeveloperProjectSettings::GetAllowedConversionFunctionClasses() const
 {
 	TArray<const UClass*> Result;
diff --git a/Engine/Plugins/Runtime/ModelViewViewModel/Source/ModelViewViewModelBlueprint/Private/MVVMViewBlueprintCompiler.cpp b/Engine/Plugins/Runtime/ModelViewViewModel/Source/ModelViewViewModelBlueprint/Private/MVVMViewBlueprintCompiler.cpp
--- a/Engine/Plugins/Runtime/ModelViewViewModel/Source/ModelViewViewModelBlueprint/Private/MVVMViewBlueprintCompiler.cpp
+++ b/Engine/Plugins/Runtime/ModelViewViewModel/Source/ModelViewViewModelBlueprint/Private/MVVMViewBlueprintCompiler.cpp
@@ -2228,7 +2228,19 @@ void FMVVMViewBlueprintCompiler::PreCompileBindings(UWidgetBlueprintGeneratedCla
 					if (!GetDefault<UMVVMDeveloperProjectSettings>()->IsConversionFunctionAllowed(Self->WidgetBlueprintCompilerContext.WidgetBlueprint(), WrapperFunction))
 					{
 						Self->AddMessageForBinding(Binding
-							, FText::Format(LOCTEXT("ConversionFunctionNotAllow", "The conversion function {0} is not allowed."), FText::FromName(WrapperFunction->GetFName()))
+							, FText::Format(LOCTEXT("ConversionFunctionNotAllow", "The conversion function {0} is not allowed."), FText::FromName(WrapperFunction ? WrapperFunction->GetFName() : FName()))
+							, EMessageType::Error
+							, FMVVMBlueprintPinId()
+						);
+						return false;
+					}
+				}
+				else if (FunctionOrWrapperFunction.GetType() == EMVVMBlueprintFunctionReferenceType::Node)
+				{
+					if (!GetDefault<UMVVMDeveloperProjectSettings>()->IsConversionFunctionAllowed(Self->WidgetBlueprintCompilerContext.WidgetBlueprint(), FunctionOrWrapperFunction.GetNode()))
+					{
+						Self->AddMessageForBinding(Binding
+							, FText::Format(LOCTEXT("ConversionFunctionNotAllow", "The conversion function {0} is not allowed."), FText::FromName(FunctionOrWrapperFunction.GetNode().Get() ? FunctionOrWrapperFunction.GetNode()->GetFName() : FName()))
 							, EMessageType::Error
 							, FMVVMBlueprintPinId()
 						);
diff --git a/Engine/Plugins/Runtime/ModelViewViewModel/Source/ModelViewViewModelBlueprint/Public/Bindings/MVVMConversionFunctionHelper.h b/Engine/Plugins/Runtime/ModelViewViewModel/Source/ModelViewViewModelBlueprint/Public/Bindings/MVVMConversionFunctionHelper.h
--- a/Engine/Plugins/Runtime/ModelViewViewModel/Source/ModelViewViewModelBlueprint/Public/Bindings/MVVMConversionFunctionHelper.h
+++ b/Engine/Plugins/Runtime/ModelViewViewModel/Source/ModelViewViewModelBlueprint/Public/Bindings/MVVMConversionFunctionHelper.h
@@ -47,7 +47,7 @@ namespace UE::MVVM::ConversionFunctionHelper
 	 */
 	MODELVIEWVIEWMODELBLUEPRINT_API TValueOrError<void, FText> CanCreateSetterGraph(UBlueprint* WidgetBlueprint, const FMVVMBlueprintPropertyPath& PropertyPath);
 
-	struct FCreateSetterGraphResult
+	struct FCreateGraphResult
 	{
 		/** The new graph created. */
 		UEdGraph* NewGraph = nullptr;
@@ -58,13 +58,13 @@ namespace UE::MVVM::ConversionFunctionHelper
 	/**
 	 * Create a graph to set a property/function.
 	 */
-	MODELVIEWVIEWMODELBLUEPRINT_API TValueOrError<FCreateSetterGraphResult, FText> CreateSetterGraph(UBlueprint* WidgetBlueprint, FName GraphName, const UFunction* Signature, const FMVVMBlueprintPropertyPath& PropertyPath, bool bIsConst, bool bTransient);
+	MODELVIEWVIEWMODELBLUEPRINT_API TValueOrError<FCreateGraphResult, FText> CreateSetterGraph(UBlueprint* WidgetBlueprint, FName GraphName, const UFunction* Signature, const FMVVMBlueprintPropertyPath& PropertyPath, bool bIsConst, bool bTransient);
 
 	/** */
-	MODELVIEWVIEWMODELBLUEPRINT_API TPair<UEdGraph*, UK2Node*> CreateGraph(UBlueprint* WidgetBlueprint, FName GraphName, const UFunction* Signature, const UFunction* FunctionToWrap, bool bIsConst, bool bTransient);
+	MODELVIEWVIEWMODELBLUEPRINT_API FCreateGraphResult CreateGraph(UBlueprint* WidgetBlueprint, FName GraphName, const UFunction* Signature, const UFunction* FunctionToWrap, bool bIsConst, bool bTransient);
 	
 	/** */
-	MODELVIEWVIEWMODELBLUEPRINT_API TPair<UEdGraph*, UK2Node*> CreateGraph(UBlueprint* WidgetBlueprint, FName GraphName, const UFunction* Signature, const TSubclassOf<UK2Node> Node, bool bIsConst, bool bTransient, TFunctionRef<void(UK2Node*)> InitNodeCallback);
+	MODELVIEWVIEWMODELBLUEPRINT_API FCreateGraphResult CreateGraph(UBlueprint* WidgetBlueprint, FName GraphName, const UFunction* Signature, const TSubclassOf<UK2Node> Node, bool bIsConst, bool bTransient, TFunctionRef<void(UK2Node*)> InitNodeCallback);
 
 	/** Find the main conversion function node from the given graph. */
 	MODELVIEWVIEWMODELBLUEPRINT_API UK2Node* GetWrapperNode(const UEdGraph* Graph);
@@ -74,6 +74,12 @@ namespace UE::MVVM::ConversionFunctionHelper
 
 	/** Find the conversion function node from the given graph. */
 	MODELVIEWVIEWMODELBLUEPRINT_API TArray<FName> FindPinId(const UEdGraphPin* GraphPin);
+
+	/** Return the pin used for arguments. */
+	MODELVIEWVIEWMODELBLUEPRINT_API TArray<UEdGraphPin*> FindInputPins(const UK2Node* Node);
+	
+	/** Return the pin used as the return value. */
+	MODELVIEWVIEWMODELBLUEPRINT_API UEdGraphPin* FindOutputPin(const UK2Node* Node);
 } //namespace
 
 #if UE_ENABLE_INCLUDE_ORDER_DEPRECATED_IN_5_2
diff --git a/Engine/Plugins/Runtime/ModelViewViewModel/Source/ModelViewViewModelBlueprint/Public/MVVMBlueprintFunctionReference.h b/Engine/Plugins/Runtime/ModelViewViewModel/Source/ModelViewViewModelBlueprint/Public/MVVMBlueprintFunctionReference.h
--- a/Engine/Plugins/Runtime/ModelViewViewModel/Source/ModelViewViewModelBlueprint/Public/MVVMBlueprintFunctionReference.h
+++ b/Engine/Plugins/Runtime/ModelViewViewModel/Source/ModelViewViewModelBlueprint/Public/MVVMBlueprintFunctionReference.h
@@ -53,6 +53,11 @@ public:
 	/** return a name of the reference. */
 	MODELVIEWVIEWMODELBLUEPRINT_API FName GetName() const;
 
+	/** The reference is a valid UFunction or a valid Node. */
+	MODELVIEWVIEWMODELBLUEPRINT_API bool IsValid(const UBlueprint* SelfContext) const;
+	/** The reference is a valid UFunction or a valid Node. */
+	MODELVIEWVIEWMODELBLUEPRINT_API bool IsValid(const UClass* SelfContext) const;
+
 	MODELVIEWVIEWMODELBLUEPRINT_API bool operator==(const FMVVMBlueprintFunctionReference& Other) const;
 
 	friend uint32 GetTypeHash(const FMVVMBlueprintFunctionReference& Value)
diff --git a/Engine/Plugins/Runtime/ModelViewViewModel/Source/ModelViewViewModelBlueprint/Public/MVVMBlueprintViewConversionFunction.h b/Engine/Plugins/Runtime/ModelViewViewModel/Source/ModelViewViewModelBlueprint/Public/MVVMBlueprintViewConversionFunction.h
--- a/Engine/Plugins/Runtime/ModelViewViewModel/Source/ModelViewViewModelBlueprint/Public/MVVMBlueprintViewConversionFunction.h
+++ b/Engine/Plugins/Runtime/ModelViewViewModel/Source/ModelViewViewModelBlueprint/Public/MVVMBlueprintViewConversionFunction.h
@@ -26,6 +26,7 @@ class MODELVIEWVIEWMODELBLUEPRINT_API UMVVMBlueprintViewConversionFunction : pub
 
 public:
 	static bool IsValidConversionFunction(const UBlueprint* WidgetBlueprint, const UFunction* Function);
+	static bool IsValidConversionNode(const UBlueprint* WidgetBlueprint, const TSubclassOf<UK2Node> Function);
 
 public:
 	/** @return the conversion function uses at runtime. The wrapper function if complex or GetFunction is simple. */
@@ -41,6 +42,9 @@ public:
 	/** @return the conversion function. */
 	FMVVMBlueprintFunctionReference GetConversionFunction() const;
 
+	/** Set the function. Generate a Graph. */
+	void Initialize(UBlueprint* SelfContext, FName GraphName, FMVVMBlueprintFunctionReference Function);
+	
 	/** Set the function. Generate a Graph. */
 	void InitializeFromFunction(UBlueprint* SelfContext, FName GraphName, const UFunction* Function);
 
@@ -121,8 +125,8 @@ private:
 	void HandleGraphChanged(const FEdGraphEditAction& Action, TWeakObjectPtr<UBlueprint> Context);
 	void HandleUserDefinedPinRenamed(UK2Node* InNode, FName OldPinName, FName NewPinName, TWeakObjectPtr<UBlueprint> WeakBlueprint);
 	void SetCachedWrapperGraph(UBlueprint* Blueprint, UEdGraph* CachedGraph, UK2Node* CachedNode);
-	UEdGraph* GetOrCreateWrapperGraphInternal(FKismetCompilerContext& Context, const UFunction* Function);
-	UEdGraph* GetOrCreateWrapperGraphInternal(UBlueprint* Blueprint, const UFunction* Function);
+	UEdGraph* GetOrCreateWrapperGraphInternal(FKismetCompilerContext& Context);
+	UEdGraph* GetOrCreateWrapperGraphInternal(UBlueprint* Blueprint);
 	bool NeedsWrapperGraphInternal(const UClass* SkeletalSelfContext) const;
 	void LoadPinValuesInternal(UBlueprint* Blueprint);
 	void CreateWrapperGraphName();
@@ -158,7 +162,7 @@ private:
 
 	FDelegateHandle OnGraphChangedHandle;
 	FDelegateHandle OnUserDefinedPinRenamedHandle;
-
+	bool bLoadingPins = false;
 
 	UPROPERTY()
 	FMemberReference FunctionReference_DEPRECATED;
diff --git a/Engine/Plugins/Runtime/ModelViewViewModel/Source/ModelViewViewModelBlueprint/Public/MVVMBlueprintViewEvent.h b/Engine/Plugins/Runtime/ModelViewViewModel/Source/ModelViewViewModelBlueprint/Public/MVVMBlueprintViewEvent.h
--- a/Engine/Plugins/Runtime/ModelViewViewModel/Source/ModelViewViewModelBlueprint/Public/MVVMBlueprintViewEvent.h
+++ b/Engine/Plugins/Runtime/ModelViewViewModel/Source/ModelViewViewModelBlueprint/Public/MVVMBlueprintViewEvent.h
@@ -151,6 +151,7 @@ private:
 	FName GraphName;
 
 	mutable TArray<FMessage> Messages;
+	bool bLoadingPins = false;
 
 	UPROPERTY(Transient, DuplicateTransient)
 	mutable TObjectPtr<UEdGraph> CachedWrapperGraph;
diff --git a/Engine/Plugins/Runtime/ModelViewViewModel/Source/ModelViewViewModelBlueprint/Public/MVVMDeveloperProjectSettings.h b/Engine/Plugins/Runtime/ModelViewViewModel/Source/ModelViewViewModelBlueprint/Public/MVVMDeveloperProjectSettings.h
--- a/Engine/Plugins/Runtime/ModelViewViewModel/Source/ModelViewViewModelBlueprint/Public/MVVMDeveloperProjectSettings.h
+++ b/Engine/Plugins/Runtime/ModelViewViewModel/Source/ModelViewViewModelBlueprint/Public/MVVMDeveloperProjectSettings.h
@@ -6,6 +6,7 @@
 
 #include "MVVMDeveloperProjectSettings.generated.h"
 
+class UK2Node;
 enum class EMVVMBlueprintViewModelContextCreationType : uint8;
 enum class EMVVMExecutionMode : uint8;
 
@@ -74,6 +75,7 @@ public:
 	bool IsPropertyAllowed(const UBlueprint* Context, const UStruct* ObjectStruct, const FProperty* Property) const;
 	bool IsFunctionAllowed(const UBlueprint* Context, const UClass* ObjectClass, const UFunction* Function) const;
 	bool IsConversionFunctionAllowed(const UBlueprint* Context, const UFunction* Function) const;
+	bool IsConversionFunctionAllowed(const UBlueprint* Context, const TSubclassOf<UK2Node> Function) const;
 
 	bool IsExecutionModeAllowed(EMVVMExecutionMode ExecutionMode) const
 	{
diff --git a/Engine/Plugins/Runtime/ModelViewViewModel/Source/ModelViewViewModelEditor/Private/MVVMEditorSubsystem.cpp b/Engine/Plugins/Runtime/ModelViewViewModel/Source/ModelViewViewModelEditor/Private/MVVMEditorSubsystem.cpp
--- a/Engine/Plugins/Runtime/ModelViewViewModel/Source/ModelViewViewModelEditor/Private/MVVMEditorSubsystem.cpp
+++ b/Engine/Plugins/Runtime/ModelViewViewModel/Source/ModelViewViewModelEditor/Private/MVVMEditorSubsystem.cpp
@@ -361,29 +361,28 @@ void UMVVMEditorSubsystem::SetSourceToDestinationConversionFunction(UWidgetBluep
 	SetSourceToDestinationConversionFunction(WidgetBlueprint, Binding, FMVVMBlueprintFunctionReference(WidgetBlueprint, NewConversionFunction));
 }
 
-void UMVVMEditorSubsystem::SetSourceToDestinationConversionFunction(UWidgetBlueprint* WidgetBlueprint, FMVVMBlueprintViewBinding& Binding, const FMVVMBlueprintFunctionReference& NewConversionFunction)
+void UMVVMEditorSubsystem::SetSourceToDestinationConversionFunction(UWidgetBlueprint* WidgetBlueprint, FMVVMBlueprintViewBinding& Binding, FMVVMBlueprintFunctionReference NewConversionFunction)
 {
 	if (UMVVMBlueprintView* View = GetView(WidgetBlueprint))
 	{
-		const UFunction* NewFunction = nullptr;
 		if (NewConversionFunction.GetType() == EMVVMBlueprintFunctionReferenceType::Function)
 		{
-			NewFunction = NewConversionFunction.GetFunction(WidgetBlueprint);
+			const UFunction* NewFunction = NewConversionFunction.GetFunction(WidgetBlueprint);
 			if (!IsValidConversionFunction(WidgetBlueprint, NewFunction, Binding.SourcePath, Binding.DestinationPath))
 			{
-				NewFunction = nullptr;
+				NewConversionFunction = FMVVMBlueprintFunctionReference();
 			}
 		}
-		const UFunction* CurrentFunction = nullptr;
-		if (Binding.Conversion.SourceToDestinationConversion != nullptr && Binding.Conversion.SourceToDestinationConversion->GetConversionFunction().GetType() == EMVVMBlueprintFunctionReferenceType::Function)
+		else if (NewConversionFunction.GetType() == EMVVMBlueprintFunctionReferenceType::Node)
 		{
-			CurrentFunction = Binding.Conversion.SourceToDestinationConversion->GetConversionFunction().GetFunction(WidgetBlueprint);
+			if (NewConversionFunction.GetNode().Get() == nullptr)
+			{
+				NewConversionFunction = FMVVMBlueprintFunctionReference();
+			}
 		}
 
-		bool bBothNull = Binding.Conversion.SourceToDestinationConversion == nullptr && NewConversionFunction.GetType() == EMVVMBlueprintFunctionReferenceType::None;
-		bool bBothSameValue = Binding.Conversion.SourceToDestinationConversion != nullptr && Binding.Conversion.SourceToDestinationConversion->GetConversionFunction() == NewConversionFunction;
-		bool bBothSameFunctionValue = CurrentFunction == NewFunction;
-		if (!bBothSameValue || !bBothNull || !NewFunction)
+		const FMVVMBlueprintFunctionReference PreviousConversionFunction = Binding.Conversion.SourceToDestinationConversion != nullptr ? Binding.Conversion.SourceToDestinationConversion->GetConversionFunction() : FMVVMBlueprintFunctionReference();
+		if (PreviousConversionFunction != NewConversionFunction)
 		{
 			FScopedTransaction Transaction(LOCTEXT("SetConversionFunction", "Set Conversion Function"));
 
@@ -398,18 +397,11 @@ void UMVVMEditorSubsystem::SetSourceToDestinationConversionFunction(UWidgetBluep
 			}
 			Binding.SourcePath = FMVVMBlueprintPropertyPath();
 
-			if (NewConversionFunction.GetType() == EMVVMBlueprintFunctionReferenceType::Function)
-			{
-				if (NewFunction != nullptr)
-				{
-					Binding.Conversion.SourceToDestinationConversion = NewObject<UMVVMBlueprintViewConversionFunction>(WidgetBlueprint);
-					FName GraphName = UE::MVVM::ConversionFunctionHelper::CreateWrapperName(Binding, true);
-					Binding.Conversion.SourceToDestinationConversion->InitializeFromFunction(WidgetBlueprint, GraphName, NewFunction);
-				}
-			}
-			else
+			if (NewConversionFunction.GetType() != EMVVMBlueprintFunctionReferenceType::None)
 			{
-				check(false); // not supported yet
+				Binding.Conversion.SourceToDestinationConversion = NewObject<UMVVMBlueprintViewConversionFunction>(WidgetBlueprint);
+				FName GraphName = UE::MVVM::ConversionFunctionHelper::CreateWrapperName(Binding, true);
+				Binding.Conversion.SourceToDestinationConversion->Initialize(WidgetBlueprint, GraphName, NewConversionFunction);
 			}
 
 			UE::MVVM::Private::OnBindingPostEditChange(View, GET_MEMBER_NAME_CHECKED(FMVVMBlueprintViewBinding, Conversion));
@@ -423,29 +415,28 @@ void UMVVMEditorSubsystem::SetDestinationToSourceConversionFunction(UWidgetBluep
 	SetDestinationToSourceConversionFunction(WidgetBlueprint, Binding, FMVVMBlueprintFunctionReference(WidgetBlueprint, NewConversionFunction));
 }
 
-void UMVVMEditorSubsystem::SetDestinationToSourceConversionFunction(UWidgetBlueprint* WidgetBlueprint, FMVVMBlueprintViewBinding& Binding, const FMVVMBlueprintFunctionReference& NewConversionFunction)
+void UMVVMEditorSubsystem::SetDestinationToSourceConversionFunction(UWidgetBlueprint* WidgetBlueprint, FMVVMBlueprintViewBinding& Binding, FMVVMBlueprintFunctionReference NewConversionFunction)
 {
 	if (UMVVMBlueprintView* View = GetView(WidgetBlueprint))
 	{
-		const UFunction* NewFunction = nullptr;
 		if (NewConversionFunction.GetType() == EMVVMBlueprintFunctionReferenceType::Function)
 		{
-			NewFunction = NewConversionFunction.GetFunction(WidgetBlueprint);
-			if (!IsValidConversionFunction(WidgetBlueprint, NewFunction, Binding.SourcePath, Binding.DestinationPath))
+			const UFunction* NewFunction = NewConversionFunction.GetFunction(WidgetBlueprint);
+			if (!IsValidConversionFunction(WidgetBlueprint, NewFunction, Binding.DestinationPath, Binding.SourcePath))
 			{
-				NewFunction = nullptr;
+				NewConversionFunction = FMVVMBlueprintFunctionReference();
 			}
 		}
-		const UFunction* CurrentFunction = nullptr;
-		if (Binding.Conversion.DestinationToSourceConversion != nullptr && Binding.Conversion.DestinationToSourceConversion->GetConversionFunction().GetType() == EMVVMBlueprintFunctionReferenceType::Function)
+		else if (NewConversionFunction.GetType() == EMVVMBlueprintFunctionReferenceType::Node)
 		{
-			CurrentFunction = Binding.Conversion.DestinationToSourceConversion->GetConversionFunction().GetFunction(WidgetBlueprint);
+			if (NewConversionFunction.GetNode().Get() == nullptr)
+			{
+				NewConversionFunction = FMVVMBlueprintFunctionReference();
+			}
 		}
 
-		bool bBothNull = Binding.Conversion.DestinationToSourceConversion == nullptr && NewConversionFunction.GetType() == EMVVMBlueprintFunctionReferenceType::None;
-		bool bBothSameValue = Binding.Conversion.DestinationToSourceConversion != nullptr && Binding.Conversion.DestinationToSourceConversion->GetConversionFunction() == NewConversionFunction;
-		bool bBothSameFunctionValue = CurrentFunction == NewFunction;
-		if (!bBothSameValue || !bBothNull || !NewFunction)
+		const FMVVMBlueprintFunctionReference PreviousConversionFunction = Binding.Conversion.DestinationToSourceConversion != nullptr ? Binding.Conversion.DestinationToSourceConversion->GetConversionFunction() : FMVVMBlueprintFunctionReference();
+		if (PreviousConversionFunction != NewConversionFunction)
 		{
 			FScopedTransaction Transaction(LOCTEXT("SetConversionFunction", "Set Conversion Function"));
 
@@ -460,18 +451,11 @@ void UMVVMEditorSubsystem::SetDestinationToSourceConversionFunction(UWidgetBluep
 			}
 			Binding.DestinationPath = FMVVMBlueprintPropertyPath();
 
-			if (NewConversionFunction.GetType() == EMVVMBlueprintFunctionReferenceType::Function)
+			if (NewConversionFunction.GetType() != EMVVMBlueprintFunctionReferenceType::None)
 			{
-				if (NewFunction != nullptr)
-				{
-					Binding.Conversion.DestinationToSourceConversion = NewObject<UMVVMBlueprintViewConversionFunction>(WidgetBlueprint);
-					FName GraphName = UE::MVVM::ConversionFunctionHelper::CreateWrapperName(Binding, true);
-					Binding.Conversion.DestinationToSourceConversion->InitializeFromFunction(WidgetBlueprint, GraphName, NewFunction);
-				}
-			}
-			else
-			{
-				check(false); // not supported yet
+				Binding.Conversion.DestinationToSourceConversion = NewObject<UMVVMBlueprintViewConversionFunction>(WidgetBlueprint);
+				FName GraphName = UE::MVVM::ConversionFunctionHelper::CreateWrapperName(Binding, true);
+				Binding.Conversion.DestinationToSourceConversion->Initialize(WidgetBlueprint, GraphName, NewConversionFunction);
 			}
 
 			UE::MVVM::Private::OnBindingPostEditChange(View, GET_MEMBER_NAME_CHECKED(FMVVMBlueprintViewBinding, Conversion));
@@ -802,6 +786,98 @@ bool UMVVMEditorSubsystem::IsValidConversionFunction(const UWidgetBlueprint* Wid
 	return true;
 }
 
+bool UMVVMEditorSubsystem::IsValidConversionNode(const UWidgetBlueprint* WidgetBlueprint, const TSubclassOf<UK2Node> Function, const FMVVMBlueprintPropertyPath& Source, const FMVVMBlueprintPropertyPath& Destination) const
+{
+	if (WidgetBlueprint == nullptr || Function.Get() == nullptr)
+	{
+		return false;
+	}
+
+	if (!UMVVMBlueprintViewConversionFunction::IsValidConversionNode(WidgetBlueprint, Function))
+	{
+		return false;
+	}
+
+	UClass* CallingContext = WidgetBlueprint->GeneratedClass ? WidgetBlueprint->GeneratedClass : WidgetBlueprint->ParentClass;
+	{
+		const FProperty* SourceProperty = nullptr;
+		TArray<UE::MVVM::FMVVMConstFieldVariant> SourceFields = Source.GetFields(CallingContext);
+		if (SourceFields.Num() > 0)
+		{
+			SourceProperty = SourceFields.Last().IsProperty() ? SourceFields.Last().GetProperty() : UE::MVVM::BindingHelper::GetReturnProperty(SourceFields.Last().GetFunction());
+		}
+
+		// check that at least one source -> argument binding is compatible
+		if (SourceProperty)
+		{
+			FEdGraphPinType SourcePinType;
+			if (!GetDefault<UEdGraphSchema_K2>()->ConvertPropertyToPinType(SourceProperty, SourcePinType))
+			{
+				return false;
+			}
+
+			TArray<UEdGraphPin*> InputPins = UE::MVVM::ConversionFunctionHelper::FindInputPins(Function.GetDefaultObject());
+			if (InputPins.Num() == 0)
+			{
+				return false;
+			}
+
+			bool bAnyCompatible = false;
+			for (const UEdGraphPin* ArgumentPin : InputPins)
+			{
+				bool bIgnoreArray = true;
+				const bool bTypesMatch = GetDefault<UEdGraphSchema_K2>()->ArePinTypesCompatible(SourcePinType, ArgumentPin->PinType, CallingContext, bIgnoreArray);
+				if (bTypesMatch)
+				{
+					bAnyCompatible = true;
+					break;
+				}
+			}
+			if (!bAnyCompatible)
+			{
+				return false;
+			}
+		}
+	}
+
+	{
+		const FProperty* DestinationProperty = nullptr;
+		TArray<UE::MVVM::FMVVMConstFieldVariant> DestFields = Destination.GetFields(CallingContext);
+		if (DestFields.Num() > 0)
+		{
+			if (!DestFields.Last().IsEmpty())
+			{
+				DestinationProperty = DestFields.Last().IsProperty() ? DestFields.Last().GetProperty() : UE::MVVM::BindingHelper::GetFirstArgumentProperty(DestFields.Last().GetFunction());
+			}
+		}
+
+		if (DestinationProperty)
+		{
+			FEdGraphPinType DestinationPinType;
+			if (!GetDefault<UEdGraphSchema_K2>()->ConvertPropertyToPinType(DestinationProperty, DestinationPinType))
+			{
+				return false;
+			}
+
+			UEdGraphPin* ReturnPin = UE::MVVM::ConversionFunctionHelper::FindOutputPin(Function.GetDefaultObject());
+			if (ReturnPin == nullptr)
+			{
+				return false;
+			}
+
+			// check that the return -> dest is valid
+			bool bIgnoreArray = true;
+			const bool bTypesMatch = GetDefault<UEdGraphSchema_K2>()->ArePinTypesCompatible(ReturnPin->PinType, DestinationPinType, CallingContext, bIgnoreArray);
+			if (!bTypesMatch)
+			{
+				return false;
+			}
+		}
+	}
+
+	return true;
+}
+
 bool UMVVMEditorSubsystem::IsSimpleConversionFunctionA(const UFunction* Function) const
 {
 	TValueOrError<const FProperty*, FText> ReturnResult = UE::MVVM::BindingHelper::TryGetReturnTypeForConversionFunction(Function);
@@ -894,19 +970,26 @@ TArray<UFunction*> UMVVMEditorSubsystem::GetAvailableConversionFunctions(const U
 		Classes.Add(WidgetBlueprint->GeneratedClass);
 		for (const UClass* Class : Classes)
 		{
-			for (TFieldIterator<UFunction> FunctionIt(Class, EFieldIteratorFlags::ExcludeSuper); FunctionIt; ++FunctionIt)
+			if (Class->IsChildOf(UK2Node::StaticClass()))
 			{
-				UFunction* Function = *FunctionIt;
-				if (IsInheritedBlueprintFunction(Function))
-				{
-					continue;
-				}
-				if (!Function->HasAllFunctionFlags(FUNC_BlueprintCallable))
+
+			}
+			else
+			{
+				for (TFieldIterator<UFunction> FunctionIt(Class, EFieldIteratorFlags::ExcludeSuper); FunctionIt; ++FunctionIt)
 				{
-					continue;
-				}
+					UFunction* Function = *FunctionIt;
+					if (IsInheritedBlueprintFunction(Function)) // remove override virtual function
+					{
+						continue;
+					}
+					if (!Function->HasAllFunctionFlags(FUNC_BlueprintCallable))
+					{
+						continue;
+					}
 
-				AddFunction(Function);
+					AddFunction(Function);
+				}
 			}
 		}
 	}
diff --git a/Engine/Plugins/Runtime/ModelViewViewModel/Source/ModelViewViewModelEditor/Public/MVVMEditorSubsystem.h b/Engine/Plugins/Runtime/ModelViewViewModel/Source/ModelViewViewModelEditor/Public/MVVMEditorSubsystem.h
--- a/Engine/Plugins/Runtime/ModelViewViewModel/Source/ModelViewViewModelEditor/Public/MVVMEditorSubsystem.h
+++ b/Engine/Plugins/Runtime/ModelViewViewModel/Source/ModelViewViewModelEditor/Public/MVVMEditorSubsystem.h
@@ -20,6 +20,7 @@ struct FMVVMBlueprintViewBinding;
 template <typename T> class TSubclassOf;
 
 class UEdGraph;
+class UK2Node;
 class UK2Node_CallFunction;
 class UMVVMBlueprintViewEvent;
 class UWidgetBlueprint;
@@ -72,10 +73,10 @@ public:
 
 	UE_DEPRECATED(5.4, "SetSourceToDestinationConversionFunction with a UFunction is deprecated.")
 	void SetSourceToDestinationConversionFunction(UWidgetBlueprint* WidgetBlueprint, FMVVMBlueprintViewBinding& Binding, const UFunction* ConversionFunction);
-	void SetSourceToDestinationConversionFunction(UWidgetBlueprint* WidgetBlueprint, FMVVMBlueprintViewBinding& Binding, const FMVVMBlueprintFunctionReference& ConversionFunction);
+	void SetSourceToDestinationConversionFunction(UWidgetBlueprint* WidgetBlueprint, FMVVMBlueprintViewBinding& Binding, FMVVMBlueprintFunctionReference ConversionFunction);
 	UE_DEPRECATED(5.4, "SetDestinationToSourceConversionFunction  with a UFunction is deprecated.")
 	void SetDestinationToSourceConversionFunction(UWidgetBlueprint* WidgetBlueprint, FMVVMBlueprintViewBinding& Binding, const UFunction* ConversionFunction);
-	void SetDestinationToSourceConversionFunction(UWidgetBlueprint* WidgetBlueprint, FMVVMBlueprintViewBinding& Binding, const FMVVMBlueprintFunctionReference& ConversionFunction);
+	void SetDestinationToSourceConversionFunction(UWidgetBlueprint* WidgetBlueprint, FMVVMBlueprintViewBinding& Binding, FMVVMBlueprintFunctionReference ConversionFunction);
 	void SetDestinationPathForBinding(UWidgetBlueprint* WidgetBlueprint, FMVVMBlueprintViewBinding& Binding, FMVVMBlueprintPropertyPath Field);
 	void SetSourcePathForBinding(UWidgetBlueprint* WidgetBlueprint, FMVVMBlueprintViewBinding& Binding, FMVVMBlueprintPropertyPath Field);
 	void OverrideExecutionModeForBinding(UWidgetBlueprint* WidgetBlueprint, FMVVMBlueprintViewBinding& Binding, EMVVMExecutionMode Mode);
@@ -91,7 +92,8 @@ public:
 	void SetCompileForEvent(UMVVMBlueprintViewEvent* Event, bool bCompile);
 
 	UFUNCTION(BlueprintCallable, Category = "Viewmodel")
-	bool IsValidConversionFunction(const UWidgetBlueprint* WidgeteBlueprint, const UFunction* Function, const FMVVMBlueprintPropertyPath& Source, const FMVVMBlueprintPropertyPath& Destination) const;
+	bool IsValidConversionFunction(const UWidgetBlueprint* WidgetBlueprint, const UFunction* Function, const FMVVMBlueprintPropertyPath& Source, const FMVVMBlueprintPropertyPath& Destination) const;
+	bool IsValidConversionNode(const UWidgetBlueprint* WidgetBlueprint, const TSubclassOf<UK2Node> Function, const FMVVMBlueprintPropertyPath& Source, const FMVVMBlueprintPropertyPath& Destination) const;
 
 	UFUNCTION(BlueprintCallable, Category = "Viewmodel")
 	bool IsSimpleConversionFunctionA(const UFunction* Function) const;
