MVVM: allow viewmodel values to be edited on nested userwidgets.
Expose viewmodels whose creation type is "Create Instance" in the details panel of nested userwidgets and allow editing them independently.

diff --git a/Engine/Plugins/Runtime/ModelViewViewModel/Source/ModelViewViewModel/Private/View/MVVMViewClass.cpp b/Engine/Plugins/Runtime/ModelViewViewModel/Source/ModelViewViewModel/Private/View/MVVMViewClass.cpp
--- a/Engine/Plugins/Runtime/ModelViewViewModel/Source/ModelViewViewModel/Private/View/MVVMViewClass.cpp
+++ b/Engine/Plugins/Runtime/ModelViewViewModel/Source/ModelViewViewModel/Private/View/MVVMViewClass.cpp
@@ -38,11 +38,41 @@ UObject* FMVVMViewClass_Source::GetOrCreateInstance(const UMVVMViewClass* InView
 	UObject* Result = nullptr;
 	const bool bOptional = IsOptional();
 
+	auto GetInstanceObjectFromProperty = [&]() -> bool
+	{
+		TValueOrError<UE::MVVM::FFieldContext, void> FieldPathResult = InViewClass->GetBindingLibrary().EvaluateFieldPath(InUserWidget, FieldPath);
+		if (FieldPathResult.HasValue())
+		{
+			TValueOrError<UObject*, void> ObjectResult = UE::MVVM::FieldPathHelper::EvaluateObjectProperty(FieldPathResult.GetValue());
+			if (ObjectResult.HasValue() && ObjectResult.GetValue() != nullptr)
+			{
+				Result = ObjectResult.GetValue();
+			}
+		}
+
+		return Result != nullptr || bOptional;
+	};
+
 	if ((Flags & (uint16)EFlags::TypeCreateInstance) != 0)
 	{
 		if (ExpectedSourceType.Get() != nullptr)
 		{
-			Result = NewObject<UObject>(InUserWidget, ExpectedSourceType.Get(), NAME_None, RF_Transient);
+			// We fetch the instance that was exposed through the editor if it exists, otherwise we create a fresh one.
+			if ((Flags & (uint16)EFlags::IsViewModelInstanceExposed) != 0)
+			{
+				if (ensure(FieldPath.IsValid()))
+				{
+					if (!GetInstanceObjectFromProperty())
+					{
+						UE::MVVM::FMessageLog Log(InUserWidget);
+						Log.Error(FText::Format(LOCTEXT("CreateInstanceInvalidStaticInstance", "The source '{0}' was not found for the static view model at initialization."), FText::FromName(PropertyName)));
+					}
+				}
+			}
+			else
+			{
+				Result = NewObject<UObject>(InUserWidget, ExpectedSourceType.Get(), NAME_None, RF_Transient);
+			}
 		}
 		else if (!bOptional)
 		{
@@ -106,20 +136,10 @@ UObject* FMVVMViewClass_Source::GetOrCreateInstance(const UMVVMViewClass* InView
 	}
 	else if (FieldPath.IsValid())
 	{
-		TValueOrError<UE::MVVM::FFieldContext, void> FieldPathResult = InViewClass->GetBindingLibrary().EvaluateFieldPath(InUserWidget, FieldPath);
-		if (FieldPathResult.HasValue())
-		{
-			TValueOrError<UObject*, void> ObjectResult = UE::MVVM::FieldPathHelper::EvaluateObjectProperty(FieldPathResult.GetValue());
-			if (ObjectResult.HasValue() && ObjectResult.GetValue() != nullptr)
-			{
-				Result = ObjectResult.GetValue();
-			}
-		}
-
-		if (Result == nullptr && !bOptional)
+		if (!GetInstanceObjectFromProperty())
 		{
 			UE::MVVM::FMessageLog Log(InUserWidget);
-			Log.Error(FText::Format(LOCTEXT("CreateInstanceInvalidBiding", "The source '{0}' was evaluated to be invalid at initialization."), FText::FromName(PropertyName)));
+			Log.Error(FText::Format(LOCTEXT("CreateInstanceInvalidBinding", "The source '{0}' was evaluated to be invalid at initialization."), FText::FromName(PropertyName)));
 		}
 	}
 
@@ -232,9 +252,13 @@ FString FMVVMViewClass_Source::ToString(const UMVVMViewClass* ViewClass, FToStri
 	if (Args.bAddCreationMode)
 	{
 		StringBuilder << TEXT("\n    CreationType: ");
-		if ((Flags & (uint8)EFlags::TypeCreateInstance) != 0)
+		if ((Flags & (uint16)EFlags::TypeCreateInstance) != 0)
 		{
-			StringBuilder << TEXT("CreateInsance");
+			StringBuilder << TEXT("CreateInstance");
+			if ((Flags & (uint16)EFlags::IsViewModelInstanceExposed) != 0)
+			{
+				StringBuilder << TEXT(", InstanceExposed");
+			}
 		}
 		else if (Resolver)
 		{
diff --git a/Engine/Plugins/Runtime/ModelViewViewModel/Source/ModelViewViewModel/Public/View/MVVMViewClass.h b/Engine/Plugins/Runtime/ModelViewViewModel/Source/ModelViewViewModel/Public/View/MVVMViewClass.h
--- a/Engine/Plugins/Runtime/ModelViewViewModel/Source/ModelViewViewModel/Public/View/MVVMViewClass.h
+++ b/Engine/Plugins/Runtime/ModelViewViewModel/Source/ModelViewViewModel/Public/View/MVVMViewClass.h
@@ -397,6 +397,7 @@ private:
 		SelfReference = 1 << 7,
 		HasTickBindings = 1 << 8,
 		IsViewModel = 1 << 9,
+		IsViewModelInstanceExposed = 1 << 10,
 	};
 
 	UPROPERTY(VisibleAnywhere, Category = "View")
diff --git a/Engine/Plugins/Runtime/ModelViewViewModel/Source/ModelViewViewModelBlueprint/Private/MVVMBlueprintViewModelContext.cpp b/Engine/Plugins/Runtime/ModelViewViewModel/Source/ModelViewViewModelBlueprint/Private/MVVMBlueprintViewModelContext.cpp
--- a/Engine/Plugins/Runtime/ModelViewViewModel/Source/ModelViewViewModelBlueprint/Private/MVVMBlueprintViewModelContext.cpp
+++ b/Engine/Plugins/Runtime/ModelViewViewModel/Source/ModelViewViewModelBlueprint/Private/MVVMBlueprintViewModelContext.cpp
@@ -19,6 +19,7 @@ FMVVMBlueprintViewModelContext::FMVVMBlueprintViewModelContext(const UClass* InC
 			ViewModelContextId = FGuid::NewGuid();
 			NotifyFieldValueClass = const_cast<UClass*>(InClass);
 			ViewModelName = InViewModelName;
+			bExposeInstanceInEditor = GetDefault<UMVVMDeveloperProjectSettings>()->bExposeViewModelInstanceInEditor;
 		}
 	}
 }
diff --git a/Engine/Plugins/Runtime/ModelViewViewModel/Source/ModelViewViewModelBlueprint/Private/MVVMViewBlueprintCompiler.cpp b/Engine/Plugins/Runtime/ModelViewViewModel/Source/ModelViewViewModelBlueprint/Private/MVVMViewBlueprintCompiler.cpp
--- a/Engine/Plugins/Runtime/ModelViewViewModel/Source/ModelViewViewModelBlueprint/Private/MVVMViewBlueprintCompiler.cpp
+++ b/Engine/Plugins/Runtime/ModelViewViewModel/Source/ModelViewViewModelBlueprint/Private/MVVMViewBlueprintCompiler.cpp
@@ -318,6 +318,46 @@ void FMVVMViewBlueprintCompiler::CleanOldData(UWidgetBlueprintGeneratedClass* Cl
 		{
 			RenameObjectToTransientPackage(Child);
 		}
+
+		// Clean up the stale viewmodel instances.
+		TMap<FGuid, TWeakObjectPtr<UObject>>& SavedViewModelInstances = BlueprintView->GetOuterUMVVMWidgetBlueprintExtension_View()->TemporaryViewModelInstances;
+		TArray<FGuid> DeleteViewModelInstances;
+
+		for (auto It = SavedViewModelInstances.CreateIterator(); It; ++It)
+		{
+			uint32 ViewModelIndex = ViewModelCreatorContexts.IndexOfByPredicate([It](const FCompilerViewModelCreatorContext& SourceCreatorContext)
+				{
+					return SourceCreatorContext.ViewModelContext.GetViewModelId() == It->Key;
+				});
+			
+			TWeakObjectPtr<UObject> StaticViewModelInstance = nullptr;
+			if (ViewModelIndex != INDEX_NONE)
+			{
+				const FMVVMBlueprintViewModelContext& ViewModelContext = ViewModelCreatorContexts[ViewModelIndex].ViewModelContext;
+				if (It->Value.IsValid())
+				{
+					if (ViewModelContext.CreationType == EMVVMBlueprintViewModelContextCreationType::CreateInstance
+						&& ViewModelContext.bExposeInstanceInEditor)
+					{
+						StaticViewModelInstance = It->Value;
+					}
+					else
+					{
+						RenameObjectToTransientPackage(It->Value.Get());
+					}
+				}
+			}
+			if (!StaticViewModelInstance.IsValid())
+			{
+				DeleteViewModelInstances.Add(It->Key);
+				SavedViewModelInstances[It->Key] = nullptr;
+			}
+		}
+
+		for (FGuid Guid : DeleteViewModelInstances)
+		{
+			SavedViewModelInstances.Remove(Guid);
+		}
 	}
 }
 
@@ -349,12 +389,36 @@ void FMVVMViewBlueprintCompiler::CreateVariables(const FWidgetBlueprintCompilerC
 		FProperty* NewProperty = Context.CreateVariable(UserWidgetProperty.Name, NewPropertyPinType);
 		if (NewProperty != nullptr)
 		{
-			NewProperty->SetPropertyFlags(CPF_BlueprintVisible | CPF_RepSkip | CPF_Transient | CPF_DuplicateTransient);
+			NewProperty->SetPropertyFlags(CPF_BlueprintVisible | CPF_RepSkip);
+
+			const bool bIsInstanceExposed = UserWidgetProperty.bInstanced && UserWidgetProperty.bInstanceExposed;
+			if (bIsInstanceExposed)
+			{
+				NewProperty->SetPropertyFlags(CPF_Edit | CPF_ExportObject | CPF_PersistentInstance | CPF_InstancedReference | CPF_NonNullable | CPF_NoClear);
+
+				if (UserWidgetProperty.AuthoritativeClass->HasAnyClassFlags(CLASS_HasInstancedReference))
+				{
+					NewProperty->SetPropertyFlags(CPF_ContainsInstancedReference);
+				}
+			}
+			else
+			{
+				NewProperty->SetPropertyFlags(CPF_Transient | CPF_DuplicateTransient);
+			}
+
+			if (UserWidgetProperty.bExposeOnSpawn)
+			{
+				NewProperty->SetPropertyFlags(CPF_ExposeOnSpawn);
+			}
+			else if (!bIsInstanceExposed)
+			{
+				NewProperty->SetPropertyFlags(CPF_DisableEditOnInstance);
+			}
+
 			if (UserWidgetProperty.bReadOnly)
 			{
 				NewProperty->SetPropertyFlags(CPF_BlueprintReadOnly);
 			}
-			NewProperty->SetPropertyFlags(UserWidgetProperty.bExposeOnSpawn ? CPF_ExposeOnSpawn : CPF_DisableEditOnInstance);
 
 #if WITH_EDITOR
 			if (!UserWidgetProperty.BlueprintSetter.IsEmpty())
@@ -377,6 +441,10 @@ void FMVVMViewBlueprintCompiler::CreateVariables(const FWidgetBlueprintCompilerC
 			{
 				NewProperty->SetMetaData(FBlueprintMetadata::MD_Private, TEXT("true"));
 			}
+			if (bIsInstanceExposed)
+			{
+				NewProperty->SetMetaData(FName("EditInline"), TEXT("true"));
+			}
 #endif
 		}
 		return NewProperty;
@@ -701,6 +769,8 @@ void FMVVMViewBlueprintCompiler::CreateRequiredProperties(const FWidgetBlueprint
 			SourceVariable.bExposeOnSpawn = bCreateSetterFunction;
 			SourceVariable.bPrivate = !bIsPublicReadable;
 			SourceVariable.bReadOnly = !bCreateSetterFunction;
+			SourceVariable.bInstanced = ViewModelContext.CreationType == EMVVMBlueprintViewModelContextCreationType::CreateInstance;
+			SourceVariable.bInstanceExposed = ViewModelContext.bExposeInstanceInEditor;
 		}
 
 		{
@@ -787,6 +857,8 @@ void FMVVMViewBlueprintCompiler::CreateRequiredProperties(const FWidgetBlueprint
 					//SourceVariable.bPrivate = true; Remove until the data is fix properly
 					SourceVariable.bPrivate = false;
 					SourceVariable.bReadOnly = true;
+					SourceVariable.bInstanced = false;
+					SourceVariable.bInstanceExposed = false;
 
 					WidgetUserPropertyCreated.Add(Widget->GetFName());
 				}
@@ -1737,6 +1809,41 @@ void FMVVMViewBlueprintCompiler::PreCompileViewModelCreatorContexts(UWidgetBluep
 				bIsPreCompileStepValid = false;
 				continue;
 			}
+
+			UObject* CDOInstance = Class->GetDefaultObject(true);
+
+			// If requested, create an instance of the viewmodel to be edited through the details panel
+			if (ViewModelContext.bExposeInstanceInEditor)
+			{
+				FObjectPropertyBase* ViewModelProperty = FindFProperty<FObjectPropertyBase>(Class, ViewModelContext.GetViewModelName());
+				if (!ViewModelProperty->GetObjectPropertyValue_InContainer(CDOInstance))
+				{
+					UObject* StaticInstance = NewObject<UObject>(CDOInstance, ViewModelContext.GetViewModelClass(), NAME_None, RF_Transactional | RF_Public);
+					BlueprintView->GetOuterUMVVMWidgetBlueprintExtension_View()->TemporaryViewModelInstances.Add(ViewModelContext.GetViewModelId(), StaticInstance);
+					if (ensure(StaticInstance->GetClass()->IsChildOf(ViewModelProperty->PropertyClass)))
+					{
+						ViewModelProperty->SetObjectPropertyValue_InContainer(CDOInstance, StaticInstance);
+					}
+				}
+
+				TValueOrError<FCompiledBindingLibraryCompiler::FFieldPathHandle, FText> ReadFieldPathResult = AddObjectFieldPath(BindingLibraryCompiler, Class, ViewModelContext.GetViewModelName().ToString(), ViewModelContext.GetViewModelClass());
+				if (ReadFieldPathResult.HasError())
+				{
+					AddMessageForViewModel(ViewModelContext
+						, ReadFieldPathResult.GetError()
+						, EMessageType::Error
+					);
+					bIsPreCompileStepValid = false;
+					continue;
+				}
+
+				SourceCreatorContext.ReadPropertyPathHandle = ReadFieldPathResult.StealValue();
+			}
+			else
+			{
+				FObjectPropertyBase* ViewModelProperty = FindFProperty<FObjectPropertyBase>(Class, ViewModelContext.GetViewModelName());
+				ViewModelProperty->SetObjectPropertyValue_InContainer(CDOInstance, nullptr);
+			}
 		}
 		else if (ViewModelContext.CreationType == EMVVMBlueprintViewModelContextCreationType::PropertyPath)
 		{
@@ -1882,6 +1989,23 @@ void FMVVMViewBlueprintCompiler::CompileSources(const FCompiledBindingLibraryCom
 		else if (ViewModelContext.CreationType == EMVVMBlueprintViewModelContextCreationType::CreateInstance)
 		{
 			bCreateInstance = true;
+
+			if (ViewModelContext.bExposeInstanceInEditor)
+			{
+				const FMVVMVCompiledFieldPath* CompiledFieldPath = CompileResult.FieldPaths.Find(SourceCreatorContext.ReadPropertyPathHandle);
+				if (CompiledFieldPath == nullptr)
+				{
+					AddMessageForViewModel(ViewModelContext
+						, LOCTEXT("ViewModelInvalidInstanceNotFound", "The path for the created viewmodel instance was not found.")
+						, EMessageType::Error
+					);
+					bIsCompileStepValid = false;
+					continue;
+				}
+
+				CompiledSourceCreator.FieldPath = *CompiledFieldPath;
+			}
+			ensure(bIsOptional == false);
 		}
 		else if (ViewModelContext.CreationType == EMVVMBlueprintViewModelContextCreationType::PropertyPath)
 		{
@@ -1969,6 +2093,7 @@ void FMVVMViewBlueprintCompiler::CompileSources(const FCompiledBindingLibraryCom
 		CompiledSourceCreator.Flags |= bCanBeSet ? (uint16)FMVVMViewClass_Source::EFlags::CanBeSet : 0;
 		CompiledSourceCreator.Flags |= bCanBeEvaluated ? (uint16)FMVVMViewClass_Source::EFlags::CanBeEvaluated : 0;
 		CompiledSourceCreator.Flags |= (uint16)FMVVMViewClass_Source::EFlags::IsViewModel;
+		CompiledSourceCreator.Flags |= ViewModelContext.bExposeInstanceInEditor ? (uint16)FMVVMViewClass_Source::EFlags::IsViewModelInstanceExposed : 0;
 
 		FSortData SortData;
 		SortData.SourceName = CompiledSourceCreator.GetName();
diff --git a/Engine/Plugins/Runtime/ModelViewViewModel/Source/ModelViewViewModelBlueprint/Private/MVVMViewBlueprintCompiler.h b/Engine/Plugins/Runtime/ModelViewViewModel/Source/ModelViewViewModelBlueprint/Private/MVVMViewBlueprintCompiler.h
--- a/Engine/Plugins/Runtime/ModelViewViewModel/Source/ModelViewViewModelBlueprint/Private/MVVMViewBlueprintCompiler.h
+++ b/Engine/Plugins/Runtime/ModelViewViewModel/Source/ModelViewViewModelBlueprint/Private/MVVMViewBlueprintCompiler.h
@@ -149,6 +149,8 @@ private:
 		bool bExposeOnSpawn = false;
 		bool bPrivate = false;
 		bool bReadOnly = false;
+		bool bInstanced = false;
+		bool bInstanceExposed = false;
 
 		const FProperty* Property = nullptr;
 	};
diff --git a/Engine/Plugins/Runtime/ModelViewViewModel/Source/ModelViewViewModelBlueprint/Public/MVVMBlueprintViewModelContext.h b/Engine/Plugins/Runtime/ModelViewViewModel/Source/ModelViewViewModelBlueprint/Public/MVVMBlueprintViewModelContext.h
--- a/Engine/Plugins/Runtime/ModelViewViewModel/Source/ModelViewViewModelBlueprint/Public/MVVMBlueprintViewModelContext.h
+++ b/Engine/Plugins/Runtime/ModelViewViewModel/Source/ModelViewViewModelBlueprint/Public/MVVMBlueprintViewModelContext.h
@@ -156,6 +156,10 @@ public:
 	UPROPERTY(EditAnywhere, Category = "Viewmodel", AdvancedDisplay)
 	bool bOptional = false;
 
+	/** Expose the viewmodel instance on every instance of the user widget for modification in editor. */
+	UPROPERTY(EditAnywhere, Category = "Viewmodel", AdvancedDisplay)
+	bool bExposeInstanceInEditor = false;
+
 	/** Can change the name in the editor. */
 	UPROPERTY()
 	bool bCanRename = true;
diff --git a/Engine/Plugins/Runtime/ModelViewViewModel/Source/ModelViewViewModelBlueprint/Public/MVVMDeveloperProjectSettings.h b/Engine/Plugins/Runtime/ModelViewViewModel/Source/ModelViewViewModelBlueprint/Public/MVVMDeveloperProjectSettings.h
--- a/Engine/Plugins/Runtime/ModelViewViewModel/Source/ModelViewViewModelBlueprint/Public/MVVMDeveloperProjectSettings.h
+++ b/Engine/Plugins/Runtime/ModelViewViewModel/Source/ModelViewViewModelBlueprint/Public/MVVMDeveloperProjectSettings.h
@@ -144,6 +144,13 @@ public:
 	UPROPERTY(EditAnywhere, config, Category = "Viewmodel")
 	bool bCanCreateViewModelInView = false;
 
+	/** 
+	 * When a viewmodel is set to Create Instance, allow modifying the viewmodel instance in the editor on all instances of the owning widget.
+	 * The per-viewmodel setting "Expose Instance In Editor" overrides this.
+	 */
+	UPROPERTY(EditAnywhere, config, Category = "Viewmodel")
+	bool bExposeViewModelInstanceInEditor = false;
+
 	/** Permission list for filtering which execution mode is allowed. */
 	UPROPERTY(EditAnywhere, config, Category = "Viewmodel")
 	EMVVMDeveloperConversionFunctionFilterType ConversionFunctionFilter = EMVVMDeveloperConversionFunctionFilterType::BlueprintActionRegistry;
diff --git a/Engine/Plugins/Runtime/ModelViewViewModel/Source/ModelViewViewModelBlueprint/Public/MVVMWidgetBlueprintExtension_View.h b/Engine/Plugins/Runtime/ModelViewViewModel/Source/ModelViewViewModelBlueprint/Public/MVVMWidgetBlueprintExtension_View.h
--- a/Engine/Plugins/Runtime/ModelViewViewModel/Source/ModelViewViewModelBlueprint/Public/MVVMWidgetBlueprintExtension_View.h
+++ b/Engine/Plugins/Runtime/ModelViewViewModel/Source/ModelViewViewModelBlueprint/Public/MVVMWidgetBlueprintExtension_View.h
@@ -68,6 +68,9 @@ public:
 		return FilterSettings;
 	}
 
+	UPROPERTY(Transient)
+	TMap<FGuid, TWeakObjectPtr<UObject>> TemporaryViewModelInstances;
+
 private:
 	UPROPERTY(Instanced)
 	TObjectPtr<UMVVMBlueprintView> BlueprintView;
diff --git a/Engine/Plugins/Runtime/ModelViewViewModel/Source/ModelViewViewModelEditor/Private/Customizations/MVVMBlueprintViewModelContextCustomization.cpp b/Engine/Plugins/Runtime/ModelViewViewModel/Source/ModelViewViewModelEditor/Private/Customizations/MVVMBlueprintViewModelContextCustomization.cpp
--- a/Engine/Plugins/Runtime/ModelViewViewModel/Source/ModelViewViewModelEditor/Private/Customizations/MVVMBlueprintViewModelContextCustomization.cpp
+++ b/Engine/Plugins/Runtime/ModelViewViewModel/Source/ModelViewViewModelEditor/Private/Customizations/MVVMBlueprintViewModelContextCustomization.cpp
@@ -349,6 +349,18 @@ void FBlueprintViewModelContextDetailCustomization::CustomizeChildren(TSharedRef
 				.IsEnabled(bCanEdit);
 		}
 
+		TSharedPtr<IPropertyHandle> ExposeInstanceInEditorHandle = PropertyHandle->GetChildHandle(GET_MEMBER_NAME_CHECKED(FMVVMBlueprintViewModelContext, bExposeInstanceInEditor), false);
+		if (ensure(ExposeInstanceInEditorHandle))
+		{
+			ChildBuilder.AddProperty(ExposeInstanceInEditorHandle.ToSharedRef())
+				.IsEnabled(bCanEdit)
+				.Visibility(MakeAttributeLambda([ContextPtr]()
+					{
+						bool bResult = ContextPtr->CreationType == EMVVMBlueprintViewModelContextCreationType::CreateInstance;
+						return bResult ? EVisibility::Visible : EVisibility::Collapsed;
+					}));
+		}
+
 		TSharedPtr<IPropertyHandle> OptionalHandle = PropertyHandle->GetChildHandle(GET_MEMBER_NAME_CHECKED(FMVVMBlueprintViewModelContext, bOptional), false);
 		if (ensure(OptionalHandle))
 		{
