MVVM: allow events to bind to functions with any number of parameters.

diff --git a/Engine/Plugins/Runtime/ModelViewViewModel/Source/ModelViewViewModel/Private/Bindings/MVVMBindingHelper.cpp b/Engine/Plugins/Runtime/ModelViewViewModel/Source/ModelViewViewModel/Private/Bindings/MVVMBindingHelper.cpp
--- a/Engine/Plugins/Runtime/ModelViewViewModel/Source/ModelViewViewModel/Private/Bindings/MVVMBindingHelper.cpp
+++ b/Engine/Plugins/Runtime/ModelViewViewModel/Source/ModelViewViewModel/Private/Bindings/MVVMBindingHelper.cpp
@@ -112,6 +112,21 @@ namespace UE::MVVM::BindingHelper
 	}
 
 
+	bool IsValidForEventBinding(const FMVVMConstFieldVariant InVariant)
+	{
+		if (InVariant.IsProperty())
+		{
+			// Same rules as other bindings apply for event binding to properties.
+			return IsValidForDestinationBinding(InVariant.GetProperty());
+		}
+		else if (InVariant.IsFunction())
+		{
+			return IsValidForEventBinding(InVariant.GetFunction());
+		}
+		return false;
+	}
+
+
 	bool IsValidForSimpleRuntimeConversion(const UFunction* InFunction)
 	{
 		if (Private::IsValidCommon(InFunction)
@@ -140,6 +155,12 @@ namespace UE::MVVM::BindingHelper
 	}
 
 
+	bool IsValidForEventBinding(const UFunction* InFunction)
+	{
+		return Private::IsValidCommon(InFunction) && !InFunction->HasAnyFunctionFlags(FUNC_Const | FUNC_BlueprintPure | FUNC_BlueprintEvent);
+	}
+
+
 #if WITH_EDITOR
 	bool IsAccessibleDirectlyForSourceBinding(const FProperty* InProperty)
 	{
diff --git a/Engine/Plugins/Runtime/ModelViewViewModel/Source/ModelViewViewModel/Private/MVVMSubsystem.cpp b/Engine/Plugins/Runtime/ModelViewViewModel/Source/ModelViewViewModel/Private/MVVMSubsystem.cpp
--- a/Engine/Plugins/Runtime/ModelViewViewModel/Source/ModelViewViewModel/Private/MVVMSubsystem.cpp
+++ b/Engine/Plugins/Runtime/ModelViewViewModel/Source/ModelViewViewModel/Private/MVVMSubsystem.cpp
@@ -52,7 +52,7 @@ bool UMVVMSubsystem::DoesWidgetTreeContainedWidget(const UWidgetTree* WidgetTree
 namespace UE::MVVM::Private
 {
 
-	FMVVMAvailableBinding GetAvailableBinding(const UFunction* Function, bool bHasNotify, bool bCanAccessPrivate, bool bCanAccessProtected)
+	FMVVMAvailableBinding GetAvailableBinding(const UFunction* Function, bool bHasNotify, bool bCanAccessPrivate, bool bCanAccessProtected, const bool bIsBindingToEvent)
 	{
 		// N.B. A function can be private/protected and can still be use when it's a BlueprintGetter/BlueprintSetter.
 		//But we bind to the property not the function.
@@ -76,7 +76,15 @@ namespace UE::MVVM::Private
 
 		const bool bIsReadable = BindingHelper::IsValidForSourceBinding(Function);
 		const bool bIsWritable = BindingHelper::IsValidForDestinationBinding(Function);
-		if (bIsReadable || bIsWritable || bHasNotify)
+		const bool bIsForEvent = BindingHelper::IsValidForEventBinding(Function);
+		if (bIsBindingToEvent)
+		{
+			if (bIsForEvent)
+			{
+				return FMVVMAvailableBinding(FMVVMBindingName(Function->GetFName()), bIsReadable, true, bHasNotify && bIsReadable);
+			}
+		}
+		else if (bIsReadable || bIsWritable || bHasNotify)
 		{
 			return FMVVMAvailableBinding(FMVVMBindingName(Function->GetFName()), bIsReadable, bIsWritable, bHasNotify && bIsReadable);
 		}
@@ -111,7 +119,7 @@ namespace UE::MVVM::Private
 		return FMVVMAvailableBinding();
 	}
 
-	TArray<FMVVMAvailableBinding> GetAvailableBindings(const UStruct* Container, const UClass* AccessorType, const UE::FieldNotification::IClassDescriptor* ClassDescriptor)
+	TArray<FMVVMAvailableBinding> GetAvailableBindings(const UStruct* Container, const UClass* AccessorType, const UE::FieldNotification::IClassDescriptor* ClassDescriptor, const bool bIsBindingToEvent)
 	{
 		TSet<FName> FieldDescriptors;
 		if (ClassDescriptor)
@@ -137,7 +145,7 @@ namespace UE::MVVM::Private
 				check(Function);
 
 				const bool bHasNotify = FieldDescriptors.Contains(Function->GetFName());
-				FMVVMAvailableBinding Binding = GetAvailableBinding(Function, bHasNotify, bCanAccessPrivateMember, bCanAccessProtectedMember);
+				FMVVMAvailableBinding Binding = GetAvailableBinding(Function, bHasNotify, bCanAccessPrivateMember, bCanAccessProtectedMember, bIsBindingToEvent);
 				if (Binding.IsValid())
 				{
 					Result.Add(Binding);
@@ -195,22 +203,22 @@ namespace UE::MVVM::Private
 		return Result;
 	}
 
-	TArray<FMVVMAvailableBinding> GetAvailableBindings(const UClass* InSubClass, const UClass* InAccessor)
+	TArray<FMVVMAvailableBinding> GetAvailableBindings(const UClass* InSubClass, const UClass* InAccessor, const bool bIsBindingToEvent)
 	{
 		if (InSubClass && InSubClass->ImplementsInterface(UNotifyFieldValueChanged::StaticClass()))
 		{
 			TScriptInterface<INotifyFieldValueChanged> DefaultObject = InSubClass->GetDefaultObject();
 			const UE::FieldNotification::IClassDescriptor& ClassDescriptor = DefaultObject->GetFieldNotificationDescriptor();
-			return GetAvailableBindings(InSubClass, InAccessor, &ClassDescriptor);
+			return GetAvailableBindings(InSubClass, InAccessor, &ClassDescriptor, bIsBindingToEvent);
 		}
 		else
 		{
-			return GetAvailableBindings(InSubClass, InAccessor, nullptr);
+			return GetAvailableBindings(InSubClass, InAccessor, nullptr, bIsBindingToEvent);
 		}
 	}
 
 
-	FMVVMAvailableBinding GetAvailableBinding(FMVVMConstFieldVariant FieldVariant, const UClass* AccessorType, const UE::FieldNotification::IClassDescriptor* ClassDescriptor)
+	FMVVMAvailableBinding GetAvailableBinding(FMVVMConstFieldVariant FieldVariant, const UClass* AccessorType, const UE::FieldNotification::IClassDescriptor* ClassDescriptor, const bool bIsBindingToEvent)
 	{
 		bool bHasNotify = false;
 		if (ClassDescriptor)
@@ -228,29 +236,29 @@ namespace UE::MVVM::Private
 		}
 		else if (FieldVariant.IsFunction())
 		{
-			return GetAvailableBinding(FieldVariant.GetFunction(), bHasNotify, bCanAccessPrivateMember, bCanAccessProtectedMember);
+			return GetAvailableBinding(FieldVariant.GetFunction(), bHasNotify, bCanAccessPrivateMember, bCanAccessProtectedMember, bIsBindingToEvent);
 		}
 
 		return FMVVMAvailableBinding();
 	}
 
 
-	FMVVMAvailableBinding GetAvailableBinding(FMVVMConstFieldVariant FieldVariant, const UClass* InAccessor)
+	FMVVMAvailableBinding GetAvailableBinding(FMVVMConstFieldVariant FieldVariant, const UClass* InAccessor, const bool bIsBindingToEvent)
 	{
 		const UClass* OwnerClass = Cast<UClass>(FieldVariant.GetOwner());
 		if (OwnerClass && OwnerClass->ImplementsInterface(UNotifyFieldValueChanged::StaticClass()))
 		{
 			TScriptInterface<INotifyFieldValueChanged> DefaultObject = OwnerClass->GetDefaultObject();
 			const UE::FieldNotification::IClassDescriptor& ClassDescriptor = DefaultObject->GetFieldNotificationDescriptor();
-			return GetAvailableBinding(FieldVariant, InAccessor, &ClassDescriptor);
+			return GetAvailableBinding(FieldVariant, InAccessor, &ClassDescriptor, bIsBindingToEvent);
 		}
 		else
 		{
-			return GetAvailableBinding(FieldVariant, InAccessor, nullptr);
+			return GetAvailableBinding(FieldVariant, InAccessor, nullptr, bIsBindingToEvent);
 		}
 	}
 
-	FMVVMAvailableBinding GetAvailableBinding(FMVVMBindingName BindingName, const UClass* InSubClass, const UClass* InAccessor)
+	FMVVMAvailableBinding GetAvailableBinding(FMVVMBindingName BindingName, const UClass* InSubClass, const UClass* InAccessor, const bool bIsBindingToEvent)
 	{
 		if (!BindingName.IsValid())
 		{
@@ -258,7 +266,7 @@ namespace UE::MVVM::Private
 		}
 
 		FMVVMConstFieldVariant FieldVariant = UE::MVVM::BindingHelper::FindFieldByName(InSubClass, BindingName);
-		return GetAvailableBinding(FieldVariant, InAccessor);
+		return GetAvailableBinding(FieldVariant, InAccessor, bIsBindingToEvent);
 	}
 } //namespace
 
@@ -271,15 +279,19 @@ TArray<FMVVMAvailableBinding> UMVVMSubsystem::K2_GetAvailableBindings(const UCla
 
 TArray<FMVVMAvailableBinding> UMVVMSubsystem::GetAvailableBindings(const UClass* Class, const UClass* Accessor)
 {
-	return UE::MVVM::Private::GetAvailableBindings(Class, Accessor);
+	return UE::MVVM::Private::GetAvailableBindings(Class, Accessor, false);
 }
 
 
 TArray<FMVVMAvailableBinding> UMVVMSubsystem::GetAvailableBindingsForStruct(const UScriptStruct* Struct)
 {
-	return UE::MVVM::Private::GetAvailableBindings(Struct, nullptr, nullptr);
+	return UE::MVVM::Private::GetAvailableBindings(Struct, nullptr, nullptr, false);
 }
 
+TArray<FMVVMAvailableBinding> UMVVMSubsystem::GetAvailableBindingsForEvent(const UClass* Class, const UClass* Accessor)
+{
+	return UE::MVVM::Private::GetAvailableBindings(Class, Accessor, true);
+}
 
 FMVVMAvailableBinding UMVVMSubsystem::K2_GetAvailableBinding(const UClass* Class, FMVVMBindingName BindingName, const UClass* Accessor) const
 {
@@ -289,15 +301,24 @@ FMVVMAvailableBinding UMVVMSubsystem::K2_GetAvailableBinding(const UClass* Class
 
 FMVVMAvailableBinding UMVVMSubsystem::GetAvailableBinding(const UClass* Class, FMVVMBindingName BindingName, const UClass* Accessor)
 {
-	return UE::MVVM::Private::GetAvailableBinding(BindingName, Class, Accessor);
+	return UE::MVVM::Private::GetAvailableBinding(BindingName, Class, Accessor, false);
 }
 
 
 FMVVMAvailableBinding UMVVMSubsystem::GetAvailableBindingForField(UE::MVVM::FMVVMConstFieldVariant FieldVariant, const UClass* Accessor)
 {
-	return UE::MVVM::Private::GetAvailableBinding(FieldVariant, Accessor);
+	return UE::MVVM::Private::GetAvailableBinding(FieldVariant, Accessor, false);
+}
+
+FMVVMAvailableBinding UMVVMSubsystem::GetAvailableBindingForEvent(UE::MVVM::FMVVMConstFieldVariant FieldVariant, const UClass* Accessor)
+{
+	return UE::MVVM::Private::GetAvailableBinding(FieldVariant, Accessor, true);
 }
 
+FMVVMAvailableBinding UMVVMSubsystem::GetAvailableBindingForEvent(const UClass* Class, FMVVMBindingName BindingName, const UClass* Accessor)
+{
+	return UE::MVVM::Private::GetAvailableBinding(BindingName, Class, Accessor, true);
+}
 
 UMVVMViewModelCollectionObject* UMVVMSubsystem::GetGlobalViewModelCollection() const
 {
diff --git a/Engine/Plugins/Runtime/ModelViewViewModel/Source/ModelViewViewModel/Public/Bindings/MVVMBindingHelper.h b/Engine/Plugins/Runtime/ModelViewViewModel/Source/ModelViewViewModel/Public/Bindings/MVVMBindingHelper.h
--- a/Engine/Plugins/Runtime/ModelViewViewModel/Source/ModelViewViewModel/Public/Bindings/MVVMBindingHelper.h
+++ b/Engine/Plugins/Runtime/ModelViewViewModel/Source/ModelViewViewModel/Public/Bindings/MVVMBindingHelper.h
@@ -47,7 +47,7 @@ namespace UE::MVVM::BindingHelper
 	 * Is the Field usable as a destination (settable) by the binding system.
 	 * @note It may be a BlueprintSetter and binding to the Property would be better in the editor.
 	 */
-	[[nodiscard]] MODELVIEWVIEWMODEL_API bool IsValidForDestinationBinding(const FMVVMConstFieldVariant InFunction);
+	[[nodiscard]] MODELVIEWVIEWMODEL_API bool IsValidForDestinationBinding(const FMVVMConstFieldVariant InVariant);
 
 	/**
 	 * Is the Function usable as a conversion function by the binding system without any wrapper.
@@ -61,6 +61,17 @@ namespace UE::MVVM::BindingHelper
 	 */
 	[[nodiscard]] MODELVIEWVIEWMODEL_API bool IsValidForComplexRuntimeConversion(const UFunction* InFunction);
 
+	/**
+	 * Is the Function usable as an event binding.
+	 * Events can trigger functions with any number of parameters and any return values.
+	 */
+	[[nodiscard]] MODELVIEWVIEWMODEL_API bool IsValidForEventBinding(const UFunction* InFunction);
+
+	/**
+	 * Is the Field usable as an event binding by the binding system.
+	 */
+	[[nodiscard]] MODELVIEWVIEWMODEL_API bool IsValidForEventBinding(const FMVVMConstFieldVariant InVariant);
+
 #if WITH_EDITOR
 	/** Is the Property usable as a source by the binding system and can it be read directly or it requires a Getter. */
 	[[nodiscard]] MODELVIEWVIEWMODEL_API bool IsAccessibleDirectlyForSourceBinding(const FProperty* InProperty);
diff --git a/Engine/Plugins/Runtime/ModelViewViewModel/Source/ModelViewViewModel/Public/MVVMSubsystem.h b/Engine/Plugins/Runtime/ModelViewViewModel/Source/ModelViewViewModel/Public/MVVMSubsystem.h
--- a/Engine/Plugins/Runtime/ModelViewViewModel/Source/ModelViewViewModel/Public/MVVMSubsystem.h
+++ b/Engine/Plugins/Runtime/ModelViewViewModel/Source/ModelViewViewModel/Public/MVVMSubsystem.h
@@ -52,6 +52,8 @@ public:
 	 */
 	static TArray<FMVVMAvailableBinding> GetAvailableBindingsForStruct(const UScriptStruct* Struct);
 
+	static TArray<FMVVMAvailableBinding> GetAvailableBindingsForEvent(const UClass* Class, const UClass* Accessor);
+
 	/** @return The AvailableBinding from a BindingName. */
 	UFUNCTION(BlueprintCallable, Category = "Viewmodel", meta = (DisplayName = "Get Available Binding"))
 	FMVVMAvailableBinding K2_GetAvailableBinding(const UClass* Class, FMVVMBindingName BindingName, const UClass* Accessor) const;
@@ -61,6 +63,10 @@ public:
 	/** @return The AvailableBinding from a field. */
 	static FMVVMAvailableBinding GetAvailableBindingForField(UE::MVVM::FMVVMConstFieldVariant Variant, const UClass* Accessor);
 
+	static FMVVMAvailableBinding GetAvailableBindingForEvent(UE::MVVM::FMVVMConstFieldVariant FieldVariant, const UClass* Accessor);
+
+	static FMVVMAvailableBinding GetAvailableBindingForEvent(const UClass* Class, FMVVMBindingName BindingName, const UClass* Accessor);
+
 	UE_DEPRECATED(5.3, "GetGlobalViewModelCollection has been deprecated, please use the game instance subsystem.")
 	UFUNCTION(BlueprintCallable, Category = "Viewmodel", meta=(DeprecatedFunction, DeprecatedMessage = "This version of GetGlobalViewModelCollection has been deprecated, please use GetGlobalViewModelCollection from the Viewmodel Game subsystem."))
 	UMVVMViewModelCollectionObject* GetGlobalViewModelCollection() const;
diff --git a/Engine/Plugins/Runtime/ModelViewViewModel/Source/ModelViewViewModelBlueprint/Private/Bindings/MVVMConversionFunctionHelper.cpp b/Engine/Plugins/Runtime/ModelViewViewModel/Source/ModelViewViewModelBlueprint/Private/Bindings/MVVMConversionFunctionHelper.cpp
--- a/Engine/Plugins/Runtime/ModelViewViewModel/Source/ModelViewViewModelBlueprint/Private/Bindings/MVVMConversionFunctionHelper.cpp
+++ b/Engine/Plugins/Runtime/ModelViewViewModel/Source/ModelViewViewModelBlueprint/Private/Bindings/MVVMConversionFunctionHelper.cpp
@@ -571,7 +571,7 @@ namespace Private
 		 */
 		int32 LocalContainerPathIndex = INDEX_NONE;
 	};
-	TValueOrError<FCanSetterGraphResult, FText> CanCreateSetterGraph(UBlueprint* WidgetBlueprint, const TArrayView<FMVVMConstFieldVariant> Path)
+	TValueOrError<FCanSetterGraphResult, FText> CanCreateSetterGraph(UBlueprint* WidgetBlueprint, const TArrayView<FMVVMConstFieldVariant> Path, const bool bIsForEvent)
 	{
 		/**
 		 * Different types:
@@ -628,7 +628,14 @@ namespace Private
 		}
 
 		const FMVVMConstFieldVariant& LastField = Path.Last();
-		if (!BindingHelper::IsValidForDestinationBinding(LastField))
+		if (bIsForEvent)
+		{
+			if (!BindingHelper::IsValidForEventBinding(LastField))
+			{
+				return MakeError(FText::Format(LOCTEXT("SetterGraph_InvalidSetter", "The setter for field {0} can't be used."), FText::FromName(LastField.GetName())));
+			}
+		}
+		else if (!BindingHelper::IsValidForDestinationBinding(LastField))
 		{
 			return MakeError(FText::Format(LOCTEXT("SetterGraph_InvalidSetter", "The setter for field {0} can't be used."), FText::FromName(LastField.GetName())));
 		}
@@ -638,7 +645,14 @@ namespace Private
 			// Can the local container be set.
 			check(Path.IsValidIndex(Result.LocalContainerPathIndex));
 			const FMVVMConstFieldVariant& Field = Path[Result.LocalContainerPathIndex];
-			if (!BindingHelper::IsValidForDestinationBinding(Field))
+			if (bIsForEvent)
+			{
+				if (!BindingHelper::IsValidForEventBinding(LastField))
+				{
+					return MakeError(FText::Format(LOCTEXT("SetterGraph_ContainerSetter", "The path contains a getter of a struct that can't be set. See field {0}."), FText::FromName(Field.GetName())));
+				}
+			}
+			else if (!BindingHelper::IsValidForDestinationBinding(Field))
 			{
 				return MakeError(FText::Format(LOCTEXT("SetterGraph_ContainerSetter", "The path contains a getter of a struct that can't be set. See field {0}."), FText::FromName(Field.GetName())));
 			}
@@ -732,7 +746,7 @@ FName CreateWrapperName(const FMVVMBlueprintViewBinding& Binding, bool bSourceTo
 TValueOrError<void, FText> CanCreateSetterGraph(UBlueprint* Blueprint, const FMVVMBlueprintPropertyPath& PropertyPath)
 {
 	TArray<UE::MVVM::FMVVMConstFieldVariant> Fields = PropertyPath.GetCompleteFields(Blueprint);
-	TValueOrError<Private::FCanSetterGraphResult, FText> Result = Private::CanCreateSetterGraph(Blueprint, Fields);
+	TValueOrError<Private::FCanSetterGraphResult, FText> Result = Private::CanCreateSetterGraph(Blueprint, Fields, false);
 	if (Result.HasError())
 	{
 		return MakeError(Result.StealError());
@@ -740,10 +754,10 @@ TValueOrError<void, FText> CanCreateSetterGraph(UBlueprint* Blueprint, const FMV
 	return MakeValue();
 }
 
-TValueOrError<FCreateGraphResult, FText> CreateSetterGraph(UBlueprint* Blueprint, FName GraphName, const UFunction* Signature, const FMVVMBlueprintPropertyPath& PropertyPath, bool bIsConst, bool bTransient)
+TValueOrError<FCreateGraphResult, FText> CreateSetterGraph(UBlueprint* Blueprint, FName GraphName, const UFunction* Signature, const FMVVMBlueprintPropertyPath& PropertyPath, bool bIsConst, bool bTransient, const bool bIsForEvent)
 {
 	TArray<UE::MVVM::FMVVMConstFieldVariant> Fields = PropertyPath.GetCompleteFields(Blueprint);
-	TValueOrError<Private::FCanSetterGraphResult, FText> CanCreateSetterGraphResult = Private::CanCreateSetterGraph(Blueprint, Fields);
+	TValueOrError<Private::FCanSetterGraphResult, FText> CanCreateSetterGraphResult = Private::CanCreateSetterGraph(Blueprint, Fields, bIsForEvent);
 	if (CanCreateSetterGraphResult.HasError())
 	{
 		return MakeError(CanCreateSetterGraphResult.StealError());
diff --git a/Engine/Plugins/Runtime/ModelViewViewModel/Source/ModelViewViewModelBlueprint/Private/MVVMBlueprintViewEvent.cpp b/Engine/Plugins/Runtime/ModelViewViewModel/Source/ModelViewViewModelBlueprint/Private/MVVMBlueprintViewEvent.cpp
--- a/Engine/Plugins/Runtime/ModelViewViewModel/Source/ModelViewViewModelBlueprint/Private/MVVMBlueprintViewEvent.cpp
+++ b/Engine/Plugins/Runtime/ModelViewViewModel/Source/ModelViewViewModelBlueprint/Private/MVVMBlueprintViewEvent.cpp
@@ -246,7 +246,7 @@ UEdGraph* UMVVMBlueprintViewEvent::CreateWrapperGraphInternal()
 	UWidgetBlueprint* WidgetBlueprint = GetWidgetBlueprintInternal();
 	bool bIsConst = false;
 	bool bTransient = true;
-	TValueOrError<UE::MVVM::ConversionFunctionHelper::FCreateGraphResult, FText> CreateSetterGraphResult = UE::MVVM::ConversionFunctionHelper::CreateSetterGraph(WidgetBlueprint, GraphName, DelegateSignature, DestinationPath, bIsConst, bTransient);
+	TValueOrError<UE::MVVM::ConversionFunctionHelper::FCreateGraphResult, FText> CreateSetterGraphResult = UE::MVVM::ConversionFunctionHelper::CreateSetterGraph(WidgetBlueprint, GraphName, DelegateSignature, DestinationPath, bIsConst, bTransient, true);
 	if (CreateSetterGraphResult.HasError())
 	{
 		SetCachedWrapperGraphInternal(nullptr, nullptr);
diff --git a/Engine/Plugins/Runtime/ModelViewViewModel/Source/ModelViewViewModelBlueprint/Public/Bindings/MVVMConversionFunctionHelper.h b/Engine/Plugins/Runtime/ModelViewViewModel/Source/ModelViewViewModelBlueprint/Public/Bindings/MVVMConversionFunctionHelper.h
--- a/Engine/Plugins/Runtime/ModelViewViewModel/Source/ModelViewViewModelBlueprint/Public/Bindings/MVVMConversionFunctionHelper.h
+++ b/Engine/Plugins/Runtime/ModelViewViewModel/Source/ModelViewViewModelBlueprint/Public/Bindings/MVVMConversionFunctionHelper.h
@@ -58,7 +58,7 @@ namespace UE::MVVM::ConversionFunctionHelper
 	/**
 	 * Create a graph to set a property/function.
 	 */
-	MODELVIEWVIEWMODELBLUEPRINT_API TValueOrError<FCreateGraphResult, FText> CreateSetterGraph(UBlueprint* WidgetBlueprint, FName GraphName, const UFunction* Signature, const FMVVMBlueprintPropertyPath& PropertyPath, bool bIsConst, bool bTransient);
+	MODELVIEWVIEWMODELBLUEPRINT_API TValueOrError<FCreateGraphResult, FText> CreateSetterGraph(UBlueprint* WidgetBlueprint, FName GraphName, const UFunction* Signature, const FMVVMBlueprintPropertyPath& PropertyPath, bool bIsConst, bool bTransient, const bool bIsForEvent);
 
 	/** */
 	MODELVIEWVIEWMODELBLUEPRINT_API FCreateGraphResult CreateGraph(UBlueprint* WidgetBlueprint, FName GraphName, const UFunction* Signature, const UFunction* FunctionToWrap, bool bIsConst, bool bTransient);
diff --git a/Engine/Plugins/Runtime/ModelViewViewModel/Source/ModelViewViewModelEditor/Private/Widgets/BindingEntry/SMVVMEventRow.cpp b/Engine/Plugins/Runtime/ModelViewViewModel/Source/ModelViewViewModelEditor/Private/Widgets/BindingEntry/SMVVMEventRow.cpp
--- a/Engine/Plugins/Runtime/ModelViewViewModel/Source/ModelViewViewModelEditor/Private/Widgets/BindingEntry/SMVVMEventRow.cpp
+++ b/Engine/Plugins/Runtime/ModelViewViewModel/Source/ModelViewViewModelEditor/Private/Widgets/BindingEntry/SMVVMEventRow.cpp
@@ -124,6 +124,7 @@ TSharedRef<SWidget> SEventRow::BuildRowWidget()
 					.OnGetSelectionContext(this, &SEventRow::GetSelectedSelectionContext, false)
 					.OnDrop(this, &SEventRow::HandleFieldSelectorDrop, false)
 					.OnDragEnter(this, &SEventRow::HandleFieldSelectorDragEnter, false)
+					.IsBindingToEvent(true)
 				]
 			]
 
diff --git a/Engine/Plugins/Runtime/ModelViewViewModel/Source/ModelViewViewModelEditor/Private/Widgets/SMVVMFieldSelector.cpp b/Engine/Plugins/Runtime/ModelViewViewModel/Source/ModelViewViewModelEditor/Private/Widgets/SMVVMFieldSelector.cpp
--- a/Engine/Plugins/Runtime/ModelViewViewModel/Source/ModelViewViewModelEditor/Private/Widgets/SMVVMFieldSelector.cpp
+++ b/Engine/Plugins/Runtime/ModelViewViewModel/Source/ModelViewViewModelEditor/Private/Widgets/SMVVMFieldSelector.cpp
@@ -38,6 +38,7 @@ void SFieldSelector::Construct(const FArguments& InArgs, const UWidgetBlueprint*
 	OnGetSelectionContext = InArgs._OnGetSelectionContext;
 	OnDragEnterEvent = InArgs._OnDragEnter;
 	OnDropEvent = InArgs._OnDrop;
+	bIsBindingToEvent = InArgs._IsBindingToEvent;
 
 	ChildSlot
 	[
@@ -134,6 +135,7 @@ TSharedRef<SWidget> SFieldSelector::HandleGetMenuContent()
 		.OnSelectionChanged(this, &SFieldSelector::HandleFieldSelectionChanged)
 		.OnMenuCloseRequested(this, &SFieldSelector::HandleMenuClosed)
 		.SelectionContext(SelectionContext)
+		.IsBindingToEvent(bIsBindingToEvent)
 		;
 
 	ComboButton->SetMenuContentWidgetToFocus(Menu->GetWidgetToFocus());
diff --git a/Engine/Plugins/Runtime/ModelViewViewModel/Source/ModelViewViewModelEditor/Private/Widgets/SMVVMFieldSelector.h b/Engine/Plugins/Runtime/ModelViewViewModel/Source/ModelViewViewModelEditor/Private/Widgets/SMVVMFieldSelector.h
--- a/Engine/Plugins/Runtime/ModelViewViewModel/Source/ModelViewViewModelEditor/Private/Widgets/SMVVMFieldSelector.h
+++ b/Engine/Plugins/Runtime/ModelViewViewModel/Source/ModelViewViewModelEditor/Private/Widgets/SMVVMFieldSelector.h
@@ -30,7 +30,8 @@ public:
 		{
 		}
 		SLATE_STYLE_ARGUMENT(FTextBlockStyle, TextStyle)
-		SLATE_ARGUMENT_DEFAULT(bool, ShowContext) = true;
+		SLATE_ARGUMENT_DEFAULT(bool, ShowContext) { true };
+		SLATE_ARGUMENT_DEFAULT(bool, IsBindingToEvent) { false };
 		SLATE_EVENT(FOnGetLinkedPinValue, OnGetLinkedValue)
 		SLATE_EVENT(FOnLinkedValueSelectionChanged, OnSelectionChanged)
 		SLATE_EVENT(FOnGetSelectionContext, OnGetSelectionContext)
@@ -66,6 +67,7 @@ private:
 	FOnGetSelectionContext OnGetSelectionContext;
 	FOnDrop OnDropEvent;
 	FOnDragEnter OnDragEnterEvent;
+	bool bIsBindingToEvent = false;
 }; 
 
 } // namespace UE::MVVM
diff --git a/Engine/Plugins/Runtime/ModelViewViewModel/Source/ModelViewViewModelEditor/Private/Widgets/SMVVMFieldSelectorMenu.cpp b/Engine/Plugins/Runtime/ModelViewViewModel/Source/ModelViewViewModelEditor/Private/Widgets/SMVVMFieldSelectorMenu.cpp
--- a/Engine/Plugins/Runtime/ModelViewViewModel/Source/ModelViewViewModelEditor/Private/Widgets/SMVVMFieldSelectorMenu.cpp
+++ b/Engine/Plugins/Runtime/ModelViewViewModel/Source/ModelViewViewModelEditor/Private/Widgets/SMVVMFieldSelectorMenu.cpp
@@ -923,7 +923,8 @@ TSharedRef<SWidget> SFieldSelectorMenu::CreateBindingListPanel(const FArguments&
 		.ShowSearchBox(false)
 		.OnDoubleClicked(this, &SFieldSelectorMenu::SetPropertyPathSelection)
 		.FieldVisibilityFlags(Private::GetFieldVisibilityFlags(SelectionContext.BindingMode, SelectionContext.bReadable, SelectionContext.bWritable))
-		.AssignableTo(AssignableToProperty);
+		.AssignableTo(AssignableToProperty)
+		.IsBindingToEvent(InArgs._IsBindingToEvent);
 
 	const bool bValidSelectedProperty = InArgs._CurrentSelected.IsSet() && InArgs._CurrentSelected.GetValue().IsPropertyPath();
 
diff --git a/Engine/Plugins/Runtime/ModelViewViewModel/Source/ModelViewViewModelEditor/Private/Widgets/SMVVMFieldSelectorMenu.h b/Engine/Plugins/Runtime/ModelViewViewModel/Source/ModelViewViewModelEditor/Private/Widgets/SMVVMFieldSelectorMenu.h
--- a/Engine/Plugins/Runtime/ModelViewViewModel/Source/ModelViewViewModelEditor/Private/Widgets/SMVVMFieldSelectorMenu.h
+++ b/Engine/Plugins/Runtime/ModelViewViewModel/Source/ModelViewViewModelEditor/Private/Widgets/SMVVMFieldSelectorMenu.h
@@ -48,6 +48,7 @@ public:
 		SLATE_EVENT(FOnLinkedValueSelectionChanged, OnSelectionChanged)
 		SLATE_EVENT(FSimpleDelegate, OnMenuCloseRequested)
 		SLATE_ARGUMENT(FFieldSelectionContext, SelectionContext)
+		SLATE_ARGUMENT_DEFAULT(bool, IsBindingToEvent) { false };
 	SLATE_END_ARGS()
 
 	void Construct(const FArguments& InArgs, const UWidgetBlueprint* InWidgetBlueprint);
diff --git a/Engine/Plugins/Runtime/ModelViewViewModel/Source/ModelViewViewModelEditor/Private/Widgets/SMVVMViewModelBindingListWidget.cpp b/Engine/Plugins/Runtime/ModelViewViewModel/Source/ModelViewViewModelEditor/Private/Widgets/SMVVMViewModelBindingListWidget.cpp
--- a/Engine/Plugins/Runtime/ModelViewViewModel/Source/ModelViewViewModelEditor/Private/Widgets/SMVVMViewModelBindingListWidget.cpp
+++ b/Engine/Plugins/Runtime/ModelViewViewModel/Source/ModelViewViewModelEditor/Private/Widgets/SMVVMViewModelBindingListWidget.cpp
@@ -158,10 +158,11 @@ namespace Private
 	}
 }
 
-FFieldIterator_Bindable::FFieldIterator_Bindable(const UWidgetBlueprint* InWidgetBlueprint, EFieldVisibility InVisibilityFlags, const FProperty* InAssignableTo) :
+FFieldIterator_Bindable::FFieldIterator_Bindable(const UWidgetBlueprint* InWidgetBlueprint, EFieldVisibility InVisibilityFlags, const FProperty* InAssignableTo, const bool InIsBindingToEvent) :
 	WidgetBlueprint(InWidgetBlueprint),
 	FieldVisibilityFlags(InVisibilityFlags),
-	AssignableTo(InAssignableTo)
+	AssignableTo(InAssignableTo),
+	bIsBindingToEvent(InIsBindingToEvent)
 {
 }
 
@@ -205,7 +206,8 @@ TArray<FFieldVariant> FFieldIterator_Bindable::GetFields(const UStruct* Struct)
 	{
 		const UWidgetBlueprint* WidgetBlueprintPtr = WidgetBlueprint.Get();
 		TSubclassOf<UObject> AccessorClass = WidgetBlueprintPtr ? WidgetBlueprintPtr->GeneratedClass : nullptr;
-		TArray<FMVVMAvailableBinding> Bindings = UMVVMSubsystem::GetAvailableBindings(const_cast<UClass*>(Class), AccessorClass);
+
+		TArray<FMVVMAvailableBinding> Bindings = bIsBindingToEvent ? UMVVMSubsystem::GetAvailableBindingsForEvent(const_cast<UClass*>(Class), AccessorClass) : UMVVMSubsystem::GetAvailableBindings(const_cast<UClass*>(Class), AccessorClass);
 		AddResult(Bindings);
 	}
 	else if (const UScriptStruct* ScriptStruct = Cast<const UScriptStruct>(Struct))
@@ -298,11 +300,11 @@ TOptional<const UStruct*> FFieldExpander_Bindable::GetExpandedFunction(const UFu
 	return TOptional<const UStruct*>();
 }
 
-TSharedRef<SWidget> ConstructFieldPreSlot(const UWidgetBlueprint* WidgetBlueprint, UE::PropertyViewer::SPropertyViewer::FHandle Handle, const FFieldVariant FieldPath)
+TSharedRef<SWidget> ConstructFieldPreSlot(const UWidgetBlueprint* WidgetBlueprint, UE::PropertyViewer::SPropertyViewer::FHandle Handle, const FFieldVariant FieldPath, const bool bIsForEvent)
 {
 	TSharedRef<SWidget> ImageWidget = SNullWidget::NullWidget;
 	TSubclassOf<UObject> AccessorClass = WidgetBlueprint ? WidgetBlueprint->SkeletonGeneratedClass : nullptr;
-	FMVVMAvailableBinding Binding = UMVVMSubsystem::GetAvailableBindingForField(FMVVMConstFieldVariant(FieldPath), AccessorClass);
+	FMVVMAvailableBinding Binding = bIsForEvent ? UMVVMSubsystem::GetAvailableBindingForEvent(FMVVMConstFieldVariant(FieldPath), AccessorClass) : UMVVMSubsystem::GetAvailableBindingForField(FMVVMConstFieldVariant(FieldPath), AccessorClass);
 	if (Binding.IsValid())
 	{
 		const FSlateBrush* Brush = nullptr;
@@ -358,10 +360,10 @@ TSharedRef<SWidget> ConstructFieldPreSlot(const UWidgetBlueprint* WidgetBlueprin
 void SSourceBindingList::Construct(const FArguments& InArgs, const UWidgetBlueprint* InWidgetBlueprint)
 {
 	WidgetBlueprint = InWidgetBlueprint;
-	FieldIterator = MakeUnique<FFieldIterator_Bindable>(InWidgetBlueprint, InArgs._FieldVisibilityFlags, InArgs._AssignableTo);
+	FieldIterator = MakeUnique<FFieldIterator_Bindable>(InWidgetBlueprint, InArgs._FieldVisibilityFlags, InArgs._AssignableTo, InArgs._IsBindingToEvent);
 	FieldExpander = MakeUnique<FFieldExpander_Bindable>();
 
-
+	bIsBindingToEvent = InArgs._IsBindingToEvent;
 	OnDoubleClicked = InArgs._OnDoubleClicked;
 
 	PropertyViewer = SNew(SPropertyViewer)
@@ -513,7 +515,7 @@ FMVVMBlueprintPropertyPath SSourceBindingList::CreateBlueprintPropertyPath(SProp
 
 			if (const UClass* OwnerClass = Cast<const UClass>(CurrentContainer))
 			{
-				FMVVMAvailableBinding Binding = UMVVMSubsystem::GetAvailableBinding(OwnerClass, FMVVMBindingName(FieldName), AccessorClass);
+				FMVVMAvailableBinding Binding = bIsBindingToEvent ? UMVVMSubsystem::GetAvailableBindingForEvent(OwnerClass, FMVVMBindingName(FieldName), AccessorClass) : UMVVMSubsystem::GetAvailableBinding(OwnerClass, FMVVMBindingName(FieldName), AccessorClass);
 				if (Binding.IsValid())
 				{
 					EFilterFlag FilterFlags;
@@ -555,7 +557,7 @@ FMVVMBlueprintPropertyPath SSourceBindingList::CreateBlueprintPropertyPath(SProp
 	else if(AccessorClass)
 	{
 		FMVVMBindingName BindingName = Source->Key.ToBindingName(WidgetBlueprintPtr);
-		FMVVMAvailableBinding Binding = UMVVMSubsystem::GetAvailableBinding(AccessorClass, BindingName, AccessorClass);
+		FMVVMAvailableBinding Binding = bIsBindingToEvent ? UMVVMSubsystem::GetAvailableBindingForEvent(AccessorClass, BindingName, AccessorClass) : UMVVMSubsystem::GetAvailableBinding(AccessorClass, BindingName, AccessorClass);
 		if (Binding.IsValid())
 		{
 			EFilterFlag FilterFlags;
@@ -591,7 +593,7 @@ TSharedPtr<SWidget> SSourceBindingList::HandleGetPreSlot(SPropertyViewer::FHandl
 {
 	if (FieldPath.Num() > 0)
 	{
-		return ConstructFieldPreSlot(WidgetBlueprint.Get(), Handle, FieldPath.Last());
+		return ConstructFieldPreSlot(WidgetBlueprint.Get(), Handle, FieldPath.Last(), bIsBindingToEvent);
 	}
 	return TSharedPtr<SWidget>();
 }
diff --git a/Engine/Plugins/Runtime/ModelViewViewModel/Source/ModelViewViewModelEditor/Private/Widgets/SMVVMViewModelBindingListWidget.h b/Engine/Plugins/Runtime/ModelViewViewModel/Source/ModelViewViewModelEditor/Private/Widgets/SMVVMViewModelBindingListWidget.h
--- a/Engine/Plugins/Runtime/ModelViewViewModel/Source/ModelViewViewModelEditor/Private/Widgets/SMVVMViewModelBindingListWidget.h
+++ b/Engine/Plugins/Runtime/ModelViewViewModel/Source/ModelViewViewModelEditor/Private/Widgets/SMVVMViewModelBindingListWidget.h
@@ -33,7 +33,7 @@ namespace UE::MVVM
 	class FFieldIterator_Bindable : public UE::PropertyViewer::FFieldIterator_BlueprintVisible
 	{
 	public:
-		FFieldIterator_Bindable(const UWidgetBlueprint* WidgetBlueprint, EFieldVisibility InVisibilityFlags, const FProperty* InAssignableTo = nullptr);
+		FFieldIterator_Bindable(const UWidgetBlueprint* WidgetBlueprint, EFieldVisibility InVisibilityFlags, const FProperty* InAssignableTo = nullptr, const bool InIsBindingToEvent = false);
 		virtual TArray<FFieldVariant> GetFields(const UStruct*) const override;
 		EFieldVisibility GetFieldVisibilityFlags() const
 		{
@@ -48,6 +48,7 @@ namespace UE::MVVM
 		TWeakObjectPtr<const UWidgetBlueprint> WidgetBlueprint; 
 		EFieldVisibility FieldVisibilityFlags = EFieldVisibility::All;
 		const FProperty* AssignableTo = nullptr;
+		bool bIsBindingToEvent = false;
 	};
 
 	/**
@@ -63,7 +64,7 @@ namespace UE::MVVM
 	};
 
 	/** Genereate the icon for the binding type. */
-	TSharedRef<SWidget> ConstructFieldPreSlot(const UWidgetBlueprint* WidgetBlueprint, UE::PropertyViewer::SPropertyViewer::FHandle Handle, const FFieldVariant FieldPath);
+	TSharedRef<SWidget> ConstructFieldPreSlot(const UWidgetBlueprint* WidgetBlueprint, UE::PropertyViewer::SPropertyViewer::FHandle Handle, const FFieldVariant FieldPath, const bool bIsForEvent = false);
 
 	/** 
 	 * 
@@ -77,6 +78,7 @@ namespace UE::MVVM
 			SLATE_ARGUMENT_DEFAULT(bool, ShowSearchBox) = false;
 			SLATE_ARGUMENT_DEFAULT(EFieldVisibility, FieldVisibilityFlags) = EFieldVisibility::All;
 			SLATE_ARGUMENT_DEFAULT(bool, EnableSelection) = true;
+			SLATE_ARGUMENT_DEFAULT(bool, IsBindingToEvent) = false;
 
 			/** 
 			 * Show only properties that are assignable to the given property. 
@@ -122,6 +124,7 @@ namespace UE::MVVM
 		FMVVMBlueprintPropertyPath SelectedPath;
 		TSharedPtr<SPropertyViewer> PropertyViewer;
 		TWeakObjectPtr<const UWidgetBlueprint> WidgetBlueprint;
+		bool bIsBindingToEvent = false;
 	};
 
 } //namespace UE::MVVM
