MVVM: Fix structure permissions when iterating/displaying.

diff --git a/Engine/Plugins/Runtime/ModelViewViewModel/Source/ModelViewViewModelEditor/Private/Widgets/SMVVMViewModelBindingListWidget.cpp b/Engine/Plugins/Runtime/ModelViewViewModel/Source/ModelViewViewModelEditor/Private/Widgets/SMVVMViewModelBindingListWidget.cpp
--- a/Engine/Plugins/Runtime/ModelViewViewModel/Source/ModelViewViewModelEditor/Private/Widgets/SMVVMViewModelBindingListWidget.cpp
+++ b/Engine/Plugins/Runtime/ModelViewViewModel/Source/ModelViewViewModelEditor/Private/Widgets/SMVVMViewModelBindingListWidget.cpp
@@ -275,7 +275,8 @@ bool FFieldExpander_Bindable::CanExpandScriptStruct(const FStructProperty* Struc
 {
 	if (UE::PropertyViewer::FFieldExpander_Default::CanExpandScriptStruct(StructProperty))
 	{
-		return GetMutableDefault<UBlueprintEditorSettings>()->GetStructPermissions().PassesFilter(StructProperty->GetFullName());
+		const FPathPermissionList& StructPermissions = GetMutableDefault<UBlueprintEditorSettings>()->GetStructPermissions();
+		return !StructPermissions.HasFiltering() || StructPermissions.PassesFilter(StructProperty->Struct->GetPathName());
 	}
 	return false;
 }
