MVVM: Fix editor that show the field list when the source is not yet valid. Introduced with CL28965726

diff --git a/Engine/Plugins/Runtime/ModelViewViewModel/Source/ModelViewViewModelEditor/Private/Types/MVVMBindingSource.cpp b/Engine/Plugins/Runtime/ModelViewViewModel/Source/ModelViewViewModelEditor/Private/Types/MVVMBindingSource.cpp
--- a/Engine/Plugins/Runtime/ModelViewViewModel/Source/ModelViewViewModelEditor/Private/Types/MVVMBindingSource.cpp
+++ b/Engine/Plugins/Runtime/ModelViewViewModel/Source/ModelViewViewModelEditor/Private/Types/MVVMBindingSource.cpp
@@ -18,16 +18,11 @@ namespace UE::MVVM
 
 const UClass* FBindingSource::GetClass() const
 {
-	return IsValid() ? Class.Get() : nullptr;
+	return Class.Get();
 }
 
 FText FBindingSource::GetDisplayName() const
 {
-	if (!IsValid())
-	{
-		return LOCTEXT("None", "<None>");
-	}
-
 	return DisplayName;
 }
 
