MVVM: fix field selector permission setting not working.
Add to the list of disallowed fields.

diff --git a/Engine/Plugins/Runtime/ModelViewViewModel/Config/BaseModelViewViewModel.ini b/Engine/Plugins/Runtime/ModelViewViewModel/Config/BaseModelViewViewModel.ini
--- a/Engine/Plugins/Runtime/ModelViewViewModel/Config/BaseModelViewViewModel.ini
+++ b/Engine/Plugins/Runtime/ModelViewViewModel/Config/BaseModelViewViewModel.ini
@@ -8,4 +8,4 @@
 +PropertyRedirects=(OldName="/Script/ModelViewViewModelBlueprint.MVVMBlueprintViewBinding.WidgetPath",NewName="DestinationPath")
 
 [/Script/ModelViewViewModelBlueprint.MVVMDeveloperProjectSettings]
-+FieldSelectorPermissions=(("/Script/UMG.Widget", (DisallowedFieldNames=("TooltipWidget","ForceVolatile","SetUserFocus","Navigation","GetParent","GetOwningPlayer","GetOwningLocalPlayer","GetGameInstance"),AdvancedFieldNames=())))
\ No newline at end of file
++FieldSelectorPermissions=(("/Script/UMG.Widget", (DisallowedFieldNames=("TooltipWidget","ForceVolatile","AddExtension", "SetUserFocus","Navigation","GetParent","GetOwningLocalPlayer","GetGameInstance","BindToAnimationEvent","BindToAnimationFinished","BindToAnimationStarted","UnbindFromAnimationFinished","UnbindFromAnimationStarted","UnbindAllFromAnimationFinished","UnbindAllFromAnimationStarted","CancelLatentActions","FlushAnimations","ForceLayoutPrepass","InvalidateLayoutAndVolatility","K2_AddFieldValueChangedDelegate","ListenForInputAction","StopListeningForInputAction","StopListeningForAllInputActions","RemoveExtension","RemoveExtensions","SetAllNavigationRules","SetNavigationRuleBase","SetNavigationRuleCustom","SetNavigationRuleCustomBoundary","SetNavigationRuleExplicit","SetOwningPlayer"),AdvancedFieldNames=())), ("/Script/ModelViewViewModel.MVVMViewModelBase", (DisallowedFieldNames=("K2_AddFieldValueChangedDelegate","K2_RemoveFieldValueChangedDelegate"),AdvancedFieldNames=())))
\ No newline at end of file
diff --git a/Engine/Plugins/Runtime/ModelViewViewModel/Source/ModelViewViewModelBlueprint/Private/MVVMDeveloperProjectSettings.cpp b/Engine/Plugins/Runtime/ModelViewViewModel/Source/ModelViewViewModelBlueprint/Private/MVVMDeveloperProjectSettings.cpp
--- a/Engine/Plugins/Runtime/ModelViewViewModel/Source/ModelViewViewModelBlueprint/Private/MVVMDeveloperProjectSettings.cpp
+++ b/Engine/Plugins/Runtime/ModelViewViewModel/Source/ModelViewViewModelBlueprint/Private/MVVMDeveloperProjectSettings.cpp
@@ -53,9 +53,23 @@ bool UMVVMDeveloperProjectSettings::PropertyHasFiltering(const UStruct* ObjectSt
 	Property->GetOwnerClass()->GetPathName(nullptr, StringBuilder);
 	FSoftClassPath StructPath;
 	StructPath.SetPath(StringBuilder);
-	if (const FMVVMDeveloperProjectWidgetSettings* Settings = FieldSelectorPermissions.Find(StructPath))
+
+	if (ObjectStruct)
 	{
-		return !Settings->DisallowedFieldNames.Find(Property->GetFName());
+		for (const TPair<FSoftClassPath, FMVVMDeveloperProjectWidgetSettings>& PermissionItem : FieldSelectorPermissions)
+		{
+			if (UClass* ConcreteClass = PermissionItem.Key.ResolveClass())
+			{
+				if (ObjectStruct->IsChildOf(ConcreteClass))
+				{
+					const FMVVMDeveloperProjectWidgetSettings& Settings = PermissionItem.Value;
+					if (Settings.DisallowedFieldNames.Contains(Property->GetFName()))
+					{
+						return false;
+					}
+				}
+			}
+		}
 	}
 	return true;
 }
@@ -97,11 +111,19 @@ bool UMVVMDeveloperProjectSettings::IsPropertyAllowed(const UBlueprint* Generati
 		AuthoritativeClass->GetPathName(nullptr, StringBuilder);
 		FSoftClassPath StructPath;
 		StructPath.SetPath(StringBuilder);
-		if (const FMVVMDeveloperProjectWidgetSettings* Settings = FieldSelectorPermissions.Find(StructPath))
+
+		for (const TPair<FSoftClassPath, FMVVMDeveloperProjectWidgetSettings>& PermissionItem : FieldSelectorPermissions)
 		{
-			if (Settings->DisallowedFieldNames.Find(Property->GetFName()))
+			if (UClass* ConcreteClass = PermissionItem.Key.ResolveClass())
 			{
-				return false;
+				if (AuthoritativeClass->IsChildOf(ConcreteClass))
+				{
+					const FMVVMDeveloperProjectWidgetSettings& Settings = PermissionItem.Value;
+					if (Settings.DisallowedFieldNames.Contains(Property->GetFName()))
+					{
+						return false;
+					}
+				}
 			}
 		}
 	}
@@ -143,11 +165,19 @@ bool UMVVMDeveloperProjectSettings::IsFunctionAllowed(const UBlueprint* Generati
 		AuthoritativeClass->GetPathName(nullptr, StringBuilder);
 		FSoftClassPath StructPath;
 		StructPath.SetPath(StringBuilder);
-		if (const FMVVMDeveloperProjectWidgetSettings* Settings = FieldSelectorPermissions.Find(StructPath))
+
+		for (const TPair<FSoftClassPath, FMVVMDeveloperProjectWidgetSettings>& PermissionItem : FieldSelectorPermissions)
 		{
-			if (Settings->DisallowedFieldNames.Find(Function->GetFName()))
+			if (UClass* ConcreteClass = PermissionItem.Key.ResolveClass())
 			{
-				return false;
+				if (AuthoritativeClass->IsChildOf(ConcreteClass))
+				{
+					const FMVVMDeveloperProjectWidgetSettings& Settings = PermissionItem.Value;
+					if (Settings.DisallowedFieldNames.Contains(Function->GetFName()))
+					{
+						return false;
+					}
+				}
 			}
 		}
 	}
