MVVM: Fix loading pin values. The TrySetDefault reset the other option so the order was important. Introduced with CL30592781.

diff --git a/Engine/Plugins/Runtime/ModelViewViewModel/Source/ModelViewViewModelBlueprint/Private/MVVMBlueprintPin.cpp b/Engine/Plugins/Runtime/ModelViewViewModel/Source/ModelViewViewModelBlueprint/Private/MVVMBlueprintPin.cpp
--- a/Engine/Plugins/Runtime/ModelViewViewModel/Source/ModelViewViewModelBlueprint/Private/MVVMBlueprintPin.cpp
+++ b/Engine/Plugins/Runtime/ModelViewViewModel/Source/ModelViewViewModelBlueprint/Private/MVVMBlueprintPin.cpp
@@ -209,10 +209,22 @@ void FMVVMBlueprintPin::CopyTo(const UBlueprint* Blueprint, UEdGraphNode* Node)
 				GetDefault<UEdGraphSchema_K2>()->SplitPin(GraphPin, false);
 			}
 
-			GetDefault<UMVVMConversionFunctionGraphSchema>()->TrySetDefaultObject(*GraphPin, DefaultObject, false);
-			GetDefault<UMVVMConversionFunctionGraphSchema>()->TrySetDefaultValue(*GraphPin, DefaultString, false);
-			GetDefault<UMVVMConversionFunctionGraphSchema>()->TrySetDefaultText(*GraphPin, DefaultText, false);
-			UE::MVVM::ConversionFunctionHelper::SetPropertyPathForPin(Blueprint, Path, GraphPin);
+			if (Path.IsValid())
+			{
+				UE::MVVM::ConversionFunctionHelper::SetPropertyPathForPin(Blueprint, Path, GraphPin);
+			}
+			else if (DefaultObject)
+			{
+				GetDefault<UMVVMConversionFunctionGraphSchema>()->TrySetDefaultObject(*GraphPin, DefaultObject, false);
+			}
+			else if (!DefaultText.IsEmpty())
+			{
+				GetDefault<UMVVMConversionFunctionGraphSchema>()->TrySetDefaultText(*GraphPin, DefaultText, false);
+			}
+			else if (!DefaultString.IsEmpty())
+			{
+				GetDefault<UMVVMConversionFunctionGraphSchema>()->TrySetDefaultValue(*GraphPin, DefaultString, false);
+			}
 		}
 	}
 }
