MVVM: add tooltips for different fields of a binding in the binding panel.

diff --git a/Engine/Plugins/Runtime/ModelViewViewModel/Source/ModelViewViewModelEditor/Private/Widgets/BindingEntry/SMVVMBindingRow.cpp b/Engine/Plugins/Runtime/ModelViewViewModel/Source/ModelViewViewModelEditor/Private/Widgets/BindingEntry/SMVVMBindingRow.cpp
--- a/Engine/Plugins/Runtime/ModelViewViewModel/Source/ModelViewViewModelEditor/Private/Widgets/BindingEntry/SMVVMBindingRow.cpp
+++ b/Engine/Plugins/Runtime/ModelViewViewModel/Source/ModelViewViewModelEditor/Private/Widgets/BindingEntry/SMVVMBindingRow.cpp
@@ -125,6 +125,7 @@ TSharedRef<SWidget> SBindingRow::BuildRowWidget()
 					.OnDrop(this, &SBindingRow::HandleFieldSelectorDrop, false)
 					.OnDragEnter(this, &SBindingRow::HandleFieldSelectorDragEnter, false)
 					.ShowContext(false)
+					.ToolTipText(FText::Join(FText::FromString("\n"), LOCTEXT("MVVMDestTooltipHeader", "Destination"), LOCTEXT("MVVMDestTooltipBody", "The property to hook data into.")))
 				]
 			]
 
@@ -169,6 +170,7 @@ TSharedRef<SWidget> SBindingRow::BuildRowWidget()
 					.OnGetSelectionContext(this, &SBindingRow::GetSelectedSelectionContext, true)
 					.OnDrop(this, &SBindingRow::HandleFieldSelectorDrop, true)
 					.OnDragEnter(this, &SBindingRow::HandleFieldSelectorDragEnter, true)
+					.ToolTipText(FText::Join(FText::FromString("\n"), LOCTEXT("MVVMSrcTooltipHeader", "Source"), LOCTEXT("MVVMSrcTooltipBody", "The source of data to bind to.")))
 				]
 			]
 
@@ -688,7 +690,7 @@ FText SBindingRow::GetCurrentBindingModeLabel() const
 {
 	if (FMVVMBlueprintViewBinding* ViewBinding = GetThisViewBinding())
 	{
-		return GetBindingModeLabel(ViewBinding->BindingType);
+		return FText::Join(FText::FromString(TEXT("\n")), LOCTEXT("BindingMode", "Mode"), GetBindingModeLabel(ViewBinding->BindingType));
 	}
 	return FText::GetEmpty();
 }
diff --git a/Engine/Plugins/Runtime/ModelViewViewModel/Source/ModelViewViewModelEditor/Private/Widgets/SMVVMSourceEntry.cpp b/Engine/Plugins/Runtime/ModelViewViewModel/Source/ModelViewViewModelEditor/Private/Widgets/SMVVMSourceEntry.cpp
--- a/Engine/Plugins/Runtime/ModelViewViewModel/Source/ModelViewViewModelEditor/Private/Widgets/SMVVMSourceEntry.cpp
+++ b/Engine/Plugins/Runtime/ModelViewViewModel/Source/ModelViewViewModelEditor/Private/Widgets/SMVVMSourceEntry.cpp
@@ -50,6 +50,8 @@ FText GetSourceDisplayName(const FBindingSource& Source)
 
 void SBindingContextEntry::Construct(const FArguments& InArgs)
 {
+	ToolTipDescription = InArgs._ToolTipDescription;
+
 	ChildSlot
 	[
 		SNew(SHorizontalBox)
@@ -84,14 +86,20 @@ void SBindingContextEntry::RefreshSource(const FBindingSource& Source)
 	{
 		if (const UClass* SourceClass = Source.GetClass())
 		{
-			const FText ToolTipText = FText::Join(FText::FromString(TEXT("\n")), Source.GetDisplayName(), SourceClass->GetDisplayNameText());
+			const FText ToolTipText = ToolTipDescription.IsEmpty() ? 
+				FText::Join(FText::FromString(TEXT("\n")), Source.GetDisplayName(), SourceClass->GetDisplayNameText()) 
+				: FText::Join(FText::FromString(TEXT("\n")), ToolTipDescription, Source.GetDisplayName(), SourceClass->GetDisplayNameText());
 			SetToolTipText(ToolTipText);
 		}
 		else
 		{
-			SetToolTipText(Source.GetDisplayName());
+			SetToolTipText(ToolTipDescription.IsEmpty() ? Source.GetDisplayName() : FText::Join(FText::FromString(TEXT("\n")), ToolTipDescription, Source.GetDisplayName()));
 		}
 	}
+	else if (!ToolTipDescription.IsEmpty())
+	{
+		SetToolTipText(ToolTipDescription);
+	}
 }
 
 } // namespace UE::MVVM
diff --git a/Engine/Plugins/Runtime/ModelViewViewModel/Source/ModelViewViewModelEditor/Private/Widgets/SMVVMSourceEntry.h b/Engine/Plugins/Runtime/ModelViewViewModel/Source/ModelViewViewModelEditor/Private/Widgets/SMVVMSourceEntry.h
--- a/Engine/Plugins/Runtime/ModelViewViewModel/Source/ModelViewViewModelEditor/Private/Widgets/SMVVMSourceEntry.h
+++ b/Engine/Plugins/Runtime/ModelViewViewModel/Source/ModelViewViewModelEditor/Private/Widgets/SMVVMSourceEntry.h
@@ -23,6 +23,7 @@ public:
 		}
 		SLATE_STYLE_ARGUMENT(FTextBlockStyle, TextStyle)
 		SLATE_ARGUMENT(UE::MVVM::FBindingSource, BindingContext)
+		SLATE_ARGUMENT_DEFAULT(FText, ToolTipDescription) { FText::GetEmpty() };
 	SLATE_END_ARGS()
 
 	void Construct(const FArguments& Args);
@@ -31,6 +32,7 @@ public:
 private:
 	TSharedPtr<STextBlock> Label;
 	TSharedPtr<SImage> Image;
+	FText ToolTipDescription;
 };
 
 } // namespace UE::MVVM
diff --git a/Engine/Plugins/Runtime/ModelViewViewModel/Source/ModelViewViewModelEditor/Private/Widgets/SMVVMSourceSelector.cpp b/Engine/Plugins/Runtime/ModelViewViewModel/Source/ModelViewViewModelEditor/Private/Widgets/SMVVMSourceSelector.cpp
--- a/Engine/Plugins/Runtime/ModelViewViewModel/Source/ModelViewViewModelEditor/Private/Widgets/SMVVMSourceSelector.cpp
+++ b/Engine/Plugins/Runtime/ModelViewViewModel/Source/ModelViewViewModelEditor/Private/Widgets/SMVVMSourceSelector.cpp
@@ -55,6 +55,7 @@ void SBindingContextSelector::Construct(const FArguments& Args, const UWidgetBlu
 					SAssignNew(SelectedSourceWidget, SBindingContextEntry)
 					.BindingContext(SelectedSource)
 					.TextStyle(TextStyle)
+					.ToolTipDescription(LOCTEXT("BindingOwner", "Binding Group Owner"))
 				]
 				+ SHorizontalBox::Slot()
 				.Padding(2, 0)
