MVVM: Add a message to the empty binding panel to act as a starting tip

diff --git a/Engine/Plugins/Runtime/ModelViewViewModel/Source/ModelViewViewModelEditor/Private/Widgets/SMVVMViewBindingPanel.cpp b/Engine/Plugins/Runtime/ModelViewViewModel/Source/ModelViewViewModelEditor/Private/Widgets/SMVVMViewBindingPanel.cpp
--- a/Engine/Plugins/Runtime/ModelViewViewModel/Source/ModelViewViewModelEditor/Private/Widgets/SMVVMViewBindingPanel.cpp
+++ b/Engine/Plugins/Runtime/ModelViewViewModel/Source/ModelViewViewModelEditor/Private/Widgets/SMVVMViewBindingPanel.cpp
@@ -393,7 +393,7 @@ TSharedRef<SWidget> SBindingsPanel::CreateDrawerDockButton()
 				+ SHorizontalBox::Slot()
 				.AutoWidth()
 				.VAlign(VAlign_Center)
-				.Padding(4.0, 0.0f)
+				.Padding(4.0f, 0.0f)
 				[
 					SNew(SImage)
 					.ColorAndOpacity(FSlateColor::UseForeground())
@@ -401,7 +401,7 @@ TSharedRef<SWidget> SBindingsPanel::CreateDrawerDockButton()
 				]
 				+ SHorizontalBox::Slot()
 				.VAlign(VAlign_Center)
-				.Padding(4.0, 0.0f)
+				.Padding(4.0f, 0.0f)
 				[
 					SNew(STextBlock)
 					.Text(LOCTEXT("DockInLayout", "Dock in Layout"))
@@ -692,22 +692,11 @@ TSharedRef<SWidget> SBindingsPanel::GenerateEditViewWidget()
 		[
 			SNew(SBorder)
 			.BorderImage(FMVVMEditorStyle::Get().GetBrush("BindingView.ViewModelWarning"))
-			.Visibility_Lambda([this]()
-				{
-					if (UMVVMWidgetBlueprintExtension_View* MVVMExtensionPtr = MVVMExtension.Get())
-					{
-						if (MVVMExtensionPtr->GetBlueprintView() != nullptr && 
-							MVVMExtensionPtr->GetBlueprintView()->GetViewModels().Num() > 0)
-						{
-							return EVisibility::Collapsed;
-						}
-					}
-					return EVisibility::Visible;
-				})
+			.Visibility(this, &SBindingsPanel::GetViewModelMessageVisibility)
 			[
 				SNew(SHorizontalBox)
 				+ SHorizontalBox::Slot()
-				.Padding(20, 20, 12, 20)
+				.Padding(20.0f, 20.0f, 12.0f, 20.0f)
 				.VAlign(VAlign_Center)
 				.HAlign(HAlign_Left)
 				.AutoWidth()
@@ -716,7 +705,7 @@ TSharedRef<SWidget> SBindingsPanel::GenerateEditViewWidget()
 					.Image(FAppStyle::Get().GetBrush("Icons.Warning"))
 				]
 				+ SHorizontalBox::Slot()
-				.Padding(0, 0, 8, 0)
+				.Padding(0.0f, 0.0f, 8.0f, 0.0f)
 				.VAlign(VAlign_Center)
 				.HAlign(HAlign_Center)
 				.AutoWidth()
@@ -725,7 +714,7 @@ TSharedRef<SWidget> SBindingsPanel::GenerateEditViewWidget()
 					.Text(LOCTEXT("MissingViewModel", "This editor requires a viewmodel that widgets can bind to, would you like to add a viewmodel now?"))
 				]
 				+ SHorizontalBox::Slot()
-				.Padding(0, 0, 20, 0)
+				.Padding(0.0f, 0.0f, 20.0f, 0.0f)
 				.VAlign(VAlign_Center)
 				.HAlign(HAlign_Right)
 				[
@@ -739,6 +728,35 @@ TSharedRef<SWidget> SBindingsPanel::GenerateEditViewWidget()
 			]
 		]
 		+ SOverlay::Slot()
+		.HAlign(HAlign_Center)
+		.VAlign(VAlign_Center)
+		[
+			SNew(SBorder)
+			.BorderImage(FMVVMEditorStyle::Get().GetBrush("BindingView.ViewModelWarning"))
+			.Visibility(this, &SBindingsPanel::GetBindingMessageVisibility)
+			[
+				SNew(SHorizontalBox)
+				+ SHorizontalBox::Slot()
+				.Padding(20.0f, 20.0f, 12.0f, 20.0f)
+				.VAlign(VAlign_Center)
+				.HAlign(HAlign_Left)
+				.AutoWidth()
+				[
+					SNew(SImage)
+					.Image(FAppStyle::Get().GetBrush("Icons.Info"))
+				]
+				+ SHorizontalBox::Slot()
+				.Padding(0.0f, 0.0f, 20.0f, 0.0f)
+				.VAlign(VAlign_Center)
+				.HAlign(HAlign_Right)
+				.AutoWidth()
+				[
+					SNew(STextBlock)
+					.Text(LOCTEXT("DragWidgetCreateBinding", "Drag a widget from the Hierarchy to create a binding."))
+				]
+			]
+		]
+		+ SOverlay::Slot()
 		[
 			SNew(SBox)
 			.Visibility(this, &SBindingsPanel::GetVisibility, true)
@@ -827,6 +845,33 @@ FReply SBindingsPanel::HandleCreateViewModelClicked()
 	return FReply::Handled();
 }
 
+EVisibility SBindingsPanel::GetViewModelMessageVisibility() const
+{
+	if (UMVVMWidgetBlueprintExtension_View* MVVMExtensionPtr = MVVMExtension.Get())
+	{
+		if (MVVMExtensionPtr->GetBlueprintView() != nullptr &&
+			MVVMExtensionPtr->GetBlueprintView()->GetViewModels().Num() > 0)
+		{
+			return EVisibility::Collapsed;
+		}
+	}
+	return EVisibility::Visible;
+}
+
+EVisibility SBindingsPanel::GetBindingMessageVisibility() const
+{
+	if (UMVVMWidgetBlueprintExtension_View* MVVMExtensionPtr = MVVMExtension.Get())
+	{
+		if (MVVMExtensionPtr->GetBlueprintView() != nullptr &&
+			MVVMExtensionPtr->GetBlueprintView()->GetViewModels().Num() > 0 &&
+			MVVMExtensionPtr->GetBlueprintView()->GetBindings().Num() == 0)
+		{
+			return EVisibility::Visible;
+		}
+	}
+	return EVisibility::Collapsed;
+}
+
 } // namespace UE::MVVM
 
 #undef LOCTEXT_NAMESPACE
diff --git a/Engine/Plugins/Runtime/ModelViewViewModel/Source/ModelViewViewModelEditor/Private/Widgets/SMVVMViewBindingPanel.h b/Engine/Plugins/Runtime/ModelViewViewModel/Source/ModelViewViewModelEditor/Private/Widgets/SMVVMViewBindingPanel.h
--- a/Engine/Plugins/Runtime/ModelViewViewModel/Source/ModelViewViewModelEditor/Private/Widgets/SMVVMViewBindingPanel.h
+++ b/Engine/Plugins/Runtime/ModelViewViewModel/Source/ModelViewViewModelEditor/Private/Widgets/SMVVMViewBindingPanel.h
@@ -77,6 +77,9 @@ private:
 
 	FReply HandleCreateViewModelClicked();
 
+	EVisibility GetViewModelMessageVisibility() const;
+	EVisibility GetBindingMessageVisibility() const;
+
 	EVisibility GetVisibility(bool bVisibleWithBindings) const;
 
 	ECheckBoxState GetDetailsVisibleCheckState() const;
