MVVM: Function in another widget was not correctly tested for availability. Prevent crash when the availability test fails.

diff --git a/Engine/Plugins/Runtime/ModelViewViewModel/Source/ModelViewViewModelBlueprint/Private/MVVMDeveloperProjectSettings.cpp b/Engine/Plugins/Runtime/ModelViewViewModel/Source/ModelViewViewModelBlueprint/Private/MVVMDeveloperProjectSettings.cpp
--- a/Engine/Plugins/Runtime/ModelViewViewModel/Source/ModelViewViewModelBlueprint/Private/MVVMDeveloperProjectSettings.cpp
+++ b/Engine/Plugins/Runtime/ModelViewViewModel/Source/ModelViewViewModelBlueprint/Private/MVVMDeveloperProjectSettings.cpp
@@ -76,7 +76,7 @@ bool UMVVMDeveloperProjectSettings::PropertyHasFiltering(const UStruct* ObjectSt
 
 namespace UE::MVVM::Private
 {
-bool ShouldDoPropertyEditorPermission(const UBlueprint* GeneratingFor, const UClass* FieldOwner)
+bool ShouldDoFieldEditorPermission(const UBlueprint* GeneratingFor, const UClass* FieldOwner)
 {
 	if (GeneratingFor && FieldOwner)
 	{
@@ -96,7 +96,7 @@ bool UMVVMDeveloperProjectSettings::IsPropertyAllowed(const UBlueprint* Generati
 	const UClass* AuthoritativeClass = Cast<const UClass>(ObjectStruct);
 	AuthoritativeClass = AuthoritativeClass ? AuthoritativeClass->GetAuthoritativeClass() : nullptr;
 
-	const bool bDoPropertyEditorPermission = UE::MVVM::Private::ShouldDoPropertyEditorPermission(GeneratingFor, AuthoritativeClass);
+	const bool bDoPropertyEditorPermission = UE::MVVM::Private::ShouldDoFieldEditorPermission(GeneratingFor, AuthoritativeClass);
 	if (bDoPropertyEditorPermission)
 	{
 		if (!FPropertyEditorPermissionList::Get().DoesPropertyPassFilter(AuthoritativeClass, Property->GetFName()))
@@ -110,7 +110,7 @@ bool UMVVMDeveloperProjectSettings::IsPropertyAllowed(const UBlueprint* Generati
 		TStringBuilder<512> StringBuilder;
 		AuthoritativeClass->GetPathName(nullptr, StringBuilder);
 		FSoftClassPath StructPath;
-		StructPath.SetPath(StringBuilder);
+		StructPath.SetPath(StringBuilder.ToView());
 
 		for (const TPair<FSoftClassPath, FMVVMDeveloperProjectWidgetSettings>& PermissionItem : FieldSelectorPermissions)
 		{
@@ -142,26 +142,29 @@ bool UMVVMDeveloperProjectSettings::IsFunctionAllowed(const UBlueprint* Generati
 		return false;
 	}
 
-	TStringBuilder<512> StringBuilder;
-	const FPathPermissionList& FunctionPermissions = GetMutableDefault<UBlueprintEditorSettings>()->GetFunctionPermissions();
-	if (FunctionPermissions.HasFiltering())
+	const bool bDoPropertyEditorPermission = UE::MVVM::Private::ShouldDoFieldEditorPermission(GeneratingFor, AuthoritativeClass);
+	if (bDoPropertyEditorPermission)
 	{
-		const UFunction* FunctionToTest = AuthoritativeClass->FindFunctionByName(Function->GetFName());
-		if (FunctionToTest == nullptr)
+		const FPathPermissionList& FunctionPermissions = GetMutableDefault<UBlueprintEditorSettings>()->GetFunctionPermissions();
+		if (FunctionPermissions.HasFiltering())
 		{
-			return false;
-		}
+			const UFunction* FunctionToTest = AuthoritativeClass->FindFunctionByName(Function->GetFName());
+			if (FunctionToTest == nullptr)
+			{
+				return false;
+			}
 
-		StringBuilder.Reset();
-		FunctionToTest->GetPathName(nullptr, StringBuilder);
-		if (!FunctionPermissions.PassesFilter(StringBuilder.ToView()))
-		{
-			return false;
+			TStringBuilder<512> StringBuilder;
+			FunctionToTest->GetPathName(nullptr, StringBuilder);
+			if (!FunctionPermissions.PassesFilter(StringBuilder.ToView()))
+			{
+				return false;
+			}
 		}
 	}
 
 	{
-		StringBuilder.Reset();
+		TStringBuilder<512> StringBuilder;
 		AuthoritativeClass->GetPathName(nullptr, StringBuilder);
 		FSoftClassPath StructPath;
 		StructPath.SetPath(StringBuilder);
diff --git a/Engine/Plugins/Runtime/ModelViewViewModel/Source/ModelViewViewModelBlueprint/Private/MVVMViewBlueprintCompiler.cpp b/Engine/Plugins/Runtime/ModelViewViewModel/Source/ModelViewViewModelBlueprint/Private/MVVMViewBlueprintCompiler.cpp
--- a/Engine/Plugins/Runtime/ModelViewViewModel/Source/ModelViewViewModelBlueprint/Private/MVVMViewBlueprintCompiler.cpp
+++ b/Engine/Plugins/Runtime/ModelViewViewModel/Source/ModelViewViewModelBlueprint/Private/MVVMViewBlueprintCompiler.cpp
@@ -1,12 +1,15 @@
 // Copyright Epic Games, Inc. All Rights Reserved.
 
 #include "MVVMViewBlueprintCompiler.h"
-#include "Blueprint/WidgetTree.h"
+
 #include "Bindings/MVVMBindingHelper.h"
 #include "Bindings/MVVMConversionFunctionHelper.h"
 #include "Bindings/MVVMFieldPathHelper.h"
+#include "BlueprintEditorSettings.h"
+#include "Blueprint/WidgetTree.h"
 #include "Components/Widget.h"
 #include "HAL/IConsoleManager.h"
+#include "Misc/NamePermissionList.h"
 #include "MVVMBlueprintView.h"
 #include "MVVMBlueprintViewConversionFunction.h"
 #include "MVVMDeveloperProjectSettings.h"
@@ -383,187 +386,216 @@ void FMVVMViewBlueprintCompiler::CreateVariables(const FWidgetBlueprintCompilerC
 		CreatePublicFunctionsDeclaration(Context);
 	}
 
-	auto CreateVariable = [&Context](const FCompilerUserWidgetProperty& UserWidgetProperty) -> FProperty*
+	// Create variable
 	{
-		FEdGraphPinType NewPropertyPinType(UEdGraphSchema_K2::PC_Object, NAME_None, UserWidgetProperty.AuthoritativeClass, EPinContainerType::None, false, FEdGraphTerminalType());
-		FProperty* NewProperty = Context.CreateVariable(UserWidgetProperty.Name, NewPropertyPinType);
-		if (NewProperty != nullptr)
+		auto CreateVariable = [&Context](const FCompilerUserWidgetProperty& UserWidgetProperty) -> FProperty*
 		{
-			NewProperty->SetPropertyFlags(CPF_BlueprintVisible | CPF_RepSkip);
-
-			const bool bIsInstanceExposed = UserWidgetProperty.bInstanced && UserWidgetProperty.bInstanceExposed;
-			if (bIsInstanceExposed)
+			FEdGraphPinType NewPropertyPinType(UEdGraphSchema_K2::PC_Object, NAME_None, UserWidgetProperty.AuthoritativeClass, EPinContainerType::None, false, FEdGraphTerminalType());
+			FProperty* NewProperty = Context.CreateVariable(UserWidgetProperty.Name, NewPropertyPinType);
+			if (NewProperty != nullptr)
 			{
-				NewProperty->SetPropertyFlags(CPF_Edit | CPF_ExportObject | CPF_PersistentInstance | CPF_InstancedReference | CPF_NonNullable | CPF_NoClear);
+				NewProperty->SetPropertyFlags(CPF_BlueprintVisible | CPF_RepSkip);
 
-				if (UserWidgetProperty.AuthoritativeClass->HasAnyClassFlags(CLASS_HasInstancedReference))
+				const bool bIsInstanceExposed = UserWidgetProperty.bInstanced && UserWidgetProperty.bInstanceExposed;
+				if (bIsInstanceExposed)
 				{
-					NewProperty->SetPropertyFlags(CPF_ContainsInstancedReference);
+					NewProperty->SetPropertyFlags(CPF_Edit | CPF_ExportObject | CPF_PersistentInstance | CPF_InstancedReference | CPF_NonNullable | CPF_NoClear);
+
+					if (UserWidgetProperty.AuthoritativeClass->HasAnyClassFlags(CLASS_HasInstancedReference))
+					{
+						NewProperty->SetPropertyFlags(CPF_ContainsInstancedReference);
+					}
+				}
+				else
+				{
+					NewProperty->SetPropertyFlags(CPF_Transient | CPF_DuplicateTransient);
 				}
-			}
-			else
-			{
-				NewProperty->SetPropertyFlags(CPF_Transient | CPF_DuplicateTransient);
-			}
 
-			if (UserWidgetProperty.bExposeOnSpawn)
-			{
-				NewProperty->SetPropertyFlags(CPF_ExposeOnSpawn);
-			}
-			else if (!bIsInstanceExposed)
-			{
-				NewProperty->SetPropertyFlags(CPF_DisableEditOnInstance);
-			}
+				if (UserWidgetProperty.bExposeOnSpawn)
+				{
+					NewProperty->SetPropertyFlags(CPF_ExposeOnSpawn);
+				}
+				else if (!bIsInstanceExposed)
+				{
+					NewProperty->SetPropertyFlags(CPF_DisableEditOnInstance);
+				}
 
-			if (UserWidgetProperty.bReadOnly)
-			{
-				NewProperty->SetPropertyFlags(CPF_BlueprintReadOnly);
-			}
+				if (UserWidgetProperty.bReadOnly)
+				{
+					NewProperty->SetPropertyFlags(CPF_BlueprintReadOnly);
+				}
 
 #if WITH_EDITOR
-			if (!UserWidgetProperty.BlueprintSetter.IsEmpty())
-			{
-				NewProperty->SetMetaData(FBlueprintMetadata::MD_PropertySetFunction, *UserWidgetProperty.BlueprintSetter);
-			}
-			if (!UserWidgetProperty.DisplayName.IsEmpty())
-			{
-				NewProperty->SetMetaData(FBlueprintMetadata::MD_DisplayName, *UserWidgetProperty.DisplayName.ToString());
-			}
-			if (!UserWidgetProperty.CategoryName.IsEmpty())
-			{
-				NewProperty->SetMetaData(FBlueprintMetadata::MD_FunctionCategory, *UserWidgetProperty.CategoryName);
-			}
-			if (UserWidgetProperty.bExposeOnSpawn)
-			{
-				NewProperty->SetMetaData(FBlueprintMetadata::MD_ExposeOnSpawn, TEXT("true"));
-			}
-			if (UserWidgetProperty.bPrivate)
-			{
-				NewProperty->SetMetaData(FBlueprintMetadata::MD_Private, TEXT("true"));
-			}
-			if (bIsInstanceExposed)
-			{
-				NewProperty->SetMetaData(FName("EditInline"), TEXT("true"));
-			}
+				if (!UserWidgetProperty.BlueprintSetter.IsEmpty())
+				{
+					NewProperty->SetMetaData(FBlueprintMetadata::MD_PropertySetFunction, *UserWidgetProperty.BlueprintSetter);
+				}
+				if (!UserWidgetProperty.DisplayName.IsEmpty())
+				{
+					NewProperty->SetMetaData(FBlueprintMetadata::MD_DisplayName, *UserWidgetProperty.DisplayName.ToString());
+				}
+				if (!UserWidgetProperty.CategoryName.IsEmpty())
+				{
+					NewProperty->SetMetaData(FBlueprintMetadata::MD_FunctionCategory, *UserWidgetProperty.CategoryName);
+				}
+				if (UserWidgetProperty.bExposeOnSpawn)
+				{
+					NewProperty->SetMetaData(FBlueprintMetadata::MD_ExposeOnSpawn, TEXT("true"));
+				}
+				if (UserWidgetProperty.bPrivate)
+				{
+					NewProperty->SetMetaData(FBlueprintMetadata::MD_Private, TEXT("true"));
+				}
+				if (bIsInstanceExposed)
+				{
+					NewProperty->SetMetaData(FName("EditInline"), TEXT("true"));
+				}
 #endif
-		}
-		return NewProperty;
-	};
+			}
+			return NewProperty;
+		};
 
-	for (FCompilerUserWidgetProperty& UserWidgetProperty : NeededUserWidgetProperties)
-	{
-		check(UserWidgetProperty.AuthoritativeClass);
-		check(!UserWidgetProperty.Name.IsNone());
-		UserWidgetProperty.Property = nullptr; // Skeletal set the property, Full needs the new property
+		for (FCompilerUserWidgetProperty& UserWidgetProperty : NeededUserWidgetProperties)
+		{
+			check(UserWidgetProperty.AuthoritativeClass);
+			check(!UserWidgetProperty.Name.IsNone());
+			UserWidgetProperty.Property = nullptr; // Skeletal set the property, Full needs the new property
 
-		FMVVMConstFieldVariant UserWidgetPropertyField = BindingHelper::FindFieldByName(Context.GetGeneratedClass(), FMVVMBindingName(UserWidgetProperty.Name));
+			FMVVMConstFieldVariant UserWidgetPropertyField = BindingHelper::FindFieldByName(Context.GetGeneratedClass(), FMVVMBindingName(UserWidgetProperty.Name));
 
-		// The class is not linked yet. It may not be available yet.
-		if (UserWidgetPropertyField.IsEmpty())
-		{
-			for (FField* Field = Context.GetGeneratedClass()->ChildProperties; Field != nullptr; Field = Field->Next)
+			// The class is not linked yet. It may not be available yet.
+			if (UserWidgetPropertyField.IsEmpty())
 			{
-				if (Field->GetFName() == UserWidgetProperty.Name)
+				for (FField* Field = Context.GetGeneratedClass()->ChildProperties; Field != nullptr; Field = Field->Next)
 				{
-					if (CastField<FProperty>(Field))
+					if (Field->GetFName() == UserWidgetProperty.Name)
 					{
-						UserWidgetPropertyField = FMVVMFieldVariant(CastField<FProperty>(Field));
+						if (CastField<FProperty>(Field))
+						{
+							UserWidgetPropertyField = FMVVMFieldVariant(CastField<FProperty>(Field));
+						}
+						else
+						{
+							WidgetBlueprintCompilerContext.MessageLog.Error(*FText::Format(LOCTEXT("FieldIsNotProperty", "The field for source '{0}' exists but is not a property."), UserWidgetProperty.DisplayName).ToString());
+							bIsCreateVariableStepValid = false;
+						}
+						break;
 					}
-					else
+				}
+				for (UField* Field = Context.GetGeneratedClass()->Children; Field != nullptr; Field = Field->Next)
+				{
+					if (Field->GetFName() == UserWidgetProperty.Name)
 					{
 						WidgetBlueprintCompilerContext.MessageLog.Error(*FText::Format(LOCTEXT("FieldIsNotProperty", "The field for source '{0}' exists but is not a property."), UserWidgetProperty.DisplayName).ToString());
 						bIsCreateVariableStepValid = false;
+						break;
 					}
-					break;
 				}
 			}
-			for (UField* Field = Context.GetGeneratedClass()->Children; Field != nullptr; Field = Field->Next)
+
+			if (UserWidgetPropertyField.IsEmpty())
 			{
-				if (Field->GetFName() == UserWidgetProperty.Name)
+				UClass* ParentClass = Context.GetGeneratedClass()->GetSuperClass();
+				if (const FProperty* Property = ParentClass->FindPropertyByName(UserWidgetProperty.Name))
 				{
-					WidgetBlueprintCompilerContext.MessageLog.Error(*FText::Format(LOCTEXT("FieldIsNotProperty", "The field for source '{0}' exists but is not a property."), UserWidgetProperty.DisplayName).ToString());
-					bIsCreateVariableStepValid = false;
-					break;
+					UserWidgetPropertyField = FMVVMFieldVariant(Property);
 				}
 			}
-		}
 
-		if (UserWidgetPropertyField.IsEmpty())
-		{
-			UClass* ParentClass = Context.GetGeneratedClass()->GetSuperClass();
-			if (const FProperty* Property = ParentClass->FindPropertyByName(UserWidgetProperty.Name))
+			// Will always create viewmodel properties.
+			// Will never create properties for animation or other Self.Object
+			// Will create properties for widget when they are not already created.
+			if (UserWidgetProperty.CreationType == FCompilerUserWidgetProperty::ECreationType::CreateOnlyIfDoesntExist && !UserWidgetPropertyField.IsEmpty())
 			{
-				UserWidgetPropertyField = FMVVMFieldVariant(Property);
-			}
-		}
-
-		// Will always create viewmodel properties.
-		// Will never create properties for animation or other Self.Object
-		// Will create properties for widget when they are not already created.
-		if (UserWidgetProperty.CreationType == FCompilerUserWidgetProperty::ECreationType::CreateOnlyIfDoesntExist && !UserWidgetPropertyField.IsEmpty())
-		{
-			// Viewmodel property cannot already exist. It will creates issue with initialization and with View::SetViewModel.
-			const UClass* OwnerClass = Cast<UClass>(UserWidgetPropertyField.GetOwner());
-			WidgetBlueprintCompilerContext.MessageLog.Error(*FText::Format(LOCTEXT("PropertyAlreadyExistInParent", "There is already a property named '{0}' in scope '{1}'."), UserWidgetProperty.DisplayName, (OwnerClass ? OwnerClass->GetDisplayNameText() : FText::GetEmpty())).ToString());
-			bIsCreateVariableStepValid = false;
-			continue;
-		}
-
-		if (!UserWidgetPropertyField.IsEmpty())
-		{
-			if (UserWidgetPropertyField.IsFunction())
-			{
-				WidgetBlueprintCompilerContext.MessageLog.Error(*FText::Format(LOCTEXT("FunctionCanBeSource", "Function can't be source. '{0}'."), UserWidgetProperty.DisplayName).ToString());
+				// Viewmodel property cannot already exist. It will creates issue with initialization and with View::SetViewModel.
+				const UClass* OwnerClass = Cast<UClass>(UserWidgetPropertyField.GetOwner());
+				WidgetBlueprintCompilerContext.MessageLog.Error(*FText::Format(LOCTEXT("PropertyAlreadyExistInParent", "There is already a property named '{0}' in scope '{1}'."), UserWidgetProperty.DisplayName, (OwnerClass ? OwnerClass->GetDisplayNameText() : FText::GetEmpty())).ToString());
 				bIsCreateVariableStepValid = false;
 				continue;
 			}
 
-			if (!BindingHelper::IsValidForSourceBinding(UserWidgetPropertyField))
+			if (!UserWidgetPropertyField.IsEmpty())
 			{
-				WidgetBlueprintCompilerContext.MessageLog.Error(*FText::Format(LOCTEXT("FieldNotAccessibleAtRuntime", "The field for source '{0}' exists but is not accessible at runtime."), UserWidgetProperty.DisplayName).ToString());
-				bIsCreateVariableStepValid = false;
-				continue;
+				if (UserWidgetPropertyField.IsFunction())
+				{
+					WidgetBlueprintCompilerContext.MessageLog.Error(*FText::Format(LOCTEXT("FunctionCanBeSource", "Function can't be source. '{0}'."), UserWidgetProperty.DisplayName).ToString());
+					bIsCreateVariableStepValid = false;
+					continue;
+				}
+
+				if (!BindingHelper::IsValidForSourceBinding(UserWidgetPropertyField))
+				{
+					WidgetBlueprintCompilerContext.MessageLog.Error(*FText::Format(LOCTEXT("FieldNotAccessibleAtRuntime", "The field for source '{0}' exists but is not accessible at runtime."), UserWidgetProperty.DisplayName).ToString());
+					bIsCreateVariableStepValid = false;
+					continue;
+				}
+
+				ensure(UserWidgetPropertyField.IsProperty());
+				const FProperty* Property = UserWidgetPropertyField.IsProperty() ? UserWidgetPropertyField.GetProperty() : nullptr;
+				const FObjectProperty* ObjectProperty = CastField<const FObjectProperty>(Property);
+				const bool bIsCompatible = ObjectProperty && UserWidgetProperty.AuthoritativeClass->IsChildOf(ObjectProperty->PropertyClass);
+				if (!bIsCompatible)
+				{
+					WidgetBlueprintCompilerContext.MessageLog.Error(*FText::Format(LOCTEXT("PropertyExistsAndNotCompatible", "There is already a property named '{0}' that is not compatible with the source of the same name."), UserWidgetProperty.DisplayName).ToString());
+					bIsCreateVariableStepValid = false;
+					continue;
+				}
+
+				const bool bIsBindWidget = FWidgetBlueprintEditorUtils::IsBindWidgetProperty(ObjectProperty);
+				if (Context.GetGeneratedClass() != ObjectProperty->GetOwnerStruct() && !bIsBindWidget)
+				{
+					// Widget needs to be BindWidget to be reused as a property.
+					const UClass* OwnerClass = Cast<UClass>(ObjectProperty->GetOwnerStruct());
+					WidgetBlueprintCompilerContext.MessageLog.Error(*FText::Format(LOCTEXT("WidgetPropertyAlreadyExist", "There is already a property named '{0}' in scope '{1}' for the widget. Are you missing a BindWidget?."), UserWidgetProperty.DisplayName, (OwnerClass ? OwnerClass->GetDisplayNameText() : FText::GetEmpty())).ToString());
+					bIsCreateVariableStepValid = false;
+					continue;
+				}
 			}
 
-			ensure(UserWidgetPropertyField.IsProperty());
-			const FProperty* Property = UserWidgetPropertyField.IsProperty() ? UserWidgetPropertyField.GetProperty() : nullptr;
-			const FObjectProperty* ObjectProperty = CastField<const FObjectProperty>(Property);
-			const bool bIsCompatible = ObjectProperty && UserWidgetProperty.AuthoritativeClass->IsChildOf(ObjectProperty->PropertyClass);
-			if (!bIsCompatible)
+			// Can we reused the property or we need to create a new one.
+			bool bCreateVariable = UserWidgetProperty.CreationType == FCompilerUserWidgetProperty::ECreationType::CreateOnlyIfDoesntExist
+				|| (UserWidgetProperty.CreationType == FCompilerUserWidgetProperty::ECreationType::CreateIfDoesntExist && UserWidgetPropertyField.IsEmpty());
+			if (bCreateVariable)
 			{
-				WidgetBlueprintCompilerContext.MessageLog.Error(*FText::Format(LOCTEXT("PropertyExistsAndNotCompatible", "There is already a property named '{0}' that is not compatible with the source of the same name."), UserWidgetProperty.DisplayName).ToString());
-				bIsCreateVariableStepValid = false;
-				continue;
+				UserWidgetProperty.Property = CreateVariable(UserWidgetProperty);
+			}
+			else if (UserWidgetPropertyField.IsProperty())
+			{
+				UserWidgetProperty.Property = UserWidgetPropertyField.GetProperty();
 			}
 
-			const bool bIsBindWidget = FWidgetBlueprintEditorUtils::IsBindWidgetProperty(ObjectProperty);
-			if (Context.GetGeneratedClass() != ObjectProperty->GetOwnerStruct() && !bIsBindWidget)
+			if (UserWidgetProperty.Property == nullptr)
 			{
-				// Widget needs to be BindWidget to be reused as a property.
-				const UClass* OwnerClass = Cast<UClass>(ObjectProperty->GetOwnerStruct());
-				WidgetBlueprintCompilerContext.MessageLog.Error(*FText::Format(LOCTEXT("WidgetPropertyAlreadyExist", "There is already a property named '{0}' in scope '{1}' for the widget. Are you missing a BindWidget?."), UserWidgetProperty.DisplayName, (OwnerClass ? OwnerClass->GetDisplayNameText() : FText::GetEmpty())).ToString());
+				WidgetBlueprintCompilerContext.MessageLog.Error(*FText::Format(LOCTEXT("VariableCouldNotBeCreated", "The variable for '{0}' could not be created."), UserWidgetProperty.DisplayName).ToString());
 				bIsCreateVariableStepValid = false;
 				continue;
 			}
 		}
+	}
 
-		// Can we reused the property or we need to create a new one.
-		bool bCreateVariable = UserWidgetProperty.CreationType == FCompilerUserWidgetProperty::ECreationType::CreateOnlyIfDoesntExist
-			|| (UserWidgetProperty.CreationType == FCompilerUserWidgetProperty::ECreationType::CreateIfDoesntExist && UserWidgetPropertyField.IsEmpty());
-		if (bCreateVariable)
-		{
-			UserWidgetProperty.Property = CreateVariable(UserWidgetProperty);
-		}
-		else if (UserWidgetPropertyField.IsProperty())
+	// Public function permissions
+	{
+		FPathPermissionList& FunctionPermissions = GetMutableDefault<UBlueprintEditorSettings>()->GetFunctionPermissions();
+		FName ContextName;
 		{
-			UserWidgetProperty.Property = UserWidgetPropertyField.GetProperty();
+			TStringBuilder<512> ContextNameStr;
+			Context.GetGeneratedClass()->GetPathName(nullptr, ContextNameStr);
+			ContextName = ContextNameStr.ToString();
 		}
 
-		if (UserWidgetProperty.Property == nullptr)
+		FunctionPermissions.UnregisterOwner(ContextName);
+
+		const bool bHasFiltering = GetMutableDefault<UBlueprintEditorSettings>()->GetFunctionPermissions().HasFiltering();
+		if (bHasFiltering)
 		{
-			WidgetBlueprintCompilerContext.MessageLog.Error(*FText::Format(LOCTEXT("VariableCouldNotBeCreated", "The variable for '{0}' could not be created."), UserWidgetProperty.DisplayName).ToString());
-			bIsCreateVariableStepValid = false;
-			continue;
+			for (FName FunctionName : FunctionPermissionsToAdd)
+			{
+				TStringBuilder<512> FunctionPath;
+				FunctionPath << ContextName;
+				FunctionPath << TEXT(':');
+				FunctionPath << FunctionName;
+				FunctionPermissions.AddAllowListItem(ContextName, FunctionPath);
+			}
 		}
 	}
 }
@@ -854,7 +886,8 @@ void FMVVMViewBlueprintCompiler::CreateRequiredProperties(const FWidgetBlueprint
 					SourceVariable.CategoryName = TEXT("Widget");
 					SourceVariable.CreationType = FCompilerUserWidgetProperty::ECreationType::CreateIfDoesntExist;
 					SourceVariable.bExposeOnSpawn = false;
-					//SourceVariable.bPrivate = true; Remove until the data is fix properly
+					//todo Remove until the data is fix properly
+					//SourceVariable.bPrivate = true;
 					SourceVariable.bPrivate = false;
 					SourceVariable.bReadOnly = true;
 					SourceVariable.bInstanced = false;
@@ -1000,7 +1033,6 @@ void FMVVMViewBlueprintCompiler::CreatePublicFunctionsDeclaration(const FWidgetB
 				, (FUNC_BlueprintCallable | FUNC_Public | FUNC_Final)
 				, TEXT("Viewmodel")
 				, false);
-			BlueprintView->TemporaryGraph.Add(Setter.SetterGraph);
 
 			if (Setter.SetterGraph == nullptr || Setter.SetterGraph->GetFName() != FName(*Setter.BlueprintSetter))
 			{
@@ -1012,6 +1044,9 @@ void FMVVMViewBlueprintCompiler::CreatePublicFunctionsDeclaration(const FWidgetB
 				continue;
 			}
 
+			BlueprintView->TemporaryGraph.Add(Setter.SetterGraph);
+			FunctionPermissionsToAdd.Add(Setter.SetterGraph->GetFName());
+
 			UE::MVVM::FunctionGraphHelper::AddFunctionArgument(Setter.SetterGraph, const_cast<UClass*>(Setter.Class), "Viewmodel");
 		}
 	}
@@ -1182,7 +1217,7 @@ void FMVVMViewBlueprintCompiler::CreateWriteFieldContexts(const FWidgetBlueprint
 					});
 				if (Found)
 				{
-					// Temporary removing this message until the assets are fixed.
+					// todo Temporary removing this message until the assets are fixed.
 					//AddMessageForBinding(Binding
 					//	, FText::Format(LOCTEXT("PropertyPathAlreadyUsed", "The property path '{0}' is already used by another binding."), PropertyPathToText(NewSkeletonClass, BlueprintView.Get(), DestinationPropertyPath))
 					//	, EMessageType::Warning
@@ -2986,10 +3021,14 @@ TValueOrError<FMVVMViewBlueprintCompiler::FCreateFieldsResult, FText> FMVVMViewB
 
 	// Generate the path with property converted to BP function
 	TValueOrError<TArray<FMVVMConstFieldVariant>, FText> SkeletalGeneratedFieldsResult = FieldPathHelper::GenerateFieldPathList(Result.GeneratedFields, bForSourceReading);
-	if (SkeletalGeneratedFieldsResult.HasError() || !IsPropertyPathValid(WidgetBlueprintCompilerContext.WidgetBlueprint(), SkeletalGeneratedFieldsResult.GetValue()))
+	if (SkeletalGeneratedFieldsResult.HasError())
 	{
 		return MakeError(FText::Format(Private::CouldNotCreateSourceFieldPathFormat, PropertyPathToText(Class, BlueprintView.Get(), PropertyPath), SkeletalGeneratedFieldsResult.GetError()));
 	}
+	if (!IsPropertyPathValid(WidgetBlueprintCompilerContext.WidgetBlueprint(), SkeletalGeneratedFieldsResult.GetValue()))
+	{
+		return MakeError(FText::Format(Private::CouldNotCreateSourceFieldPathFormat, PropertyPathToText(Class, BlueprintView.Get(), PropertyPath), LOCTEXT("NotAValidPropertyPath", "The path is not valid.")));
+	}
 	Result.SkeletalGeneratedFields = SkeletalGeneratedFieldsResult.StealValue();
 
 	return MakeValue(Result);
diff --git a/Engine/Plugins/Runtime/ModelViewViewModel/Source/ModelViewViewModelBlueprint/Private/MVVMViewBlueprintCompiler.h b/Engine/Plugins/Runtime/ModelViewViewModel/Source/ModelViewViewModelBlueprint/Private/MVVMViewBlueprintCompiler.h
--- a/Engine/Plugins/Runtime/ModelViewViewModel/Source/ModelViewViewModelBlueprint/Private/MVVMViewBlueprintCompiler.h
+++ b/Engine/Plugins/Runtime/ModelViewViewModel/Source/ModelViewViewModelBlueprint/Private/MVVMViewBlueprintCompiler.h
@@ -330,6 +330,12 @@ private:
 	};
 	TArray<TSharedRef<FCompilerEvent>> ValidEvents;
 
+	/**
+	 * List of public expose function
+	 */
+	TArray<FName> FunctionPermissionsToAdd;
+
+private:
 	TMap<FName, UWidget*> WidgetNameToWidgetPointerMap;
 	FWidgetBlueprintCompilerContext& WidgetBlueprintCompilerContext;
 	TStrongObjectPtr<UMVVMBlueprintView> BlueprintView = nullptr;
