Generated widget properties are private unless specified "as variable" in the DetailView. This is to prevent mistakes and to not break encapsulation.
#tests full cook

diff --git a/Engine/Plugins/Runtime/ModelViewViewModel/Source/ModelViewViewModelBlueprint/Private/MVVMViewBlueprintCompiler.cpp b/Engine/Plugins/Runtime/ModelViewViewModel/Source/ModelViewViewModelBlueprint/Private/MVVMViewBlueprintCompiler.cpp
--- a/Engine/Plugins/Runtime/ModelViewViewModel/Source/ModelViewViewModelBlueprint/Private/MVVMViewBlueprintCompiler.cpp
+++ b/Engine/Plugins/Runtime/ModelViewViewModel/Source/ModelViewViewModelBlueprint/Private/MVVMViewBlueprintCompiler.cpp
@@ -886,8 +886,7 @@ void FMVVMViewBlueprintCompiler::CreateRequiredProperties(const FWidgetBlueprint
 					SourceVariable.CategoryName = TEXT("Widget");
 					SourceVariable.CreationType = FCompilerUserWidgetProperty::ECreationType::CreateIfDoesntExist;
 					SourceVariable.bExposeOnSpawn = false;
-					//todo Remove until the data is fix properly
-					//SourceVariable.bPrivate = true;
+					SourceVariable.bPrivate = true;
 					SourceVariable.bPrivate = false;
 					SourceVariable.bReadOnly = true;
 					SourceVariable.bInstanced = false;
