MVVM: Fix ensure that warn when there is no delayed bindings to execute. Fix ExecuteAtInit flag when one of the field is not a FieldNotify. The ensure was a false positif and the flag was not removed correctly on the second to last source. Update the ToString method to better identify compiler problems in the future.

diff --git a/Engine/Plugins/Runtime/ModelViewViewModel/Source/ModelViewViewModel/Private/View/MVVMBindingSubsystem.cpp b/Engine/Plugins/Runtime/ModelViewViewModel/Source/ModelViewViewModel/Private/View/MVVMBindingSubsystem.cpp
--- a/Engine/Plugins/Runtime/ModelViewViewModel/Source/ModelViewViewModel/Private/View/MVVMBindingSubsystem.cpp
+++ b/Engine/Plugins/Runtime/ModelViewViewModel/Source/ModelViewViewModel/Private/View/MVVMBindingSubsystem.cpp
@@ -167,5 +167,10 @@ void UMVVMBindingSubsystem::RemoveDelayedBindings(const UMVVMView* View, FMVVMVi
 				(*FoundView).RemoveAtSwap(Index);
 			}
 		}
+
+		if (FoundView->Num() == 0)
+		{
+			DelayedBindings.Remove(View);
+		}
 	}
 }
diff --git a/Engine/Plugins/Runtime/ModelViewViewModel/Source/ModelViewViewModel/Private/View/MVVMView.cpp b/Engine/Plugins/Runtime/ModelViewViewModel/Source/ModelViewViewModel/Private/View/MVVMView.cpp
--- a/Engine/Plugins/Runtime/ModelViewViewModel/Source/ModelViewViewModel/Private/View/MVVMView.cpp
+++ b/Engine/Plugins/Runtime/ModelViewViewModel/Source/ModelViewViewModel/Private/View/MVVMView.cpp
@@ -341,8 +341,6 @@ void UMVVMView::InitializeSourceBindings(FMVVMView_SourceKey SourceKey, bool bRu
 			{
 				const FMVVMViewClass_Binding& ClassBinding = ClassExtension->GetBinding(SourceBinding.GetBindingKey());
 				ExecuteBindingImmediately(ClassBinding, SourceBinding.GetBindingKey());
-
-				ensureMsgf(SourceBinding.GetFieldId().IsValid() == ClassBinding.IsOneWay(), TEXT("If the binding has no id, then it must be a OneTime."));
 			}
 		}
 
diff --git a/Engine/Plugins/Runtime/ModelViewViewModel/Source/ModelViewViewModel/Private/View/MVVMViewClass.cpp b/Engine/Plugins/Runtime/ModelViewViewModel/Source/ModelViewViewModel/Private/View/MVVMViewClass.cpp
--- a/Engine/Plugins/Runtime/ModelViewViewModel/Source/ModelViewViewModel/Private/View/MVVMViewClass.cpp
+++ b/Engine/Plugins/Runtime/ModelViewViewModel/Source/ModelViewViewModel/Private/View/MVVMViewClass.cpp
@@ -189,10 +189,8 @@ FString FMVVMViewClass_Source::ToString(const UMVVMViewClass* ViewClass, FToStri
 	StringBuilder << TEXT(", SourceName: ");
 	StringBuilder << GetName();
 
-	if (Args.bAddFlags)
-	{
-		bool bHasFlag = false;
-		auto AddPipe = [&StringBuilder, &bHasFlag]()
+	bool bHasFlag = false;
+	auto AddPipe = [&StringBuilder, &bHasFlag]()
 		{
 			if (bHasFlag)
 			{
@@ -201,7 +199,39 @@ FString FMVVMViewClass_Source::ToString(const UMVVMViewClass* ViewClass, FToStri
 			bHasFlag = true;
 		};
 
-		StringBuilder << TEXT(", Flags: ");
+	if (Args.bAddFields)
+	{
+		StringBuilder << TEXT("\n    FieldIds: ");
+		bHasFlag = false;
+		for (const FMVVMViewClass_FieldId& FieldId : GetFieldIds())
+		{
+			AddPipe();
+			StringBuilder << FieldId.GetName();
+		}
+	}
+
+	if (Args.bAddBindings)
+	{
+		StringBuilder << TEXT("\n    Bindings: ");
+		bHasFlag = false;
+		for (const FMVVMViewClass_SourceBinding& SourceBinding : GetBindings())
+		{
+			StringBuilder << TEXT("\n        FieldId: ");
+			StringBuilder << SourceBinding.GetFieldId().GetFieldName();
+			StringBuilder << TEXT(" BindingIndex: ");
+			StringBuilder << SourceBinding.GetBindingKey().GetIndex();
+			StringBuilder << TEXT(" Flags: ");
+			if (SourceBinding.ExecuteAtInitialization())
+			{
+				StringBuilder << TEXT("ExecuteAtInit ");
+			}
+		}
+	}
+
+	if (Args.bAddFlags)
+	{
+		StringBuilder << TEXT("\n    Flags: ");
+		bHasFlag = false;
 		if (IsUserWidget())
 		{
 			AddPipe();
@@ -455,8 +485,8 @@ FString FMVVMViewClass_Binding::ToString(const UMVVMViewClass* ViewClass, FToStr
 		}
 	}
 
-	StringBuilder << TEXT("Sources: ");
-	StringBuilder << SourceBitField;
+	StringBuilder << TEXT("Sources: 0x");
+	StringBuilder << FString::Printf(TEXT("%x"), SourceBitField);
 
 	if (Args.bAddFlags)
 	{
@@ -736,30 +766,39 @@ FString UMVVMViewClass::ToString(FToStringArgs Args) const
 
 	TStringBuilder<2048> Builder;
 	Builder << TEXT("Sources: ");
-	for (const FMVVMViewClass_Source& Source : GetSources())
+	for (int32 Index = 0; Index < GetSources().Num(); ++Index)
 	{
+		const FMVVMViewClass_Source& Source = GetSources()[Index];
 		Builder << TEXT("\n");
+		Builder << TEXT('(') << Index << TEXT(')');
 		Builder << Source.ToString(this, Args.Source);
 	}
 
-	Builder << TEXT("Bindings: ");
-	for (const FMVVMViewClass_Binding& Binding : GetBindings())
+	Builder << TEXT("\nBindings: ");
+	Builder << TEXT("Sources: ");
+	for (int32 Index = 0; Index < GetBindings().Num(); ++Index)
 	{
+		const FMVVMViewClass_Binding& Binding = GetBindings()[Index];
 		Builder << TEXT("\n");
+		Builder << TEXT('(') << Index << TEXT(')');
 		Builder << Binding.ToString(this, Args.Binding);
 	}
 
-	Builder << TEXT("Evaluate: ");
-	for (const FMVVMViewClass_EvaluateSource& EvaluateSource : GetEvaluateSources())
+	Builder << TEXT("\nEvaluates: ");
+	for (int32 Index = 0; Index < GetEvaluateSources().Num(); ++Index)
 	{
+		const FMVVMViewClass_EvaluateSource& EvaluateSource = GetEvaluateSources()[Index];
 		Builder << TEXT("\n");
+		Builder << TEXT('(') << Index << TEXT(')');
 		Builder << EvaluateSource.ToString(this, Args.Evaluate);
 	}
 
-	Builder << TEXT("Events: ");
-	for (const FMVVMViewClass_Event& Event : GetEvents())
+	Builder << TEXT("\nEvents: ");
+	for (int32 Index = 0; Index < GetEvents().Num(); ++Index)
 	{
+		const FMVVMViewClass_Event& Event = GetEvents()[Index];
 		Builder << TEXT("\n");
+		Builder << TEXT('(') << Index << TEXT(')');
 		Builder << Event.ToString(this, Args.Event);
 	}
 
diff --git a/Engine/Plugins/Runtime/ModelViewViewModel/Source/ModelViewViewModelBlueprint/Private/MVVMViewBlueprintCompiler.cpp b/Engine/Plugins/Runtime/ModelViewViewModel/Source/ModelViewViewModelBlueprint/Private/MVVMViewBlueprintCompiler.cpp
--- a/Engine/Plugins/Runtime/ModelViewViewModel/Source/ModelViewViewModelBlueprint/Private/MVVMViewBlueprintCompiler.cpp
+++ b/Engine/Plugins/Runtime/ModelViewViewModel/Source/ModelViewViewModelBlueprint/Private/MVVMViewBlueprintCompiler.cpp
@@ -2572,8 +2572,9 @@ void FMVVMViewBlueprintCompiler::CompileBindings(const FCompiledBindingLibraryCo
 		NewBinding.SourceBitField = 0;
 		NewBinding.EditorId = Binding.BindingId;
 
+		TArray<FMVVMViewClass_SourceKey, TInlineAllocator<16>>  SharedExecuteAtInitializationBindings;
+		int32 SharedBindingsCount = 0;
 		// Find the source needed by the binding. Also generate every bindings on that source (that need to register to the FieldNotify).
-		TArray<FMVVMViewClass_SourceKey, TInlineAllocator<16>>  SharedBindings;
 		for (TSharedPtr<FGeneratedReadFieldPathContext> ReadPath : ValidBinding->ReadPaths)
 		{
 			if (ReadPath->OptionalSource == nullptr)
@@ -2626,28 +2627,33 @@ void FMVVMViewBlueprintCompiler::CompileBindings(const FCompiledBindingLibraryCo
 				// Count the number of shared instance of that binding. (they can come from the same Source)
 				//This flag is used at runtime to know if we should delay the binding execution.
 				//We only delay when the field changes. It should only be shared if it has field.
-				SharedBindings.Add(FieldClassSourceKey);
+				++SharedBindingsCount;
+			}
+
+			if (bExecuteAtInitialization)
+			{
+				SharedExecuteAtInitializationBindings.Add(FieldClassSourceKey);
 			}
 		}
 
 		// Set binding flags.
 		NewBinding.Flags |= (!ValidBinding->bIsOneTimeBinding) ? (uint8)FMVVMViewClass_Binding::EFlags::OneWay : 0;
-		NewBinding.Flags |= (SharedBindings.Num() > 1) ? (uint8)FMVVMViewClass_Binding::EFlags::Shared : 0;
+		NewBinding.Flags |= (SharedBindingsCount > 1) ? (uint8)FMVVMViewClass_Binding::EFlags::Shared : 0;
 		NewBinding.Flags |= (Binding.bOverrideExecutionMode) ? (uint8)FMVVMViewClass_Binding::EFlags::OverrideExecuteMode : 0;
 		NewBinding.Flags |= (Binding.bEnabled) ? (uint8)FMVVMViewClass_Binding::EFlags::EnabledByDefault : 0;
 
-		// Only the last binding in the complex conversion.
-		if (SharedBindings.Num() > 0)
+		// Only the last binding in the complex conversion has ExecuteAtInitialization.
+		if (SharedExecuteAtInitializationBindings.Num() > 1)
 		{
-			SharedBindings.Sort([](const FMVVMViewClass_SourceKey& A, const FMVVMViewClass_SourceKey&B)
+			SharedExecuteAtInitializationBindings.Sort([](const FMVVMViewClass_SourceKey& A, const FMVVMViewClass_SourceKey&B)
 				{
 					return A.GetIndex() < B.GetIndex();
 				});
 
-			//Only the last one has ExecuteAtInitialization. Remove the flag on all the others.
-			for (int32 Index = 0; Index < SharedBindings.Num() - 2; ++Index)
+			//Remove the flag on all the binding on all the sources ecept the last one.
+			for (int32 Index = 0; Index < SharedExecuteAtInitializationBindings.Num() - 1; ++Index)
 			{
-				FMVVMViewClass_Source& ClassSource = ViewExtension->Sources[SharedBindings[Index].GetIndex()];
+				FMVVMViewClass_Source& ClassSource = ViewExtension->Sources[SharedExecuteAtInitializationBindings[Index].GetIndex()];
 				for (FMVVMViewClass_SourceBinding& SourceBinding : ClassSource.Bindings)
 				{
 					if (SourceBinding.GetBindingKey() == BindingKey)
@@ -2656,9 +2662,9 @@ void FMVVMViewBlueprintCompiler::CompileBindings(const FCompiledBindingLibraryCo
 					}
 				}
 			}
-			// keep only one ExecuteAtInitialization on the last source
+			// Keep only one ExecuteAtInitialization on the last source.
 			{
-				FMVVMViewClass_Source& ClassSource = ViewExtension->Sources[SharedBindings.Last().GetIndex()];
+				FMVVMViewClass_Source& ClassSource = ViewExtension->Sources[SharedExecuteAtInitializationBindings.Last().GetIndex()];
 				bool bFound = false;
 				for (int32 Index = ClassSource.Bindings.Num() - 1; Index >= 0; --Index)
 				{
