MVVM: add customization for property paths and conversion functions of bindings in the details panel

diff --git a/Engine/Plugins/Runtime/ModelViewViewModel/Source/ModelViewViewModelBlueprint/Public/MVVMBlueprintViewBinding.h b/Engine/Plugins/Runtime/ModelViewViewModel/Source/ModelViewViewModelBlueprint/Public/MVVMBlueprintViewBinding.h
--- a/Engine/Plugins/Runtime/ModelViewViewModel/Source/ModelViewViewModelBlueprint/Public/MVVMBlueprintViewBinding.h
+++ b/Engine/Plugins/Runtime/ModelViewViewModel/Source/ModelViewViewModelBlueprint/Public/MVVMBlueprintViewBinding.h
@@ -56,14 +56,14 @@ struct MODELVIEWVIEWMODELBLUEPRINT_API FMVVMBlueprintViewConversionPath
 	 * The graph that contains the conversion function when converting the value from the destination to the source.
 	 * When the function doesn't need a wrapper.
 	 */
-	UPROPERTY(Instanced)
+	UPROPERTY(VisibleAnywhere, NoClear, Category = "Viewmodel", Instanced)
 	TObjectPtr<UMVVMBlueprintViewConversionFunction> DestinationToSourceConversion;
 
 	/**
 	 * The graph that contains the conversion function when converting the value from the source to the destination.
 	 * When the function doesn't need a wrapper.
 	 */
-	UPROPERTY(Instanced)
+	UPROPERTY(VisibleAnywhere, NoClear, Category = "Viewmodel", Instanced)
 	TObjectPtr<UMVVMBlueprintViewConversionFunction> SourceToDestinationConversion;
 
 public:
diff --git a/Engine/Plugins/Runtime/ModelViewViewModel/Source/ModelViewViewModelEditor/Private/Customizations/MVVMConversionPathCustomization.cpp b/Engine/Plugins/Runtime/ModelViewViewModel/Source/ModelViewViewModelEditor/Private/Customizations/MVVMConversionPathCustomization.cpp
new file mode 100644
--- /dev/null
+++ b/Engine/Plugins/Runtime/ModelViewViewModel/Source/ModelViewViewModelEditor/Private/Customizations/MVVMConversionPathCustomization.cpp
@@ -0,0 +1,149 @@
+// Copyright Epic Games, Inc. All Rights Reserved.
+#include "MVVMConversionPathCustomization.h"
+#include "DetailWidgetRow.h"
+#include "IDetailChildrenBuilder.h"
+#include "MVVMBlueprintFunctionReference.h"
+#include "MVVMBlueprintViewBinding.h"
+#include "MVVMBlueprintViewConversionFunction.h"
+#include "WidgetBlueprint.h"
+#include "Widgets/Input/SEditableTextBox.h"
+#include "Widgets/Layout/SBox.h"
+#include "Widgets/SBoxPanel.h"
+#include "Widgets/SMVVMFunctionParameter.h"
+#include "Widgets/Text/STextBlock.h"
+
+#define LOCTEXT_NAMESPACE "MVVMConversionPath"
+
+namespace UE::MVVM
+{
+FConversionPathCustomization::FConversionPathCustomization(UWidgetBlueprint* InWidgetBlueprint)
+{
+	check(InWidgetBlueprint);
+	WidgetBlueprint = InWidgetBlueprint;
+}
+
+void FConversionPathCustomization::CustomizeHeader(TSharedRef<IPropertyHandle> InPropertyHandle, FDetailWidgetRow& HeaderRow, IPropertyTypeCustomizationUtils& CustomizationUtils)
+{
+	ParentHandle = InPropertyHandle->GetParentHandle();
+	BindingModeHandle = ParentHandle->GetChildHandle(GET_MEMBER_NAME_CHECKED(FMVVMBlueprintViewBinding, BindingType));
+	HeaderRow
+		.ShouldAutoExpand()
+		.NameContent()
+		[
+			InPropertyHandle->CreatePropertyNameWidget()
+		]
+		.ValueContent()
+		[
+			InPropertyHandle->CreatePropertyValueWidget()
+		];
+}
+
+void FConversionPathCustomization::CustomizeChildren(TSharedRef<IPropertyHandle> InPropertyHandle, IDetailChildrenBuilder& ChildBuilder, IPropertyTypeCustomizationUtils& CustomizationUtils)
+{
+	AddRowForProperty(ChildBuilder, InPropertyHandle, true);
+	AddRowForProperty(ChildBuilder, InPropertyHandle, false);
+}
+
+FText FConversionPathCustomization::GetFunctionPathText(UMVVMBlueprintViewConversionFunction* BlueprintViewConversionFunction) const
+{
+	const FText& NoneText = LOCTEXT("None", "None");
+	if (!BlueprintViewConversionFunction)
+	{
+		return NoneText;
+	}
+	
+	if (UWidgetBlueprint* WidgetBP = WidgetBlueprint.Get())
+	{
+		if (const UFunction* Function = BlueprintViewConversionFunction->GetConversionFunction().GetFunction(WidgetBP))
+		{
+			return FText::FromString(Function->GetPathName());
+		}
+	}
+	return NoneText;
+}
+
+void FConversionPathCustomization::AddRowForProperty(IDetailChildrenBuilder& ChildBuilder, const TSharedPtr<IPropertyHandle>& InPropertyHandle, bool bSourceToDestination)
+{
+	check(InPropertyHandle->GetProperty()->GetOwnerStruct() == FMVVMBlueprintViewBinding::StaticStruct());
+	 
+	TArray<FMVVMBlueprintViewBinding*> ViewBindings;
+	TArray<void*> RawData;
+	ParentHandle->AccessRawData(RawData);
+	for (void* Data : RawData)
+	{
+		FMVVMBlueprintViewBinding* Binding = reinterpret_cast<FMVVMBlueprintViewBinding*>(Data);
+		if (Binding)
+		{
+			ViewBindings.Add(Binding);
+		}
+	}
+
+	// Multi-selection is not currently supported in the binding panel.
+	if (RawData.Num() != 1)
+	{
+		return;
+	}
+
+	const FName FunctionPropertyName = bSourceToDestination ? GET_MEMBER_NAME_CHECKED(FMVVMBlueprintViewConversionPath, SourceToDestinationConversion) : GET_MEMBER_NAME_CHECKED(FMVVMBlueprintViewConversionPath, DestinationToSourceConversion);
+	TSharedRef<IPropertyHandle> FunctionProperty = InPropertyHandle->GetChildHandle(FunctionPropertyName).ToSharedRef();
+	UMVVMBlueprintViewConversionFunction* ConversionFunction = ViewBindings[0]->Conversion.GetConversionFunction(bSourceToDestination);
+
+	ChildBuilder.AddCustomRow(FText::FromName(FunctionPropertyName))
+		.NameContent()
+		[
+			FunctionProperty->CreatePropertyNameWidget()
+		]
+		.ValueContent()
+		[
+			SNew(SHorizontalBox)
+			+ SHorizontalBox::Slot()
+			[
+				SNew(SEditableTextBox)
+				.Text(this, &FConversionPathCustomization::GetFunctionPathText, ConversionFunction)
+			]
+		];
+			
+		const EMVVMBindingMode BindingMode = GetBindingMode();
+
+		if (ConversionFunction == nullptr || !IsForwardBinding(GetBindingMode()))
+		{
+			return;
+		}
+
+		const TArrayView<const FMVVMBlueprintPin>& Pins = ConversionFunction->GetPins();
+		for (const FMVVMBlueprintPin& Pin : Pins)
+		{
+			const FText& ArgumentName = FText::FromString(Pin.GetId().ToString());
+			ChildBuilder.AddCustomRow(ArgumentName)
+				.NameContent()
+				[
+					SNew(SBox)
+					.Padding(FMargin(16, 0, 0, 0))
+					[
+						SNew(STextBlock)
+						.Text(ArgumentName)
+					]
+				]
+				.ValueContent()
+				[
+					SNew(UE::MVVM::SFunctionParameter, WidgetBlueprint.Get())
+					.BindingId(ViewBindings[0]->BindingId)
+					.ParameterId(Pin.GetId())
+					.SourceToDestination(bSourceToDestination)
+					.AllowDefault(true)
+					.IsEnabled(false)
+				];
+		}
+}
+
+EMVVMBindingMode FConversionPathCustomization::GetBindingMode() const
+{
+	uint8 EnumValue = 0;
+	if (BindingModeHandle->GetValue(EnumValue) == FPropertyAccess::Success)
+	{
+		return static_cast<EMVVMBindingMode>(EnumValue);
+	}
+	return EMVVMBindingMode::OneWayToDestination;
+}
+}
+#undef LOCTEXT_NAMESPACE 
\ No newline at end of file
diff --git a/Engine/Plugins/Runtime/ModelViewViewModel/Source/ModelViewViewModelEditor/Private/Customizations/MVVMConversionPathCustomization.h b/Engine/Plugins/Runtime/ModelViewViewModel/Source/ModelViewViewModelEditor/Private/Customizations/MVVMConversionPathCustomization.h
new file mode 100644
--- /dev/null
+++ b/Engine/Plugins/Runtime/ModelViewViewModel/Source/ModelViewViewModelEditor/Private/Customizations/MVVMConversionPathCustomization.h
@@ -0,0 +1,34 @@
+// Copyright Epic Games, Inc. All Rights Reserved.
+#pragma once
+
+#include "IPropertyTypeCustomization.h"
+
+enum class EMVVMBindingMode : uint8;
+struct EVisibility;
+class UWidgetBlueprint;
+class UMVVMBlueprintViewConversionFunction;
+
+namespace UE::MVVM
+{
+class FConversionPathCustomization : public IPropertyTypeCustomization
+{
+public:
+	static TSharedRef<IPropertyTypeCustomization> MakeInstance(UWidgetBlueprint* InWidgetBlueprint)
+	{
+		return MakeShared<FConversionPathCustomization>(InWidgetBlueprint);
+	}
+	FConversionPathCustomization(UWidgetBlueprint* InWidgetBlueprint);
+	virtual void CustomizeHeader(TSharedRef<IPropertyHandle> InPropertyHandle, FDetailWidgetRow& HeaderRow, IPropertyTypeCustomizationUtils& CustomizationUtils) override;
+	virtual void CustomizeChildren(TSharedRef<IPropertyHandle> InPropertyHandle, IDetailChildrenBuilder& ChildBuilder, IPropertyTypeCustomizationUtils& CustomizationUtils) override;
+
+private:
+	FText GetFunctionPathText(UMVVMBlueprintViewConversionFunction* BlueprintViewConversionFunction) const;
+	void AddRowForProperty(IDetailChildrenBuilder& ChildBuilder, const TSharedPtr<IPropertyHandle>& Property, bool bSourceToDestination);
+	EMVVMBindingMode GetBindingMode() const;
+
+private:
+	TWeakObjectPtr<UWidgetBlueprint> WidgetBlueprint;
+	TSharedPtr<IPropertyHandle> ParentHandle;
+	TSharedPtr<IPropertyHandle> BindingModeHandle;
+};
+} // namespace UE::MVVM
\ No newline at end of file
diff --git a/Engine/Plugins/Runtime/ModelViewViewModel/Source/ModelViewViewModelEditor/Private/Customizations/MVVMPropertyPathCustomization.cpp b/Engine/Plugins/Runtime/ModelViewViewModel/Source/ModelViewViewModelEditor/Private/Customizations/MVVMPropertyPathCustomization.cpp
new file mode 100644
--- /dev/null
+++ b/Engine/Plugins/Runtime/ModelViewViewModel/Source/ModelViewViewModelEditor/Private/Customizations/MVVMPropertyPathCustomization.cpp
@@ -0,0 +1,132 @@
+// Copyright Epic Games, Inc. All Rights Reserved.
+#include "MVVMPropertyPathCustomization.h"
+
+#include "DetailWidgetRow.h"
+#include "Editor.h"
+#include "MVVMBlueprintView.h"
+#include "MVVMBlueprintViewConversionFunction.h"
+#include "MVVMEditorSubsystem.h"
+#include "WidgetBlueprint.h"
+#include "Widgets/SMVVMFieldDisplay.h"
+
+#define LOCTEXT_NAMESPACE "MVVMPropertyPathCustomization"
+
+namespace UE::MVVM
+{
+FPropertyPathCustomization::FPropertyPathCustomization(UWidgetBlueprint* InWidgetBlueprint) :
+	WidgetBlueprint(InWidgetBlueprint)
+{
+	check(WidgetBlueprint.IsValid());
+}
+
+FPropertyPathCustomization::~FPropertyPathCustomization()
+{
+}
+
+void FPropertyPathCustomization::CustomizeHeader(TSharedRef<IPropertyHandle> InPropertyHandle, FDetailWidgetRow& HeaderRow, IPropertyTypeCustomizationUtils& CustomizationUtils)
+{
+	PropertyHandle = InPropertyHandle;
+
+	const FName PropertyName = PropertyHandle->GetProperty()->GetFName();
+	FText PropertyDisplayName;
+	if (PropertyName == GET_MEMBER_NAME_CHECKED(FMVVMBlueprintViewBinding, DestinationPath))
+	{
+		bIsSource = false;
+		PropertyDisplayName = LOCTEXT("PropertyDisplay_Widget", "Destination");
+	}
+	else if (PropertyName == GET_MEMBER_NAME_CHECKED(FMVVMBlueprintViewBinding, SourcePath))
+	{
+		bIsSource = true;
+		PropertyDisplayName = LOCTEXT("PropertyDisplay_Source", "Source");
+	}
+	else
+	{
+		ensureMsgf(false, TEXT("MVVMPropertyPathCustomization used in unknown context."));
+	}
+
+	FMVVMBlueprintViewBinding* OwningBinding = GetOwningBinding();
+	OwningBindingId = OwningBinding ? OwningBinding->BindingId : FGuid();
+
+	HeaderRow
+		.NameWidget
+		.VAlign(VAlign_Center)
+		[
+			PropertyHandle->CreatePropertyNameWidget(PropertyDisplayName)
+		]
+	.ValueWidget
+		.VAlign(VAlign_Center)
+		.HAlign(HAlign_Fill)
+		[
+			SNew(SFieldDisplay, WidgetBlueprint.Get())
+			.OnGetLinkedValue(this, &FPropertyPathCustomization::GetFieldValue)
+			.TextStyle(FAppStyle::Get(), "SmallText")
+		];
+}
+
+FMVVMBlueprintPropertyPath FPropertyPathCustomization::GetPropertyPathValue() const
+{
+	check(PropertyHandle->GetProperty()->GetOwnerStruct() == FMVVMBlueprintViewBinding::StaticStruct());
+	TArray<void*> RawData;
+	PropertyHandle->AccessRawData(RawData);
+	if (RawData.Num() > 1)
+	{
+		return FMVVMBlueprintPropertyPath();
+	}
+
+	for (void* RawPtr : RawData)
+	{
+		if (RawPtr == nullptr)
+		{
+			continue;
+		}
+		FMVVMBlueprintPropertyPath* PropertyPath = (FMVVMBlueprintPropertyPath*)RawPtr;
+		if (PropertyPath)
+		{
+			return *PropertyPath;
+		}
+	}
+	return FMVVMBlueprintPropertyPath();
+}
+
+FMVVMBlueprintViewBinding* FPropertyPathCustomization::GetOwningBinding() const
+{
+	if (TSharedPtr<IPropertyHandle> ParentHandle = PropertyHandle->GetParentHandle())
+	{
+		TArray<void*> RawData;
+		ParentHandle->AccessRawData(RawData);
+		for (void* Data : RawData)
+		{
+			FMVVMBlueprintViewBinding* Binding = reinterpret_cast<FMVVMBlueprintViewBinding*>(Data);
+			if (Binding)
+			{
+				return Binding;
+			}
+		}
+	}
+	return nullptr;
+}
+
+FMVVMLinkedPinValue FPropertyPathCustomization::GetFieldValue() const
+{
+	if (OwningBindingId.IsValid())
+	{
+		if (UWidgetBlueprint* WidgetBP = WidgetBlueprint.Get())
+		{
+			UMVVMEditorSubsystem* EditorSubsystem = GEditor->GetEditorSubsystem<UMVVMEditorSubsystem>();
+			if (const UMVVMBlueprintView* View = EditorSubsystem->GetView(WidgetBlueprint.Get()))
+			{
+				if (const FMVVMBlueprintViewBinding* OwningBinding = View->GetBinding(OwningBindingId))
+				{
+					if (UMVVMBlueprintViewConversionFunction* ConversionFunction = OwningBinding->Conversion.GetConversionFunction(bIsSource))
+					{
+						return FMVVMLinkedPinValue(WidgetBP, ConversionFunction->GetConversionFunction());
+					}
+				}
+			}
+		}
+	}
+	return FMVVMLinkedPinValue(GetPropertyPathValue());
+}
+
+} // namespace UE::MVVM
+#undef LOCTEXT_NAMESPACE
\ No newline at end of file
diff --git a/Engine/Plugins/Runtime/ModelViewViewModel/Source/ModelViewViewModelEditor/Private/Customizations/MVVMPropertyPathCustomization.h b/Engine/Plugins/Runtime/ModelViewViewModel/Source/ModelViewViewModelEditor/Private/Customizations/MVVMPropertyPathCustomization.h
new file mode 100644
--- /dev/null
+++ b/Engine/Plugins/Runtime/ModelViewViewModel/Source/ModelViewViewModelEditor/Private/Customizations/MVVMPropertyPathCustomization.h
@@ -0,0 +1,36 @@
+// Copyright Epic Games, Inc. All Rights Reserved.
+#pragma once
+#include "IPropertyTypeCustomization.h"
+#include "Types/MVVMLinkedPinValue.h"
+
+struct FMVVMBlueprintPropertyPath;
+struct FMVVMBlueprintViewBinding;
+
+namespace UE::MVVM
+{
+class FPropertyPathCustomization : public IPropertyTypeCustomization
+{
+public:
+	FPropertyPathCustomization(UWidgetBlueprint* WidgetBlueprint);
+	~FPropertyPathCustomization();
+	virtual void CustomizeHeader(TSharedRef<IPropertyHandle> InPropertyHandle, FDetailWidgetRow& HeaderRow, IPropertyTypeCustomizationUtils& CustomizationUtils) override;
+	virtual void CustomizeChildren(TSharedRef<IPropertyHandle> InPropertyHandle, IDetailChildrenBuilder& ChildBuilder, IPropertyTypeCustomizationUtils& CustomizationUtils) override {}
+	static TSharedRef<IPropertyTypeCustomization> MakeInstance(UWidgetBlueprint* InWidgetBlueprint)
+	{
+		check(InWidgetBlueprint != nullptr);
+		return MakeShared<FPropertyPathCustomization>(InWidgetBlueprint);
+	}
+
+	FMVVMLinkedPinValue GetFieldValue() const;
+	FMVVMBlueprintViewBinding* GetOwningBinding() const;
+
+private:
+	FMVVMBlueprintPropertyPath GetPropertyPathValue() const;
+
+private:
+	TSharedPtr<IPropertyHandle> PropertyHandle;
+	TWeakObjectPtr<UWidgetBlueprint> WidgetBlueprint;
+	bool bIsSource = false;
+	FGuid OwningBindingId;
+};
+} // namespace UE::MVVM
\ No newline at end of file
diff --git a/Engine/Plugins/Runtime/ModelViewViewModel/Source/ModelViewViewModelEditor/Private/Widgets/SMVVMFieldDisplay.cpp b/Engine/Plugins/Runtime/ModelViewViewModel/Source/ModelViewViewModelEditor/Private/Widgets/SMVVMFieldDisplay.cpp
new file mode 100644
--- /dev/null
+++ b/Engine/Plugins/Runtime/ModelViewViewModel/Source/ModelViewViewModelEditor/Private/Widgets/SMVVMFieldDisplay.cpp
@@ -0,0 +1,123 @@
+// Copyright Epic Games, Inc. All Rights Reserved.
+
+#include "Widgets/SMVVMFieldDisplay.h"
+
+#include "MVVMPropertyPath.h"
+#include "Styling/MVVMEditorStyle.h"
+#include "Styling/SlateWidgetStyleAsset.h"
+#include "WidgetBlueprint.h"
+#include "Widgets/Layout/SBox.h"
+#include "Widgets/Layout/SWidgetSwitcher.h"
+#include "Widgets/SMVVMCachedViewBindingPropertyPath.h"
+#include "Widgets/SMVVMCachedViewBindingConversionFunction.h"
+#include "Widgets/Text/STextBlock.h"
+
+#define LOCTEXT_NAMESPACE "MVVMFieldDisplay"
+
+namespace UE::MVVM
+{
+
+void SFieldDisplay::Construct(const FArguments& InArgs, const UWidgetBlueprint* InWidgetBlueprint)
+{
+	check(InWidgetBlueprint);
+	
+	TextStyle = InArgs._TextStyle;
+	OnGetLinkedValue = InArgs._OnGetLinkedValue;
+
+	ChildSlot
+	[
+		SNew(SWidgetSwitcher)
+		.WidgetIndex(this, &SFieldDisplay::GetCurrentDisplayIndex)
+		//0-Property/Function (from Widget or Viewmodel).
+		//0-Property/Function argument of conversion function
+		+ SWidgetSwitcher::Slot()
+		.Padding(0.0f, 0.0f, 8.0f, 0.0f)
+		.HAlign(HAlign_Left)
+		.VAlign(VAlign_Center)
+		[
+			SNew(SCachedViewBindingPropertyPath, InWidgetBlueprint)
+			.TextStyle(TextStyle)
+			.ShowContext(InArgs._ShowContext)
+			.OnGetPropertyPath(this, &SFieldDisplay::HandleGetPropertyPath)
+		]
+
+		//1-Conversion Function
+		+ SWidgetSwitcher::Slot()
+		.Padding(0.0f, 0.0f, 8.0f, 0.0f)
+		.HAlign(HAlign_Left)
+		.VAlign(VAlign_Center)
+		[
+			SNew(SCachedViewBindingConversionFunction, InWidgetBlueprint)
+			.TextStyle(TextStyle)
+			.OnGetConversionFunction(this, &SFieldDisplay::HandleGetConversionFunction)
+		]
+
+		//2-Nothing selected.
+		+ SWidgetSwitcher::Slot()
+		[
+			SNew(SBox)
+			.Padding(FMargin(8.0f, 0.0f, 8.0f, 0.0f))
+			[
+				SNew(STextBlock)
+				.TextStyle(FAppStyle::Get(), "HintText")
+				.Text(LOCTEXT("None", "No field selected"))
+			]
+		]
+	];
+}
+
+int32 SFieldDisplay::GetCurrentDisplayIndex() const
+{
+	if (OnGetLinkedValue.IsBound())
+	{
+		FMVVMLinkedPinValue LinkedValue = OnGetLinkedValue.Execute();
+		if (LinkedValue.IsConversionFunction() || LinkedValue.IsConversionNode())
+		{
+			return 1;
+		}
+		else if (LinkedValue.IsPropertyPath())
+		{
+			return 0;
+		}
+	}
+	return 2;
+}
+
+FMVVMBlueprintPropertyPath SFieldDisplay::HandleGetPropertyPath() const
+{
+	if (OnGetLinkedValue.IsBound())
+	{
+		FMVVMLinkedPinValue LinkedValue = OnGetLinkedValue.Execute();
+		if (LinkedValue.IsPropertyPath())
+		{
+			return LinkedValue.GetPropertyPath();
+		}
+	}
+	return FMVVMBlueprintPropertyPath();
+}
+
+TVariant<const UFunction*, TSubclassOf<UK2Node>, FEmptyVariantState> SFieldDisplay::HandleGetConversionFunction() const
+{
+	using FReturnType = TVariant<const UFunction*, TSubclassOf<UK2Node>, FEmptyVariantState>;
+	if (OnGetLinkedValue.IsBound())
+	{
+		FMVVMLinkedPinValue LinkedValue = OnGetLinkedValue.Execute();
+		if (ensure(LinkedValue.IsConversionFunction() || LinkedValue.IsConversionNode()))
+		{
+			if (LinkedValue.IsConversionFunction())
+			{
+				return FReturnType(TInPlaceType<const UFunction*>(), LinkedValue.GetConversionFunction());
+			}
+			else
+			{
+				check(LinkedValue.IsConversionNode());
+				return FReturnType(TInPlaceType<TSubclassOf<UK2Node>>(), LinkedValue.GetConversionNode());
+			}
+		}
+	}
+	return FReturnType(TInPlaceType<FEmptyVariantState>());
+}
+
+} // namespace UE::MVVM
+
+#undef LOCTEXT_NAMESPACE
diff --git a/Engine/Plugins/Runtime/ModelViewViewModel/Source/ModelViewViewModelEditor/Private/Widgets/SMVVMFieldDisplay.h b/Engine/Plugins/Runtime/ModelViewViewModel/Source/ModelViewViewModelEditor/Private/Widgets/SMVVMFieldDisplay.h
new file mode 100644
--- /dev/null
+++ b/Engine/Plugins/Runtime/ModelViewViewModel/Source/ModelViewViewModelEditor/Private/Widgets/SMVVMFieldDisplay.h
@@ -0,0 +1,44 @@
+// Copyright Epic Games, Inc. All Rights Reserved.
+
+#pragma once
+
+#include "Styling/CoreStyle.h"
+#include "Styling/SlateTypes.h"
+#include "Styling/SlateWidgetStyleAsset.h"
+#include "Types/MVVMLinkedPinValue.h"
+#include "Widgets/SCompoundWidget.h"
+
+struct FMVVMBlueprintPropertyPath;
+
+namespace UE::MVVM
+{
+
+class SFieldDisplay : public SCompoundWidget
+{
+public:
+	DECLARE_DELEGATE_RetVal(FMVVMLinkedPinValue, FOnGetLinkedPinValue);
+
+	SLATE_BEGIN_ARGS(SFieldDisplay)
+		: _TextStyle(&FCoreStyle::Get().GetWidgetStyle<FTextBlockStyle>("NormalText"))
+		{
+		}
+		SLATE_STYLE_ARGUMENT(FTextBlockStyle, TextStyle)
+		SLATE_ARGUMENT_DEFAULT(bool, ShowContext) = true;
+		SLATE_EVENT(FOnGetLinkedPinValue, OnGetLinkedValue)
+	SLATE_END_ARGS()
+
+	void Construct(const FArguments& InArgs, const UWidgetBlueprint* InWidgetBlueprint);
+
+private:
+	int32 GetCurrentDisplayIndex() const;
+	FMVVMBlueprintPropertyPath HandleGetPropertyPath() const;
+	TVariant<const UFunction*, TSubclassOf<UK2Node>, FEmptyVariantState> HandleGetConversionFunction() const;
+
+public:
+	FOnGetLinkedPinValue OnGetLinkedValue;
+
+private:
+	const FTextBlockStyle* TextStyle = nullptr;
+}; 
+
+} // namespace UE::MVVM
diff --git a/Engine/Plugins/Runtime/ModelViewViewModel/Source/ModelViewViewModelEditor/Private/Widgets/SMVVMFieldSelector.cpp b/Engine/Plugins/Runtime/ModelViewViewModel/Source/ModelViewViewModelEditor/Private/Widgets/SMVVMFieldSelector.cpp
--- a/Engine/Plugins/Runtime/ModelViewViewModel/Source/ModelViewViewModelEditor/Private/Widgets/SMVVMFieldSelector.cpp
+++ b/Engine/Plugins/Runtime/ModelViewViewModel/Source/ModelViewViewModelEditor/Private/Widgets/SMVVMFieldSelector.cpp
@@ -9,8 +9,6 @@
 
 #include "Widgets/Input/SComboButton.h"
 #include "Widgets/Layout/SBox.h"
-#include "Widgets/Layout/SWidgetSwitcher.h"
-#include "Widgets/Text/STextBlock.h"
 
 #define LOCTEXT_NAMESPACE "MVVMFieldSelector"
 
@@ -31,9 +29,7 @@ void SFieldSelector::Construct(const FArguments& InArgs, const UWidgetBlueprint*
 {
 	WidgetBlueprint = InWidgetBlueprint;
 	check(InWidgetBlueprint);
-	
-	TextStyle = InArgs._TextStyle;
-	OnGetLinkedValue = InArgs._OnGetLinkedValue;
+
 	OnSelectionChanged = InArgs._OnSelectionChanged;
 	OnGetSelectionContext = InArgs._OnGetSelectionContext;
 	OnDragEnterEvent = InArgs._OnDragEnter;
@@ -51,77 +47,26 @@ void SFieldSelector::Construct(const FArguments& InArgs, const UWidgetBlueprint*
 			.ContentPadding(FMargin(4.0f, 2.0f))
 			.ButtonContent()
 			[
-				SNew(SWidgetSwitcher)
-				.WidgetIndex(this, &SFieldSelector::GetCurrentDisplayIndex)
-				//0-Property/Function (from Widget or Viewmodel).
-				//0-Property/Function argument of conversion function
-				+ SWidgetSwitcher::Slot()
-				.Padding(0.0f, 0.0f, 8.0f, 0.0f)
-				.HAlign(HAlign_Left)
-				.VAlign(VAlign_Center)
-				[
-					SNew(SCachedViewBindingPropertyPath, WidgetBlueprint.Get())
-					.TextStyle(TextStyle)
-					.ShowContext(InArgs._ShowContext)
-					.OnGetPropertyPath(this, &SFieldSelector::HandleGetPropertyPath)
-				]
-
-				//1-Conversion Function
-				+ SWidgetSwitcher::Slot()
-				.Padding(0.0f, 0.0f, 8.0f, 0.0f)
-				.HAlign(HAlign_Left)
-				.VAlign(VAlign_Center)
-				[
-					SNew(SCachedViewBindingConversionFunction, WidgetBlueprint.Get())
-					.TextStyle(TextStyle)
-					.OnGetConversionFunction(this, &SFieldSelector::HandleGetConversionFunction)
-				]
-
-				//2-Nothing selected.
-				+ SWidgetSwitcher::Slot()
-				[
-					SNew(SBox)
-					.Padding(FMargin(8.0f, 0.0f, 8.0f, 0.0f))
-					[
-						SNew(STextBlock)
-						.TextStyle(FAppStyle::Get(), "HintText")
-						.Text(LOCTEXT("None", "No field selected"))
-					]
-				]
+				SAssignNew(FieldDisplay, SFieldDisplay, InWidgetBlueprint)
+				.TextStyle(InArgs._TextStyle)
+				.OnGetLinkedValue(InArgs._OnGetLinkedValue)
 			]
 		]
 	];
 }
 
-int32 SFieldSelector::GetCurrentDisplayIndex() const
-{
-	if (OnGetLinkedValue.IsBound())
-	{
-		FMVVMLinkedPinValue LinkedValue = OnGetLinkedValue.Execute();
-		if (LinkedValue.IsConversionFunction() || LinkedValue.IsConversionNode())
-		{
-			return 1;
-		}
-		else if (LinkedValue.IsPropertyPath())
-		{
-			return 0;
-		}
-	}
-	return 2;
-}
-
 TSharedRef<SWidget> SFieldSelector::HandleGetMenuContent()
 {
 	const UWidgetBlueprint* WidgetBlueprintPtr = WidgetBlueprint.Get();
-	if (!WidgetBlueprintPtr)
+	if (!WidgetBlueprintPtr || !FieldDisplay)
 	{
 		return SNullWidget::NullWidget;
 	}
 
 	TOptional<FMVVMLinkedPinValue> CurrentSelected;
-	if (OnGetLinkedValue.IsBound())
+	if (FieldDisplay->OnGetLinkedValue.IsBound())
 	{
-		CurrentSelected = OnGetLinkedValue.Execute();
+		CurrentSelected = FieldDisplay->OnGetLinkedValue.Execute();
 	}
 
 	FFieldSelectionContext SelectionContext;
@@ -156,41 +101,6 @@ void SFieldSelector::HandleFieldSelectionChanged(FMVVMLinkedPinValue LinkedValue
 	}
 }
 
-FMVVMBlueprintPropertyPath SFieldSelector::HandleGetPropertyPath() const
-{
-	if (OnGetLinkedValue.IsBound())
-	{
-		FMVVMLinkedPinValue LinkedValue = OnGetLinkedValue.Execute();
-		if (LinkedValue.IsPropertyPath())
-		{
-			return LinkedValue.GetPropertyPath();
-		}
-	}
-	return FMVVMBlueprintPropertyPath();
-}
-
-TVariant<const UFunction*, TSubclassOf<UK2Node>, FEmptyVariantState> SFieldSelector::HandleGetConversionFunction() const
-{
-	using FReturnType = TVariant<const UFunction*, TSubclassOf<UK2Node>, FEmptyVariantState>;
-	if (OnGetLinkedValue.IsBound())
-	{
-		FMVVMLinkedPinValue LinkedValue = OnGetLinkedValue.Execute();
-		if (ensure(LinkedValue.IsConversionFunction() || LinkedValue.IsConversionNode()))
-		{
-			if (LinkedValue.IsConversionFunction())
-			{
-				return FReturnType(TInPlaceType<const UFunction*>(), LinkedValue.GetConversionFunction());
-			}
-			else
-			{
-				check(LinkedValue.IsConversionNode());
-				return FReturnType(TInPlaceType<TSubclassOf<UK2Node>>(), LinkedValue.GetConversionNode());
-			}
-		}
-	}
-	return FReturnType(TInPlaceType<FEmptyVariantState>());
-}
-
 void SFieldSelector::OnDragEnter(const FGeometry& MyGeometry, const FDragDropEvent& DragDropEvent)
 {
 	if (OnDragEnterEvent.IsBound())
diff --git a/Engine/Plugins/Runtime/ModelViewViewModel/Source/ModelViewViewModelEditor/Private/Widgets/SMVVMFieldSelector.h b/Engine/Plugins/Runtime/ModelViewViewModel/Source/ModelViewViewModelEditor/Private/Widgets/SMVVMFieldSelector.h
--- a/Engine/Plugins/Runtime/ModelViewViewModel/Source/ModelViewViewModelEditor/Private/Widgets/SMVVMFieldSelector.h
+++ b/Engine/Plugins/Runtime/ModelViewViewModel/Source/ModelViewViewModelEditor/Private/Widgets/SMVVMFieldSelector.h
@@ -6,8 +6,7 @@
 #include "Styling/SlateTypes.h"
 #include "Styling/SlateWidgetStyleAsset.h"
 #include "Widgets/SCompoundWidget.h"
-#include "Widgets/SMVVMCachedViewBindingPropertyPath.h"
-#include "Widgets/SMVVMCachedViewBindingConversionFunction.h"
+#include "Widgets/SMVVMFieldDisplay.h"
 #include "Widgets/SMVVMFieldSelectorMenu.h"
 
 namespace UE::MVVM { class SCachedViewBindingPropertyPath; }
@@ -20,7 +19,6 @@ namespace UE::MVVM
 class SFieldSelector : public SCompoundWidget
 {
 public:
-	DECLARE_DELEGATE_RetVal(FMVVMLinkedPinValue, FOnGetLinkedPinValue);
 	DECLARE_DELEGATE_RetVal(FFieldSelectionContext, FOnGetSelectionContext);
 	DECLARE_DELEGATE_RetVal_TwoParams(FReply, FOnDrop, const FGeometry&, const FDragDropEvent&);
 	DECLARE_DELEGATE_TwoParams(FOnDragEnter, const FGeometry&, const FDragDropEvent&);
@@ -32,7 +30,7 @@ public:
 		SLATE_STYLE_ARGUMENT(FTextBlockStyle, TextStyle)
 		SLATE_ARGUMENT_DEFAULT(bool, ShowContext) { true };
 		SLATE_ARGUMENT_DEFAULT(bool, IsBindingToEvent) { false };
-		SLATE_EVENT(FOnGetLinkedPinValue, OnGetLinkedValue)
+		SLATE_EVENT(SFieldDisplay::FOnGetLinkedPinValue, OnGetLinkedValue)
 		SLATE_EVENT(FOnLinkedValueSelectionChanged, OnSelectionChanged)
 		SLATE_EVENT(FOnGetSelectionContext, OnGetSelectionContext)
 		SLATE_EVENT(FOnDrop, OnDrop)
@@ -44,25 +42,18 @@ public:
 	virtual void OnDragEnter(const FGeometry& MyGeometry, const FDragDropEvent& DragDropEvent) override;
 	virtual void OnDragLeave(const FDragDropEvent& DragDropEvent) override;
 
-
 private:
-	int32 GetCurrentDisplayIndex() const;
-
 	TSharedRef<SWidget> CreateSourcePanel();
 	TSharedRef<SWidget> HandleGetMenuContent();
 
 	void HandleFieldSelectionChanged(FMVVMLinkedPinValue NewValue);
-	FMVVMBlueprintPropertyPath HandleGetPropertyPath() const;
-	TVariant<const UFunction*, TSubclassOf<UK2Node>, FEmptyVariantState> HandleGetConversionFunction() const;
 	void HandleMenuClosed();
 
 private:
 	TSharedPtr<SCachedViewBindingPropertyPath> PropertyPathWidget;
 	TSharedPtr<SComboButton> ComboButton;
-
+	TSharedPtr<SFieldDisplay> FieldDisplay;
 	TWeakObjectPtr<const UWidgetBlueprint> WidgetBlueprint;
-	const FTextBlockStyle* TextStyle = nullptr;
-	FOnGetLinkedPinValue OnGetLinkedValue;
 	FOnLinkedValueSelectionChanged OnSelectionChanged;
 	FOnGetSelectionContext OnGetSelectionContext;
 	FOnDrop OnDropEvent;
diff --git a/Engine/Plugins/Runtime/ModelViewViewModel/Source/ModelViewViewModelEditor/Private/Widgets/SMVVMViewBindingPanel.cpp b/Engine/Plugins/Runtime/ModelViewViewModel/Source/ModelViewViewModelEditor/Private/Widgets/SMVVMViewBindingPanel.cpp
--- a/Engine/Plugins/Runtime/ModelViewViewModel/Source/ModelViewViewModelEditor/Private/Widgets/SMVVMViewBindingPanel.cpp
+++ b/Engine/Plugins/Runtime/ModelViewViewModel/Source/ModelViewViewModelEditor/Private/Widgets/SMVVMViewBindingPanel.cpp
@@ -12,6 +12,8 @@
 #include "WidgetBlueprintEditor.h"
 #include "WidgetBlueprintToolMenuContext.h"
 
+#include "Customizations/MVVMConversionPathCustomization.h"
+#include "Customizations/MVVMPropertyPathCustomization.h"
 #include "Hierarchy/HierarchyWidgetDragDropOp.h"
 #include "IStructureDetailsView.h"
 #include "PropertyEditorModule.h"
@@ -126,6 +128,9 @@ void SBindingsPanel::Construct(const FArguments& InArgs, TSharedPtr<FWidgetBluep
 
 		FStructureDetailsViewArgs StructureDetailsViewArgs;
 		StructDetailsView = PropertyEditorModule.CreateStructureDetailView(DetailsViewArgs, StructureDetailsViewArgs, MakeShared<FStructOnScope>(FMVVMBlueprintViewBinding::StaticStruct(), reinterpret_cast<uint8*>(&NotifyHook->Binding)));
+		StructDetailsView->GetDetailsView()->RegisterInstancedCustomPropertyTypeLayout(FMVVMBlueprintPropertyPath::StaticStruct()->GetFName(), FOnGetPropertyTypeCustomizationInstance::CreateStatic(&UE::MVVM::FPropertyPathCustomization::MakeInstance, WidgetBlueprint));
+		StructDetailsView->GetDetailsView()->RegisterInstancedCustomPropertyTypeLayout(FMVVMBlueprintViewConversionPath::StaticStruct()->GetFName(), FOnGetPropertyTypeCustomizationInstance::CreateStatic(&UE::MVVM::FConversionPathCustomization::MakeInstance, WidgetBlueprint));
+		StructDetailsView->GetDetailsView()->SetIsPropertyEditingEnabledDelegate(FIsPropertyEditingEnabled::CreateSP(this, &SBindingsPanel::IsDetailsViewEditingEnabled));
 	}
 
 	HandleBlueprintViewChangedDelegate();
@@ -136,6 +141,10 @@ SBindingsPanel::~SBindingsPanel()
 	if (UMVVMWidgetBlueprintExtension_View* Extension = MVVMExtension.Get())
 	{
 		Extension->OnBlueprintViewChangedDelegate().Remove(BlueprintViewChangedDelegateHandle);
+		if (UMVVMBlueprintView* View = Extension->GetBlueprintView())
+		{
+			View->OnBindingsUpdated.RemoveAll(this);
+		}
 	}
 	if (TSharedPtr<FWidgetBlueprintEditor> WidgetEditor = WeakBlueprintEditor.Pin())
 	{
@@ -152,6 +161,11 @@ void SBindingsPanel::SaveSettings()
 	GConfig->SetInt(TEXT("MVVMViewBindingPanel"), TEXT("LastAddBindingMode"), static_cast<int32>(AddBindingMode), *GEditorPerProjectIni);
 }
 
+bool SBindingsPanel::IsDetailsViewEditingEnabled() const
+{
+	return false;
+}
+
 void SBindingsPanel::LoadSettings()
 {
 	GConfig->SetInt(TEXT("MVVMViewBindingPanel"), TEXT("LastAddBindingMode"), static_cast<int32>(AddBindingMode), *GEditorPerProjectIni);
@@ -199,6 +213,8 @@ void SBindingsPanel::OnBindingListSelectionChanged(TConstArrayView<FMVVMBlueprin
 	if (FMVVMBlueprintViewBinding* Binding = Selection.Num() > 0 ? Selection.Last() : nullptr)
 	{
 		NotifyHook->Binding = *Binding;
+		TSharedRef<FStructOnScope> StructScope = MakeShared<FStructOnScope>(FMVVMBlueprintViewBinding::StaticStruct(), reinterpret_cast<uint8*>(Binding));
+		StructDetailsView->SetStructureData(StructScope);
 		DetailContainer->SetContent(StructDetailsView->GetWidget().ToSharedRef());
 	}
 	else
@@ -210,6 +226,12 @@ void SBindingsPanel::OnBindingListSelectionChanged(TConstArrayView<FMVVMBlueprin
 }
 
 void SBindingsPanel::Tick(const FGeometry& AllottedGeometry, const double InCurrentTime, const float InDeltaTime)
+{
+	RefreshNotifyHookBinding();
+	Super::Tick(AllottedGeometry, InCurrentTime, InDeltaTime);
+}
+
+void SBindingsPanel::RefreshNotifyHookBinding()
 {
 	if (NotifyHook->Binding.BindingId.IsValid())
 	{
@@ -221,8 +243,6 @@ void SBindingsPanel::Tick(const FGeometry& AllottedGeometry, const double InCurr
 			}
 		}
 	}
-
-	Super::Tick(AllottedGeometry, InCurrentTime, InDeltaTime);
 }
 
 void SBindingsPanel::AddBindingToWidgetList(const TSet<FWidgetReference>& WidgetsToAddBinding)
@@ -291,6 +311,24 @@ bool SBindingsPanel::CanAddBinding() const
 	return MVVMExtensionPtr && MVVMExtensionPtr->GetBlueprintView() != nullptr;
 }
 
+void SBindingsPanel::RefreshDetailsView()
+{
+	UMVVMWidgetBlueprintExtension_View* MVVMExtensionPtr = MVVMExtension.Get();
+	if (MVVMExtensionPtr && MVVMExtensionPtr->GetBlueprintView())
+	{
+		RefreshNotifyHookBinding();
+		if (NotifyHook->Binding.BindingId.IsValid() && MVVMExtensionPtr->GetBlueprintView()->GetBinding(NotifyHook->Binding.BindingId))
+		{
+			TSharedRef<FStructOnScope> StructScope = MakeShared<FStructOnScope>(FMVVMBlueprintViewBinding::StaticStruct(), reinterpret_cast<uint8*>(&NotifyHook->Binding));
+			StructDetailsView->SetStructureData(StructScope);
+			DetailContainer->SetContent(StructDetailsView->GetWidget().ToSharedRef());
+			return;
+		}
+	}
+	DetailsView->SetObject(nullptr);
+	DetailContainer->SetContent(DetailsView.ToSharedRef());
+}
+
 FText SBindingsPanel::GetAddBindingText() const
 {
 	if (TSharedPtr<FWidgetBlueprintEditor> WidgetEditor = WeakBlueprintEditor.Pin())
@@ -556,6 +594,10 @@ TSharedRef<SWidget> SBindingsPanel::GenerateEditViewWidget()
 	if (MVVMExtension.Get())
 	{
 		BindingsList = SNew(SBindingsList, StaticCastSharedRef<SBindingsPanel>(AsShared()), WeakBlueprintEditor.Pin(), MVVMExtension.Get());
+		if (UMVVMBlueprintView* View = MVVMExtension->GetBlueprintView())
+		{
+			View->OnBindingsUpdated.AddSP(this, &SBindingsPanel::RefreshDetailsView);
+		}
 	}
 
 	DetailContainer = SNew(SBorder)
diff --git a/Engine/Plugins/Runtime/ModelViewViewModel/Source/ModelViewViewModelEditor/Private/Widgets/SMVVMViewBindingPanel.h b/Engine/Plugins/Runtime/ModelViewViewModel/Source/ModelViewViewModelEditor/Private/Widgets/SMVVMViewBindingPanel.h
--- a/Engine/Plugins/Runtime/ModelViewViewModel/Source/ModelViewViewModelEditor/Private/Widgets/SMVVMViewBindingPanel.h
+++ b/Engine/Plugins/Runtime/ModelViewViewModel/Source/ModelViewViewModelEditor/Private/Widgets/SMVVMViewBindingPanel.h
@@ -77,6 +77,10 @@ private:
 
 	FReply HandleCreateViewModelClicked();
 
+	bool IsDetailsViewEditingEnabled() const;
+	void RefreshDetailsView();
+	void RefreshNotifyHookBinding();
+
 	EVisibility GetViewModelMessageVisibility() const;
 	EVisibility GetBindingMessageVisibility() const;
 
