UE-206013 - StrongObjectPtr Crash running with -asyncloadingthread

Early runs of async loading in the editor were tripping the TStrongObjectPtr's check that it can only be created on the game thread. The was caused by the FMVVMViewBlueprintCompiler::BlueprintView field. Since we're dealing with a non-UObject here, we typically introduce the object reference either by reflecting the struct, or forcing it to inherit from FGCObject.

However, we can also take advantage of the ownership structure. Both the FMVVMViewBlueprintCompiler and UMVVMBlueprintView are owned by UMVVMWidgetBlueprintExtension_View. Technically, FMVVMViewBlueprintCompiler doesn't need a reference to UMVVMBlueprintView if it's already owned by the extension.

#rnx

diff --git a/Engine/Plugins/Runtime/ModelViewViewModel/Source/ModelViewViewModelBlueprint/Private/MVVMViewBlueprintCompiler.cpp b/Engine/Plugins/Runtime/ModelViewViewModel/Source/ModelViewViewModelBlueprint/Private/MVVMViewBlueprintCompiler.cpp
--- a/Engine/Plugins/Runtime/ModelViewViewModel/Source/ModelViewViewModelBlueprint/Private/MVVMViewBlueprintCompiler.cpp
+++ b/Engine/Plugins/Runtime/ModelViewViewModel/Source/ModelViewViewModelBlueprint/Private/MVVMViewBlueprintCompiler.cpp
@@ -159,7 +159,7 @@ FMVVMViewBlueprintCompiler::FMVVMViewBlueprintCompiler(FWidgetBlueprintCompilerC
 	, BlueprintView(InBlueprintView)
 	, BindingLibraryCompiler(InCreationContext.WidgetBlueprint())
 {
-	check(BlueprintView.IsValid());
+	check(BlueprintView);
 }
 
 
diff --git a/Engine/Plugins/Runtime/ModelViewViewModel/Source/ModelViewViewModelBlueprint/Private/MVVMViewBlueprintCompiler.h b/Engine/Plugins/Runtime/ModelViewViewModel/Source/ModelViewViewModelBlueprint/Private/MVVMViewBlueprintCompiler.h
--- a/Engine/Plugins/Runtime/ModelViewViewModel/Source/ModelViewViewModelBlueprint/Private/MVVMViewBlueprintCompiler.h
+++ b/Engine/Plugins/Runtime/ModelViewViewModel/Source/ModelViewViewModelBlueprint/Private/MVVMViewBlueprintCompiler.h
@@ -338,7 +338,7 @@ private:
 private:
 	TMap<FName, UWidget*> WidgetNameToWidgetPointerMap;
 	FWidgetBlueprintCompilerContext& WidgetBlueprintCompilerContext;
-	TStrongObjectPtr<UMVVMBlueprintView> BlueprintView = nullptr;
+	TObjectPtr<UMVVMBlueprintView> BlueprintView = nullptr;
 	FCompiledBindingLibraryCompiler BindingLibraryCompiler;
 	bool bIsCreateVariableStepValid = true;
 	bool bIsCreateFunctionsStepValid = true;
