MVVM: Change the source sort algo to support destination.
At initalization, one at the time, the source initialize and then execute the bindings. Binding can use a source a soon to be initialize source and a binding can set the value of a source. Both cases need to be taken into account when sorting the sources.
When initializing, the binding can trigger a FieldNotify of another initialized source that can trigger that same source that try to initialize. Add a message if that occurs.

diff --git a/Engine/Plugins/Runtime/ModelViewViewModel/Source/ModelViewViewModel/Private/View/MVVMView.cpp b/Engine/Plugins/Runtime/ModelViewViewModel/Source/ModelViewViewModel/Private/View/MVVMView.cpp
--- a/Engine/Plugins/Runtime/ModelViewViewModel/Source/ModelViewViewModel/Private/View/MVVMView.cpp
+++ b/Engine/Plugins/Runtime/ModelViewViewModel/Source/ModelViewViewModel/Private/View/MVVMView.cpp
@@ -451,20 +451,24 @@ void UMVVMView::ExecuteBindingImmediately(const FMVVMViewClass_Binding& ClassBin
 		}
 #endif
 	}
-	else if ((ClassBinding.GetSources() & ClassExtension->GetOptionalSources()) != ClassBinding.GetSources())
+	else
 	{
+		const uint64 MissingSources = ClassBinding.GetSources() & (~ValidSources);
+		if ((MissingSources & ClassExtension->GetOptionalSources()) != MissingSources)
+		{
 #if UE_WITH_MVVM_DEBUGGING
-		UE::MVVM::FMessageLog Log(GetUserWidget());
-		Log.Error(FText::Format(LOCTEXT("ExecuteBindingFailInvalidSource", "The binding '{0}' was not executed. There are invalid sources.")
-			, FText::FromString(ClassBinding.ToString(ClassExtension, FMVVMViewClass_Binding::FToStringArgs::Short()))
-		));
-		UE::MVVM::FDebugging::BroadcastLibraryBindingExecuted(this, KeyForLog, FMVVMCompiledBindingLibrary::EExecutionFailingReason::InvalidSource);
+			UE::MVVM::FMessageLog Log(GetUserWidget());
+			Log.Error(FText::Format(LOCTEXT("ExecuteBindingFailInvalidSource", "The binding '{0}' was not executed. There are invalid sources.")
+				, FText::FromString(ClassBinding.ToString(ClassExtension, FMVVMViewClass_Binding::FToStringArgs::Short()))
+			));
+			UE::MVVM::FDebugging::BroadcastLibraryBindingExecuted(this, KeyForLog, FMVVMCompiledBindingLibrary::EExecutionFailingReason::InvalidSource);
 #else
-		UE::MVVM::FMessageLog Log(GetUserWidget());
-		Log.Error(FText::Format(LOCTEXT("ExecuteBindingFailInvalidSource", "The binding '{0}' was not executed. There are invalid sources.")
-			, FText::AsNumber(KeyForLog.GetIndex())
-		));
+			UE::MVVM::FMessageLog Log(GetUserWidget());
+			Log.Error(FText::Format(LOCTEXT("ExecuteBindingFailInvalidSource", "The binding '{0}' was not executed. There are invalid sources.")
+				, FText::AsNumber(KeyForLog.GetIndex())
+			));
 #endif
+		}
 	}
 }
 
@@ -483,7 +487,7 @@ bool UMVVMView::EvaluateSource(FMVVMViewClass_SourceKey SourceKey)
 	UObject* NewSource = ClassSource.GetOrCreateInstance(ClassExtension, this, UserWidget);
 	ensureMsgf((NewSource == nullptr || NewSource->GetClass()->ImplementsInterface(UNotifyFieldValueChanged::StaticClass())), TEXT("The source has implement the interface. It should be check at compile time."));
 	bool bResult = SetSourceInternal(SourceKey, NewSource, true);
-	if (bResult)
+	if (!bResult)
 	{
 		UE::MVVM::FMessageLog Log(GetUserWidget());
 		Log.Error(FText::Format(LOCTEXT("ExecuteBindingFailEvaluate", "The evaluate source '{0}' fail execution.")
@@ -517,10 +521,17 @@ void UMVVMView::HandledLibraryBindingValueChanged(UObject* InSource, UE::FieldNo
 		}
 		const FMVVMView_SourceKey ViewSourceKey = FMVVMView_SourceKey(ViewSourceIndex);
 		const FMVVMView_Source& ViewSource = Sources[ViewSourceIndex];
-		ensure(ViewSource.bBindingsInitialized);
+
 		if (!ViewSource.bBindingsInitialized || !ViewSource.bSourceInitialized)
 		{
 			// we do not want to run a binding while we are initializing the bindings.
+#if UE_WITH_MVVM_DEBUGGING
+			UE::MVVM::FMessageLog Log(GetUserWidget());
+			Log.Warning(FText::Format(LOCTEXT("ExecuteBindingFailWhileInitializing", "The Field '{0}' could not execute the bindings while initializing."), FText::FromName(InFieldId.GetName())));
+#else
+			UE::MVVM::FMessageLog Log(GetUserWidget());
+			Log.Warning(FText::Format(LOCTEXT("ExecuteBindingFailInvalidSource", "The Field '{0}'  could not execute the bindings while initializing."), FText::FromName(InFieldId.GetName())));
+#endif
 			return;
 		}
 
@@ -560,7 +571,7 @@ void UMVVMView::HandledLibraryBindingValueChanged(UObject* InSource, UE::FieldNo
 							ensureAlwaysMsgf(false, TEXT("Recursive binding detected"));
 							//Todo add more infos. Callstack maybe? Log the chain?
 							UE::MVVM::FMessageLog Log(Self->GetUserWidget());
-							Log.Error(LOCTEXT("RecursionDetected", "A recursive binding was detected (ie. A->B->C->A->B->C) at runtime."));
+							Log.Warning(LOCTEXT("RecursionDetected", "A recursive binding was detected (ie. A->B->C->A->B->C) at runtime."));
 							return;
 						}
 
@@ -604,7 +615,7 @@ void UMVVMView::ExecuteDelayedBinding(const FMVVMViewClass_BindingKey& DelayedBi
 				ensureAlwaysMsgf(false, TEXT("Recursive binding detected"));
 				//Todo add more infos. Callstack maybe? Log the chain?
 				UE::MVVM::FMessageLog Log(Self->GetUserWidget());
-				Log.Error(LOCTEXT("RecursionDetected", "A recursive binding was detected (ie. A->B->C->A->B->C) at runtime."));
+				Log.Warning(LOCTEXT("RecursionDetected", "A recursive binding was detected (ie. A->B->C->A->B->C) at runtime."));
 				return;
 			}
 
diff --git a/Engine/Plugins/Runtime/ModelViewViewModel/Source/ModelViewViewModelBlueprint/Private/MVVMViewBlueprintCompiler.cpp b/Engine/Plugins/Runtime/ModelViewViewModel/Source/ModelViewViewModelBlueprint/Private/MVVMViewBlueprintCompiler.cpp
--- a/Engine/Plugins/Runtime/ModelViewViewModel/Source/ModelViewViewModelBlueprint/Private/MVVMViewBlueprintCompiler.cpp
+++ b/Engine/Plugins/Runtime/ModelViewViewModel/Source/ModelViewViewModelBlueprint/Private/MVVMViewBlueprintCompiler.cpp
@@ -999,7 +999,7 @@ void FMVVMViewBlueprintCompiler::CreateRequiredProperties(const FWidgetBlueprint
 					bIsCreateVariableStepValid = false;
 					continue;
 				}
-				RunGenerateCompilerSourceContext(false, Pin.GetPath(), Pin.GetId());
+				RunGenerateCompilerSourceContext(true, Pin.GetPath(), Pin.GetId());
 			}
 		}
 	}
@@ -1230,6 +1230,7 @@ void FMVVMViewBlueprintCompiler::CreateWriteFieldContexts(const FWidgetBlueprint
 			TSharedRef<FGeneratedWriteFieldPathContext> WriteFieldPath = MakeShared<FGeneratedWriteFieldPathContext>();
 			GeneratedWriteFieldPaths.Add(WriteFieldPath);
 			WriteFieldPath->UsedByBindings.AddUnique(ValidBinding);
+			//WriteFieldPath->OptionalSource;
 			WriteFieldPath->GeneratedFields = MoveTemp(FieldContextResult.GetValue().GeneratedFields);
 			WriteFieldPath->SkeletalGeneratedFields = MoveTemp(FieldContextResult.GetValue().SkeletalGeneratedFields);
 			WriteFieldPath->GeneratedFrom = DestinationPropertyPath.GetSource(WidgetBlueprintCompilerContext.WidgetBlueprint());
@@ -1275,6 +1276,7 @@ void FMVVMViewBlueprintCompiler::CreateWriteFieldContexts(const FWidgetBlueprint
 					TSharedRef<FGeneratedWriteFieldPathContext> WriteFieldPath = MakeShared<FGeneratedWriteFieldPathContext>();
 					GeneratedWriteFieldPaths.Add(WriteFieldPath);
 					WriteFieldPath->UsedByEvents.AddUnique(ValidEvent);
+					//WriteFieldPath->OptionalSource;
 					WriteFieldPath->GeneratedFields = MoveTemp(FieldContextResult.GetValue().GeneratedFields);
 					WriteFieldPath->SkeletalGeneratedFields = MoveTemp(FieldContextResult.GetValue().SkeletalGeneratedFields);
 					WriteFieldPath->GeneratedFrom = EventPtr->GetDestinationPath().GetSource(WidgetBlueprintCompilerContext.WidgetBlueprint());
@@ -1412,10 +1414,12 @@ bool FMVVMViewBlueprintCompiler::PreCompile(UWidgetBlueprintGeneratedClass* Clas
 	AddWarningForPropertyWithMVVMAndLegacyBinding(Class);
 
 	GeneratedReadFieldPaths.Reset();
-	FixWriteFieldPathContext(Class);
 	CreateReadFieldContexts(Class);
 	CreateCreatorContentFromBindingSource(Class);
 
+	// NB. The dynamic sources are created.
+	FixFieldPathContext(Class);
+
 	if (!AreStepsValid())
 	{
 		return false;
@@ -1424,6 +1428,7 @@ bool FMVVMViewBlueprintCompiler::PreCompile(UWidgetBlueprintGeneratedClass* Clas
 	PreCompileViewModelCreatorContexts(Class);
 	PreCompileBindings(Class);
 	PreCompileEvents(Class);
+	PreCompileSourceDependencies(Class);
 
 	return AreStepsValid();
 }
@@ -1488,11 +1493,149 @@ bool FMVVMViewBlueprintCompiler::Compile(UWidgetBlueprintGeneratedClass* Class,
 }
 
 
-void FMVVMViewBlueprintCompiler::FixWriteFieldPathContext(UWidgetBlueprintGeneratedClass* Class)
+void FMVVMViewBlueprintCompiler::FixFieldPathContext(UWidgetBlueprintGeneratedClass* Class)
 {
+	TSharedRef<FCompilerBindingSource>* SelfSourcePtr = NeededBindingSources.FindByPredicate([](const TSharedRef<FCompilerBindingSource>& Other)
+		{
+			return Other->Type == FCompilerBindingSource::EType::Self;
+		});
+
+	auto FindProperty = [](FMVVMConstFieldVariant Field) -> const FProperty*
+	{
+		if (Field.IsProperty())
+		{
+			return Field.GetProperty();
+		}
+		else if (Field.IsFunction() && Field.GetFunction())
+		{
+			return BindingHelper::GetReturnProperty(Field.GetFunction());
+		}
+		return nullptr;
+	};
+
+	auto FindDynamicCreatorContext = [Self = this, &FindProperty](TArray<UE::MVVM::FMVVMConstFieldVariant>& GeneratedFields, int32 Index, TSharedPtr<FCompilerBindingSource>& InOutSource) -> bool
+	{
+		bool bBreak = true;
+
+		UE::MVVM::FMVVMConstFieldVariant NextField = GeneratedFields[Index];
+		if (const FObjectProperty* FieldObjectProperty = CastField<FObjectProperty>(FindProperty(NextField)))
+		{
+			const UClass* NextFieldObjectPtrClass = FieldObjectProperty->PropertyClass;
+			for (const TSharedRef<FCompilerSourceViewModelDynamicCreatorContext>& DynamicCreatorContext : Self->SourceViewModelDynamicCreatorContexts)
+			{
+				if (DynamicCreatorContext->ParentSource == InOutSource
+					&& DynamicCreatorContext->NotificationId.GetFieldName() == NextField.GetName()
+					&& DynamicCreatorContext->Source->AuthoritativeClass == NextFieldObjectPtrClass)
+				{
+					InOutSource = DynamicCreatorContext->Source;
+					bBreak = false;
+					break;
+				}
+			}
+		}
+
+		return bBreak;
+	};
+
+	//// Sanity check. Set Source to FGeneratedReadFieldPathContext
+	//{
+	//	for (TSharedRef<FGeneratedReadFieldPathContext>& GeneratedRead : GeneratedReadFieldPaths)
+	//	{
+	//		// ReadFieldPath should already be valid. Only confirm it here.
+	//		TSharedPtr<FCompilerBindingSource> NewSource = SelfSourcePtr ? *SelfSourcePtr : TSharedPtr<FCompilerBindingSource>();
+	//		if (GeneratedRead->GeneratedFields.Num() > 0)
+	//		{
+	//			UE::MVVM::FMVVMConstFieldVariant NextField = GeneratedRead->GeneratedFields[0];
+	//			if (const FObjectProperty* FieldObjectProperty = CastField<FObjectProperty>(FindProperty(NextField)))
+	//			{
+	//				UClass* NextFieldObjectPtrClass = FieldObjectProperty->PropertyClass;
+	//				TSharedRef<FCompilerBindingSource>* BindingSource = NeededBindingSources.FindByPredicate(
+	//					[NextField, NextFieldObjectPtrClass](TSharedRef<FCompilerBindingSource>& Other)
+	//					{
+	//						return Other->Name == NextField.GetName()
+	//							&& NextFieldObjectPtrClass == Other->AuthoritativeClass;
+	//					}
+	//				);
+
+	//				if (BindingSource)
+	//				{
+	//					NewSource = *BindingSource;
+	//					for (int32 GeneratedFieldIndex = 1; GeneratedFieldIndex < GeneratedRead->GeneratedFields.Num() - 1; ++GeneratedFieldIndex)
+	//					{
+	//						if (FindDynamicCreatorContext(GeneratedRead->GeneratedFields, GeneratedFieldIndex, NewSource))
+	//						{
+	//							break;
+	//						}
+	//					}
+	//				}
+	//			}
+	//		}
+
+	//		if (GeneratedRead->Source != NewSource)
+	//		{
+	//			// We only do this test as a sanity check and because it's the same algo for FGeneratedWriteFieldPathContext 
+	//			AddMessage(FText::Format(LOCTEXT("InternalErrorNotSameSource", "Internal error. The read path {0} source do not matches with the previous calculated source."), ::UE::MVVM::FieldPathHelper::ToText(GeneratedRead->GeneratedFields))
+	//				, EMessageType::Info
+	//			);
+	//		}
+	//	}
+	//}
+
+	// Set OptionalSource and OptionalDependencySource to FGeneratedWriteFieldPathContext
+	for (TSharedRef<FGeneratedWriteFieldPathContext>& GeneratedDestination : GeneratedWriteFieldPaths)
+	{
+		//ViewmodelA.ViewmodelB = ViewmodelC.Value				-> ViewmodelA needs to init before ViewmodelA_ViewmodelB, ViewmodelA before ViewmodelC, ViewmodelC before ViewmodelA_ViewmodelB
+			// OptionSource = ViewmodelA
+			// OptionalDependencySource = ViewmodelA_ViewmodelB. We do not have the info for ViewmodelC (only in the binding) and there could be more than one read path. Set the flag and add the dependency later.
+		//ViewmodelA.ViewmodelB.Value = ViewmodelC.Value		-> ViewmodelA needs to init before ViewmodelA_ViewmodelB, ViewmodelA before ViewmodelC, ViewmodelA_ViewmodelB before ViewmodelC
+			// OptionSource = ViewmodelA_ViewmodelB
+			// OptionalDependencySource = null. We set a value, not a source.
+
+		if (GeneratedDestination->GeneratedFields.Num() > 0)
+		{
+			UE::MVVM::FMVVMConstFieldVariant NextField = GeneratedDestination->GeneratedFields[0];
+			if (const FObjectProperty* FieldObjectProperty = CastField<FObjectProperty>(FindProperty(NextField)))
+			{
+				const UClass* NextFieldObjectPtrClass = FieldObjectProperty->PropertyClass;
+				TSharedRef<FCompilerBindingSource>* BindingSource = NeededBindingSources.FindByPredicate(
+					[NextField, NextFieldObjectPtrClass](TSharedRef<FCompilerBindingSource>& Other)
+					{
+						return Other->Name == NextField.GetName()
+							&& NextFieldObjectPtrClass == Other->AuthoritativeClass;
+					}
+				);
+
+				if (BindingSource)
+				{
+					if (GeneratedDestination->GeneratedFields.Num() > 1)
+					{
+						GeneratedDestination->OptionalSource = *BindingSource;
+						int32 GeneratedFieldIndex = 1;
+						for (; GeneratedFieldIndex < GeneratedDestination->GeneratedFields.Num() - 1; ++GeneratedFieldIndex)
+						{
+							if (FindDynamicCreatorContext(GeneratedDestination->GeneratedFields, GeneratedFieldIndex, GeneratedDestination->OptionalSource))
+							{
+								break;
+							}
+						}
+						if (GeneratedFieldIndex == GeneratedDestination->GeneratedFields.Num() - 1)
+						{
+							FindDynamicCreatorContext(GeneratedDestination->GeneratedFields, GeneratedFieldIndex, GeneratedDestination->OptionalDependencySource);
+						}
+					}
+					else
+					{
+						GeneratedDestination->OptionalSource = SelfSourcePtr ? *SelfSourcePtr : TSharedPtr<FCompilerBindingSource>();
+						GeneratedDestination->OptionalDependencySource = *BindingSource;
+					}
+				}
+			}
+		}
+	}
+
+	// If a graph was generated for the FGeneratedWriteFieldPathContext, then use it instead for the SkeletalGeneratedFields
 	for (TSharedRef<FGeneratedWriteFieldPathContext>& GeneratedDestination : GeneratedWriteFieldPaths)
 	{
-		// If a graph was generated. Use it instead for the SkeletalGeneratedFields
 		if (!GeneratedDestination->GeneratedFunctionName.IsNone())
 		{
 			UFunction* GeneratedFunction = Class->FindFunctionByName(GeneratedDestination->GeneratedFunctionName);
@@ -1583,7 +1726,7 @@ void FMVVMViewBlueprintCompiler::CreateReadFieldContexts(UWidgetBlueprintGenerat
 					}
 
 					ReadFieldContext->NotificationField = CreateFieldResult.GetValue();
-					ReadFieldContext->OptionalSource = ReadFieldContext->NotificationField->Source;
+					ReadFieldContext->Source = ReadFieldContext->NotificationField->Source;
 				}
 			}
 			return MakeValue();
@@ -1602,7 +1745,7 @@ void FMVVMViewBlueprintCompiler::CreateReadFieldContexts(UWidgetBlueprintGenerat
 			if (!Found)
 			{
 				Found = MakeShared<FGeneratedReadFieldPathContext>();
-				Found->OptionalSource = MoveTemp(CreateSourceResult.GetValue().Source);
+				Found->Source = MoveTemp(CreateSourceResult.GetValue().OptionalSource);
 				Found->GeneratedFields = MoveTemp(CreateSourceResult.GetValue().GeneratedFields);
 				Found->SkeletalGeneratedFields = MoveTemp(CreateSourceResult.GetValue().SkeletalGeneratedFields);
 
@@ -1638,7 +1781,7 @@ void FMVVMViewBlueprintCompiler::CreateReadFieldContexts(UWidgetBlueprintGenerat
 						if (!Found)
 						{
 							Found = MakeShared<FGeneratedReadFieldPathContext>();
-							Found->OptionalSource = MoveTemp(CreateSourceResult.GetValue().Source);
+							Found->Source = MoveTemp(CreateSourceResult.GetValue().OptionalSource);
 							Found->GeneratedFields = MoveTemp(CreateSourceResult.GetValue().GeneratedFields);
 							Found->SkeletalGeneratedFields = MoveTemp(CreateSourceResult.GetValue().SkeletalGeneratedFields);
 
@@ -1685,7 +1828,7 @@ void FMVVMViewBlueprintCompiler::CreateReadFieldContexts(UWidgetBlueprintGenerat
 					if (!Found)
 					{
 						Found = MakeShared<FGeneratedReadFieldPathContext>();
-						Found->OptionalSource = MoveTemp(CreateSourceResult.GetValue().Source);
+						Found->Source = MoveTemp(CreateSourceResult.GetValue().OptionalSource);
 						Found->GeneratedFields = MoveTemp(CreateSourceResult.GetValue().GeneratedFields);
 						Found->SkeletalGeneratedFields = MoveTemp(CreateSourceResult.GetValue().SkeletalGeneratedFields);
 						GeneratedReadFieldPaths.Add(Found.ToSharedRef());
@@ -1975,30 +2118,49 @@ void FMVVMViewBlueprintCompiler::CompileSources(const FCompiledBindingLibraryCom
 	struct FSortData
 	{
 		FSortData() = default;
-		FName SourceName;
-		FName ParentSourceName;
+		TSharedPtr<FCompilerBindingSource> Source;
+		TArray<FName, TInlineAllocator<8>> Errors;
 		int32 SortIndex = -1;
+		bool bCalculating = false;
 
-		void CalculateSortIndex(TMap<FName, FSortData>& Map)
+		int32 CalculateSortIndex(FName RequestedBy, TMap<FName, FSortData>& Map)
 		{
 			if (SortIndex < 0)
 			{
-				if (ParentSourceName.IsNone())
+				if (bCalculating)
 				{
 					SortIndex = 0;
+					Errors.Add(RequestedBy);
 				}
 				else
 				{
-					Map[ParentSourceName].CalculateSortIndex(Map);
-					// calculate the Depth recursively
-					SortIndex = Map[ParentSourceName].SortIndex + 1;
+					bCalculating = true;
+					if (Source->Dependencies.Num() == 0)
+					{
+						SortIndex = 0;
+					}
+					else
+					{
+						for (TWeakPtr<FCompilerBindingSource> WeakDependency : Source->Dependencies)
+						{
+							TSharedPtr<FCompilerBindingSource> Dependency = WeakDependency.Pin();
+							if (ensure(Dependency))
+							{
+								SortIndex = FMath::Max(Map[Dependency->Name].CalculateSortIndex(Source->Name, Map) + 1, SortIndex);
+							}
+						}
+					}
+					bCalculating = false;
 				}
 			}
+			return SortIndex;
 		}
 	};
 
 	TMap<FName, FSortData> SortDatas;
+	SortDatas.Reserve(ViewModelCreatorContexts.Num() + WidgetCreatorContexts.Num());
 	TArray<FMVVMViewClass_Source> UnsortedSourceCreators;
+	UnsortedSourceCreators.Reserve(ViewModelCreatorContexts.Num() + WidgetCreatorContexts.Num());
 
 	for (FCompilerViewModelCreatorContext& SourceCreatorContext : ViewModelCreatorContexts)
 	{
@@ -2101,14 +2263,9 @@ void FMVVMViewBlueprintCompiler::CompileSources(const FCompiledBindingLibraryCom
 			continue;
 		}
 
-		FName ParentSourceName;
 		if (SourceCreatorContext.DynamicContext)
 		{
-			if (SourceCreatorContext.DynamicContext->ParentSource)
-			{
-				ParentSourceName = SourceCreatorContext.DynamicContext->ParentSource->Name;
-			}
-			else
+			if (!SourceCreatorContext.DynamicContext->ParentSource)
 			{
 				AddMessageForViewModel(ViewModelContext
 					, LOCTEXT("ViewModelWithoutValidDyanmicParentSource", "The viewmodel doesn't have a valid parent source.")
@@ -2129,10 +2286,14 @@ void FMVVMViewBlueprintCompiler::CompileSources(const FCompiledBindingLibraryCom
 		CompiledSourceCreator.Flags |= (uint16)FMVVMViewClass_Source::EFlags::IsViewModel;
 		CompiledSourceCreator.Flags |= ViewModelContext.bExposeInstanceInEditor ? (uint16)FMVVMViewClass_Source::EFlags::IsViewModelInstanceExposed : 0;
 
-		FSortData SortData;
-		SortData.SourceName = CompiledSourceCreator.GetName();
-		SortData.ParentSourceName = ParentSourceName;
-		SortDatas.Add(CompiledSourceCreator.GetName(), SortData);
+		{
+			FSortData SortData;
+			SortData.Source = SourceCreatorContext.Source.ToSharedRef();
+			SortDatas.Add(CompiledSourceCreator.GetName(), SortData);
+
+			ensure(SortData.Source->Name == CompiledSourceCreator.GetName());
+		}
+
 		UnsortedSourceCreators.Add(MoveTemp(CompiledSourceCreator));
 	}
 
@@ -2140,7 +2301,7 @@ void FMVVMViewBlueprintCompiler::CompileSources(const FCompiledBindingLibraryCom
 	for (FCompilerWidgetCreatorContext& WidgetCreator : WidgetCreatorContexts)
 	{
 		FMVVMViewClass_Source CompiledSourceCreator;
-		ensure(WidgetCreator.Source->AuthoritativeClass&& WidgetCreator.Source->AuthoritativeClass->ImplementsInterface(UNotifyFieldValueChanged::StaticClass()));
+		ensure(WidgetCreator.Source->AuthoritativeClass && WidgetCreator.Source->AuthoritativeClass->ImplementsInterface(UNotifyFieldValueChanged::StaticClass()));
 		CompiledSourceCreator.ExpectedSourceType = const_cast<UClass*>(WidgetCreator.Source->AuthoritativeClass);
 		CompiledSourceCreator.PropertyName = WidgetCreator.Source->Name;
 		CompiledSourceCreator.Flags = 0;
@@ -2160,8 +2321,12 @@ void FMVVMViewBlueprintCompiler::CompileSources(const FCompiledBindingLibraryCom
 			CompiledSourceCreator.FieldPath = *CompiledFieldPath;
 		}
 
-		FSortData SortData;
-		SortDatas.Add(CompiledSourceCreator.GetName(), SortData);
+		{
+			FSortData SortData;
+			SortData.Source = WidgetCreator.Source;
+			SortDatas.Add(CompiledSourceCreator.GetName(), SortData);
+		}
+
 		UnsortedSourceCreators.Add(MoveTemp(CompiledSourceCreator));
 	}
 
@@ -2173,7 +2338,7 @@ void FMVVMViewBlueprintCompiler::CompileSources(const FCompiledBindingLibraryCom
 			const FMVVMViewClass_Source* FoundClassSource = UnsortedSourceCreators.FindByPredicate([ToFindName = BindingSource->Name](const FMVVMViewClass_Source& Other){ return Other.GetName() == ToFindName; });
 			if (FoundClassSource == nullptr)
 			{
-				AddMessage(FText::Format(LOCTEXT("CompileSources_MissingSources", "Internal error. The source {0} was not compiled.."), FText::FromName(BindingSource->Name))
+				AddMessage(FText::Format(LOCTEXT("CompileSources_MissingSources", "Internal error. The source {0} was not compiled."), FText::FromName(BindingSource->Name))
 					, EMessageType::Warning
 				);
 			}
@@ -2182,14 +2347,28 @@ void FMVVMViewBlueprintCompiler::CompileSources(const FCompiledBindingLibraryCom
 		ensure(SortDatas.Num() == UnsortedSourceCreators.Num());
 	}
 
-	// sort the source creators by priority and then by name
+	// Sort the source creators by priority and then by name
 	if (UnsortedSourceCreators.Num() > 1)
 	{
+		// Calculate dependencies
 		for (auto& SortDataPair : SortDatas)
 		{
-			SortDataPair.Value.CalculateSortIndex(SortDatas);
+			SortDataPair.Value.CalculateSortIndex(SortDataPair.Key, SortDatas);
+		}
+
+		// Report errors
+		for (const auto& SortDataPair : SortDatas)
+		{
+			for (FName OtherSource : SortDataPair.Value.Errors)
+			{
+				AddMessage(
+					FText::Format(LOCTEXT("CompileSources_Dependencies", "The source {0} circularly depends on the source {1}"), FText::FromName(SortDataPair.Key), FText::FromName(OtherSource))
+					, EMessageType::Info
+				);
+			}
 		}
 
+		// Sort the sources
 		UnsortedSourceCreators.Sort([&SortDatas](const FMVVMViewClass_Source& A, const FMVVMViewClass_Source& B)
 			{
 				int32 ASortIndex = SortDatas[A.GetName()].SortIndex;
@@ -2572,19 +2751,34 @@ void FMVVMViewBlueprintCompiler::CompileBindings(const FCompiledBindingLibraryCo
 		NewBinding.SourceBitField = 0;
 		NewBinding.EditorId = Binding.BindingId;
 
+		// Add the write source.
+		if (ValidBinding->WritePath && ValidBinding->WritePath->OptionalSource)
+		{
+			int32 ViewExtensionSourceCreatorsIndex = ViewExtension->Sources.IndexOfByPredicate([LookFor = ValidBinding->WritePath->OptionalSource->Name](const FMVVMViewClass_Source& Other)
+				{
+					return Other.GetName() == LookFor;
+				});
+			if (ViewExtension->Sources.IsValidIndex(ViewExtensionSourceCreatorsIndex))
+			{
+				
+				FMVVMViewClass_SourceKey FieldClassSourceKey = FMVVMViewClass_SourceKey(ViewExtensionSourceCreatorsIndex);
+				NewBinding.SourceBitField |= FieldClassSourceKey.GetBit();
+			}
+		}
+
 		TArray<FMVVMViewClass_SourceKey, TInlineAllocator<16>>  SharedExecuteAtInitializationBindings;
 		int32 SharedBindingsCount = 0;
 		// Find the source needed by the binding. Also generate every bindings on that source (that need to register to the FieldNotify).
 		for (TSharedPtr<FGeneratedReadFieldPathContext> ReadPath : ValidBinding->ReadPaths)
 		{
-			if (ReadPath->OptionalSource == nullptr)
+			if (ReadPath->Source == nullptr)
 			{
 				AddMessageForBinding(Binding, LOCTEXT("InvalidSourceInternal", "Internal error. The binding has an invalid source."), EMessageType::Error, FMVVMBlueprintPinId());
 				bIsCompileStepValid = false;
 				continue;
 			}
 
-			int32 ViewExtensionSourceCreatorsIndex = ViewExtension->Sources.IndexOfByPredicate([LookFor = ReadPath->OptionalSource->Name](const FMVVMViewClass_Source& Other)
+			int32 ViewExtensionSourceCreatorsIndex = ViewExtension->Sources.IndexOfByPredicate([LookFor = ReadPath->Source->Name](const FMVVMViewClass_Source& Other)
 				{
 					return Other.GetName() == LookFor;
 				});
@@ -2876,6 +3070,47 @@ void FMVVMViewBlueprintCompiler::CompileEvents(const FCompiledBindingLibraryComp
 }
 
 
+void FMVVMViewBlueprintCompiler::PreCompileSourceDependencies(UWidgetBlueprintGeneratedClass* Class)
+{
+	//i.	ViewmodelA.ViewmodelB = ViewmodelC.Value				-> ViewmodelA needs to init before ViewmodelA_ViewmodelB, ViewmodelA before ViewmodelC, ViewmodelC before ViewmodelA_ViewmodelB
+	//ii.	ViewmodelA.ViewmodelB.Value = ViewmodelC.Value			-> ViewmodelA needs to init before ViewmodelA_ViewmodelB, ViewmodelA before ViewmodelC, ViewmodelA_ViewmodelB before ViewmodelC
+
+	// Dynamic sources depends on their parent. (i and ii ViewmodelA before ViewmodelA_ViewmodelB)
+	for (const TSharedRef<FCompilerSourceViewModelDynamicCreatorContext>& DynamicContext : SourceViewModelDynamicCreatorContexts)
+	{
+		ensure(DynamicContext->Source);
+		ensure(DynamicContext->ParentSource);
+		DynamicContext->Source->Dependencies.AddUnique(DynamicContext->ParentSource);
+	}
+
+	for (const TSharedRef<FCompilerBinding>& ValidBinding : ValidBindings)
+	{
+		if (ValidBinding->WritePath)
+		{
+			for (const TSharedPtr<FGeneratedReadFieldPathContext>& ReadPath : ValidBinding->ReadPaths)
+			{
+				// All the read depends on the write (i and ii ViewmodelA before ViewmodelC)
+				if (ensure(ReadPath->Source))
+				{
+					if (ValidBinding->WritePath->OptionalSource)
+					{
+						ReadPath->Source->Dependencies.AddUnique(ValidBinding->WritePath->OptionalSource);
+					}
+
+					// The write depends on the read (i ViewmodelC before ViewmodelA_ViewmodelB)
+					if (ValidBinding->WritePath->OptionalDependencySource)
+					{
+						ValidBinding->WritePath->OptionalDependencySource->Dependencies.AddUnique(ReadPath->Source);
+					}
+				}
+
+				// NB (ii ViewmodelA_ViewmodelB before ViewmodelC) is implicit
+			}
+		}
+	}
+}
+
+
 void FMVVMViewBlueprintCompiler::SortSourceFields(const FCompiledBindingLibraryCompiler::FCompileResult& CompileResult, UWidgetBlueprintGeneratedClass* Class, UMVVMViewClass* ViewExtension)
 {
 	for (FMVVMViewClass_Source& Source : ViewExtension->Sources)
@@ -2918,6 +3153,7 @@ TValueOrError<FMVVMViewBlueprintCompiler::FGetFieldsResult, FText> FMVVMViewBlue
 			}
 			return MakeValue(*Found);
 		}
+		// It can be valid to not have a BindingSource, if it's a widget in a WritePath
 		return MakeValue(TSharedPtr<FCompilerBindingSource>());
 	};
 
@@ -2943,8 +3179,8 @@ TValueOrError<FMVVMViewBlueprintCompiler::FGetFieldsResult, FText> FMVVMViewBlue
 			return MakeError(LOCTEXT("ViewModelShouldHaveSource", "Internal error. Viewmodel should have a source."));
 		}
 
-		Result.GeneratedFields = GetFields(Class, PropertyName, PropertyPath.GetFields(Class));
-		Result.Source = FindSourceResult.StealValue();
+		Result.GeneratedFields = AppendBaseField(Class, PropertyName, PropertyPath.GetFields(Class));
+		Result.OptionalSource = FindSourceResult.StealValue();
 		break;
 	}
 	case EMVVMBlueprintFieldPathSource::SelfContext:
@@ -2959,8 +3195,8 @@ TValueOrError<FMVVMViewBlueprintCompiler::FGetFieldsResult, FText> FMVVMViewBlue
 			return MakeError(LOCTEXT("WidgetBlueprintShouldHaveSource", "Internal error. The blueprint should have a source."));
 		}
 
-		Result.GeneratedFields = GetFields(Class, FName(), PropertyPath.GetFields(Class));
-		Result.Source = FindSourceResult.StealValue();
+		Result.GeneratedFields = AppendBaseField(Class, FName(), PropertyPath.GetFields(Class));
+		Result.OptionalSource = FindSourceResult.StealValue();
 		break;
 	}
 	case EMVVMBlueprintFieldPathSource::Widget:
@@ -2978,20 +3214,20 @@ TValueOrError<FMVVMViewBlueprintCompiler::FGetFieldsResult, FText> FMVVMViewBlue
 			return MakeError(FindSourceResult.StealError());
 		}
 
-		Result.GeneratedFields = GetFields(Class, DestinationWidgetName, PropertyPath.GetFields(Class));
-		Result.Source = FindSourceResult.StealValue();
+		Result.GeneratedFields = AppendBaseField(Class, DestinationWidgetName, PropertyPath.GetFields(Class));
+		Result.OptionalSource = FindSourceResult.StealValue();
 		break;
 	}
 	default:
-		ensureAlwaysMsgf(false, TEXT("Not supported yet."));
-		Result.GeneratedFields = GetFields(Class, FName(), PropertyPath.GetFields(Class));
+		ensureAlwaysMsgf(false, TEXT("Not supported."));
+		Result.GeneratedFields = AppendBaseField(Class, FName(), PropertyPath.GetFields(Class));
 		break;
 	}
 	return MakeValue(MoveTemp(Result));
 }
 
 
-TArray<FMVVMConstFieldVariant> FMVVMViewBlueprintCompiler::GetFields(const UClass* Class, FName PropertyName, TArray<FMVVMConstFieldVariant> Properties)
+TArray<FMVVMConstFieldVariant> FMVVMViewBlueprintCompiler::AppendBaseField(const UClass* Class, FName PropertyName, TArray<FMVVMConstFieldVariant> Properties)
 {
 	if (PropertyName.IsNone())
 	{
@@ -3010,14 +3246,16 @@ TValueOrError<FMVVMViewBlueprintCompiler::FCreateFieldsResult, FText> FMVVMViewB
 	FMVVMViewBlueprintCompiler::FCreateFieldsResult Result;
 
 	// Evaluate the getter/setter path.
-	TValueOrError<FGetFieldsResult, FText> GetFieldResult = GetFields(Class, PropertyPath);
-	if (GetFieldResult.HasError())
 	{
-		return MakeError(GetFieldResult.StealError());
-	}
+		TValueOrError<FGetFieldsResult, FText> GetFieldResult = GetFields(Class, PropertyPath);
+		if (GetFieldResult.HasError())
+		{
+			return MakeError(GetFieldResult.StealError());
+		}
 
-	Result.Source = MoveTemp(GetFieldResult.GetValue().Source);
-	Result.GeneratedFields = MoveTemp(GetFieldResult.GetValue().GeneratedFields);
+		Result.OptionalSource = MoveTemp(GetFieldResult.GetValue().OptionalSource);
+		Result.GeneratedFields = MoveTemp(GetFieldResult.GetValue().GeneratedFields);
+	}
 
 	if (!IsPropertyPathValid(WidgetBlueprintCompilerContext.WidgetBlueprint(), Result.GeneratedFields))
 	{
@@ -3052,7 +3290,7 @@ TValueOrError<TSharedPtr<FMVVMViewBlueprintCompiler::FCompilerNotifyFieldId>, FT
 	}
 
 	const FieldPathHelper::FParsedNotifyBindingInfo& BindingInfo = BindingInfoResult.GetValue();
-	if (!BindingInfo.NotifyFieldId.IsValid() || ReadFieldContext->OptionalSource == nullptr)
+	if (!BindingInfo.NotifyFieldId.IsValid() || ReadFieldContext->Source == nullptr)
 	{
 		return MakeValue(TSharedPtr<FCompilerNotifyFieldId>());
 	}
@@ -3060,7 +3298,7 @@ TValueOrError<TSharedPtr<FMVVMViewBlueprintCompiler::FCompilerNotifyFieldId>, FT
 
 	FCompilerNotifyFieldId Result;
 	Result.NotificationId = BindingInfo.NotifyFieldId;
-	Result.Source = ReadFieldContext->OptionalSource;
+	Result.Source = ReadFieldContext->Source;
 	Result.ViewModelDynamic.Reset();
 
 	auto GetClassFromField = [](UE::MVVM::FMVVMConstFieldVariant Field) -> const UClass*
@@ -3089,7 +3327,7 @@ TValueOrError<TSharedPtr<FMVVMViewBlueprintCompiler::FCompilerNotifyFieldId>, FT
 		const UClass* ExpectedClass = nullptr;
 		if (BindingInfo.ViewModelIndex < 1 && BindingInfo.NotifyFieldClass)
 		{
-			ExpectedClass = ReadFieldContext->OptionalSource->AuthoritativeClass;
+			ExpectedClass = ReadFieldContext->Source->AuthoritativeClass;
 		}
 		else if (BindingInfo.ViewModelIndex >= 0)
 		{
diff --git a/Engine/Plugins/Runtime/ModelViewViewModel/Source/ModelViewViewModelBlueprint/Private/MVVMViewBlueprintCompiler.h b/Engine/Plugins/Runtime/ModelViewViewModel/Source/ModelViewViewModelBlueprint/Private/MVVMViewBlueprintCompiler.h
--- a/Engine/Plugins/Runtime/ModelViewViewModel/Source/ModelViewViewModelBlueprint/Private/MVVMViewBlueprintCompiler.h
+++ b/Engine/Plugins/Runtime/ModelViewViewModel/Source/ModelViewViewModelBlueprint/Private/MVVMViewBlueprintCompiler.h
@@ -89,12 +89,13 @@ private:
 	// PreCompile
 	void FixCompilerBindingSelfSource(UWidgetBlueprintGeneratedClass* Class);
 	void AddWarningForPropertyWithMVVMAndLegacyBinding(UWidgetBlueprintGeneratedClass* Class);
-	void FixWriteFieldPathContext(UWidgetBlueprintGeneratedClass* Class);
+	void FixFieldPathContext(UWidgetBlueprintGeneratedClass* Class);
 	void CreateReadFieldContexts(UWidgetBlueprintGeneratedClass* Class);
 	void CreateCreatorContentFromBindingSource(UWidgetBlueprintGeneratedClass* Class);
 	void PreCompileViewModelCreatorContexts(UWidgetBlueprintGeneratedClass* Class);
 	void PreCompileBindings(UWidgetBlueprintGeneratedClass* Class);
 	void PreCompileEvents(UWidgetBlueprintGeneratedClass* Class);
+	void PreCompileSourceDependencies(UWidgetBlueprintGeneratedClass* Class);
 
 	// Compile
 	void CompileSources(const FCompiledBindingLibraryCompiler::FCompileResult& CompileResult, UWidgetBlueprintGeneratedClass* Class, UMVVMViewClass* ViewExtension);
@@ -121,6 +122,7 @@ private:
 			Self = 3,
 		};
 		const UClass* AuthoritativeClass = nullptr;
+		TArray<TWeakPtr<FCompilerBindingSource>> Dependencies;
 		FName Name;
 		EType Type;
 		bool bIsOptional = false;
@@ -234,7 +236,7 @@ private:
 		TArray<TWeakPtr<FCompilerBinding>> UsedByBindings;
 		TArray<TWeakPtr<FCompilerEvent>> UsedByEvents;
 
-		TSharedPtr<FCompilerBindingSource> OptionalSource; // Can be invalid if it's an event
+		TSharedPtr<FCompilerBindingSource> Source;
 		TArray<UE::MVVM::FMVVMConstFieldVariant> GeneratedFields; // the string path converted into field
 		TArray<UE::MVVM::FMVVMConstFieldVariant> SkeletalGeneratedFields; // the field path converted with getter and setter
 
@@ -253,8 +255,17 @@ private:
 		TArray<TWeakPtr<FCompilerBinding>> UsedByBindings;
 		TArray<TWeakPtr<FCompilerEvent>> UsedByEvents;
 
-		TArray<UE::MVVM::FMVVMConstFieldVariant> GeneratedFields; // the string path converted into field
-		TArray<UE::MVVM::FMVVMConstFieldVariant> SkeletalGeneratedFields; // the field path converted with getter and setter
+		/**
+		 * Can be invalid if it's a widget with no Read Path.
+		 * It is the start of the path.
+		 */
+		TSharedPtr<FCompilerBindingSource> OptionalSource;
+		TSharedPtr<FCompilerBindingSource> OptionalDependencySource;
+
+		/** The string path converted into field. It always start from the UserWidget. */
+		TArray<UE::MVVM::FMVVMConstFieldVariant> GeneratedFields;
+		/** The field path converted with getter and setter. It always start from the UserWidget. */
+		TArray<UE::MVVM::FMVVMConstFieldVariant> SkeletalGeneratedFields;
 		EMVVMBlueprintFieldPathSource GeneratedFrom = EMVVMBlueprintFieldPathSource::None;
 		bool bCanBeSetInNative = true;
 		bool bUseByNativeBinding = false;
@@ -361,12 +372,12 @@ private:
 
 	struct FGetFieldsResult
 	{
-		TSharedPtr<FCompilerBindingSource> Source;
+		TSharedPtr<FCompilerBindingSource> OptionalSource;
 		TArray<FMVVMConstFieldVariant> GeneratedFields;
 	};
 	struct FCreateFieldsResult
 	{
-		TSharedPtr<FCompilerBindingSource> Source;
+		TSharedPtr<FCompilerBindingSource> OptionalSource;
 		TArray<UE::MVVM::FMVVMConstFieldVariant> GeneratedFields;
 		TArray<UE::MVVM::FMVVMConstFieldVariant> SkeletalGeneratedFields;
 	};
@@ -375,7 +386,7 @@ private:
 	TValueOrError<TSharedPtr<FCompilerNotifyFieldId>, FText> CreateNotifyFieldId(const UWidgetBlueprintGeneratedClass* Class, const TSharedPtr<FGeneratedReadFieldPathContext>& ReadFieldContext, const FMVVMBlueprintViewBinding& Binding);
 
 
-	static TArray<FMVVMConstFieldVariant> GetFields(const UClass* Class, FName PropertyName, TArray<FMVVMConstFieldVariant> Properties);
+	static TArray<FMVVMConstFieldVariant> AppendBaseField(const UClass* Class, FName PropertyName, TArray<FMVVMConstFieldVariant> Properties);
 	static bool IsPropertyPathValid(const UBlueprint* Context, TArrayView<const FMVVMConstFieldVariant> PropertyPath);
 	static bool CanBeSetInNative(TArrayView<const FMVVMConstFieldVariant> PropertyPath);
 };
