Localization Warning fixes.

#rnx

diff --git a/Engine/Plugins/Animation/IKRig/Source/IKRigEditor/Private/RetargetEditor/SRetargetAnimAssetsWindow.cpp b/Engine/Plugins/Animation/IKRig/Source/IKRigEditor/Private/RetargetEditor/SRetargetAnimAssetsWindow.cpp
--- a/Engine/Plugins/Animation/IKRig/Source/IKRigEditor/Private/RetargetEditor/SRetargetAnimAssetsWindow.cpp
+++ b/Engine/Plugins/Animation/IKRig/Source/IKRigEditor/Private/RetargetEditor/SRetargetAnimAssetsWindow.cpp
@@ -1250,13 +1250,13 @@ void SRetargetAnimAssetsWindow::ShowAssetWarnings()
 	// missing source mesh
 	if (!Settings->SourceSkeletalMesh)
 	{
-		Log.LogError(LOCTEXT( "MissingMeshError", "No source mesh assigned."));
+		Log.LogError(LOCTEXT( "MissingSourceMeshError", "No source mesh assigned."));
 	}
 
 	// missing target mesh
 	if (!Settings->TargetSkeletalMesh)
 	{
-		Log.LogError(LOCTEXT( "MissingMeshError", "No target mesh assigned."));
+		Log.LogError(LOCTEXT( "MissinTargetMeshError", "No target mesh assigned."));
 	}
 
 	// auto retarget results
diff --git a/Engine/Plugins/Runtime/ModelViewViewModel/Source/ModelViewViewModel/Private/View/MVVMView.cpp b/Engine/Plugins/Runtime/ModelViewViewModel/Source/ModelViewViewModel/Private/View/MVVMView.cpp
--- a/Engine/Plugins/Runtime/ModelViewViewModel/Source/ModelViewViewModel/Private/View/MVVMView.cpp
+++ b/Engine/Plugins/Runtime/ModelViewViewModel/Source/ModelViewViewModel/Private/View/MVVMView.cpp
@@ -525,13 +525,8 @@ void UMVVMView::HandledLibraryBindingValueChanged(UObject* InSource, UE::FieldNo
 		if (!ViewSource.bBindingsInitialized || !ViewSource.bSourceInitialized)
 		{
 			// we do not want to run a binding while we are initializing the bindings.
-#if UE_WITH_MVVM_DEBUGGING
 			UE::MVVM::FMessageLog Log(GetUserWidget());
 			Log.Warning(FText::Format(LOCTEXT("ExecuteBindingFailWhileInitializing", "The Field '{0}' could not execute the bindings while initializing."), FText::FromName(InFieldId.GetName())));
-#else
-			UE::MVVM::FMessageLog Log(GetUserWidget());
-			Log.Warning(FText::Format(LOCTEXT("ExecuteBindingFailInvalidSource", "The Field '{0}'  could not execute the bindings while initializing."), FText::FromName(InFieldId.GetName())));
-#endif
 			return;
 		}
 
