MVVM: Enable compilation warning when there are 2 bindings trying to change the same value. Since the order is not configurable that can lead to binding executing in different orders and do run in an unexpected behavior.

diff --git a/Engine/Plugins/Runtime/ModelViewViewModel/Source/ModelViewViewModelBlueprint/Private/MVVMViewBlueprintCompiler.cpp b/Engine/Plugins/Runtime/ModelViewViewModel/Source/ModelViewViewModelBlueprint/Private/MVVMViewBlueprintCompiler.cpp
--- a/Engine/Plugins/Runtime/ModelViewViewModel/Source/ModelViewViewModelBlueprint/Private/MVVMViewBlueprintCompiler.cpp
+++ b/Engine/Plugins/Runtime/ModelViewViewModel/Source/ModelViewViewModelBlueprint/Private/MVVMViewBlueprintCompiler.cpp
@@ -1216,12 +1216,11 @@ void FMVVMViewBlueprintCompiler::CreateWriteFieldContexts(const FWidgetBlueprint
 					});
 				if (Found)
 				{
-					// todo Temporary removing this message until the assets are fixed.
-					//AddMessageForBinding(Binding
-					//	, FText::Format(LOCTEXT("PropertyPathAlreadyUsed", "The property path '{0}' is already used by another binding."), PropertyPathToText(NewSkeletonClass, BlueprintView.Get(), DestinationPropertyPath))
-					//	, EMessageType::Warning
-					//	, FMVVMBlueprintPinId()
-					//);
+					AddMessageForBinding(Binding
+						, FText::Format(LOCTEXT("PropertyPathAlreadyUsed", "The property path '{0}' is already used by another binding."), PropertyPathToText(NewSkeletonClass, BlueprintView.Get(), DestinationPropertyPath))
+						, EMessageType::Warning
+						, FMVVMBlueprintPinId()
+					);
 					ValidBinding->WritePath = *Found;
 					continue;
 				}
