MVVM: Always assign the viewmodel value to the property value. It's a bug to not assign it when it's null.

diff --git a/Engine/Plugins/Runtime/ModelViewViewModel/Source/ModelViewViewModel/Private/View/MVVMView.cpp b/Engine/Plugins/Runtime/ModelViewViewModel/Source/ModelViewViewModel/Private/View/MVVMView.cpp
--- a/Engine/Plugins/Runtime/ModelViewViewModel/Source/ModelViewViewModel/Private/View/MVVMView.cpp
+++ b/Engine/Plugins/Runtime/ModelViewViewModel/Source/ModelViewViewModel/Private/View/MVVMView.cpp
@@ -173,25 +173,26 @@ void UMVVMView::InitializeSourceInternal(UObject* NewSource, FMVVMViewClass_Sour
 	if (NewSource)
 	{
 		ValidSources |= SourceKey.GetBit();
-		if (ClassSource.RequireSettingUserWidgetProperty())
-		{
-			UUserWidget* UserWidget = GetUserWidget();
-			FObjectPropertyBase* FoundObjectProperty = FindFProperty<FObjectPropertyBase>(UserWidget->GetClass(), ClassSource.GetUserWidgetPropertyName());
-			if (ensureAlwaysMsgf(FoundObjectProperty, TEXT("The compiler should have added the property")))
-			{
-				if (ensure(NewSource->GetClass()->IsChildOf(FoundObjectProperty->PropertyClass)))
-				{
-					FoundObjectProperty->SetObjectPropertyValue_InContainer(UserWidget, NewSource);
-					ViewSource.bAssignedToUserWidgetProperty = true;
-				}
-			}
-		}
 	}
 	else
 	{
 		ValidSources &= ~SourceKey.GetBit();
 	}
 
+	if (ClassSource.RequireSettingUserWidgetProperty())
+	{
+		UUserWidget* UserWidget = GetUserWidget();
+		FObjectPropertyBase* FoundObjectProperty = FindFProperty<FObjectPropertyBase>(UserWidget->GetClass(), ClassSource.GetUserWidgetPropertyName());
+		if (ensureAlwaysMsgf(FoundObjectProperty, TEXT("The compiler should have added the property")))
+		{
+			if (NewSource == nullptr || ensure(NewSource->GetClass()->IsChildOf(FoundObjectProperty->PropertyClass)))
+			{
+				FoundObjectProperty->SetObjectPropertyValue_InContainer(UserWidget, NewSource);
+				ViewSource.bAssignedToUserWidgetProperty = true;
+			}
+		}
+	}
+
 	ViewSource.Source = NewSource;
 }
 
