MVVM: Add the UserWidget to the message log. The UserWidget is a link that the user can click on.

diff --git a/Engine/Plugins/Runtime/ModelViewViewModel/Source/ModelViewViewModel/Private/MVVMMessageLog.cpp b/Engine/Plugins/Runtime/ModelViewViewModel/Source/ModelViewViewModel/Private/MVVMMessageLog.cpp
--- a/Engine/Plugins/Runtime/ModelViewViewModel/Source/ModelViewViewModel/Private/MVVMMessageLog.cpp
+++ b/Engine/Plugins/Runtime/ModelViewViewModel/Source/ModelViewViewModel/Private/MVVMMessageLog.cpp
@@ -2,5 +2,121 @@
 
 #include "MVVMMessageLog.h"
 
+#include "Blueprint/UserWidget.h"
+#include "View/MVVMViewClass.h"
+
+#include "Logging/TokenizedMessage.h"
+#include "Misc/UObjectToken.h"
+
+#if WITH_EDITOR
+#include "Kismet2/KismetEditorUtilities.h" //UnrealEd
+#endif
+
+#define LOCTEXT_NAMESPACE "FMVVMMessageLog"
+
 DEFINE_LOG_CATEGORY(LogMVVM);
-const FName UE::MVVM::FMessageLog::LogName = "Model View Viewmodel";
\ No newline at end of file
+const FName UE::MVVM::FMessageLog::LogName = "Model View Viewmodel";
+
+
+namespace UE::MVVM
+{
+
+namespace Private
+{
+void OnMessageLogLinkActivated_UObjectToken(const TSharedRef<IMessageToken>& Token)
+{
+#if WITH_EDITOR
+	if (Token->GetType() == EMessageToken::Object)
+	{
+		const TSharedRef<FUObjectToken> UObjectToken = StaticCastSharedRef<FUObjectToken>(Token);
+		if (const UClass* UserWidget = Cast<UClass>(UObjectToken->GetObject().Get()))
+		{
+			FKismetEditorUtilities::BringKismetToFocusAttentionOnObject(UserWidget);
+		}
+	}
+#endif
+}
+void OnMessageLogLinkActivated_Binding(TWeakObjectPtr<const UUserWidget> WeakUserWidget, FGuid BindingId)
+{
+#if WITH_EDITOR
+		if (const UUserWidget* UserWidget = WeakUserWidget.Get())
+		{
+			FKismetEditorUtilities::BringKismetToFocusAttentionOnObject(UserWidget->GetClass());
+		}
+#endif
+}
+
+void AddUserWidget(TSharedRef<FTokenizedMessage>& Message, const UUserWidget* UserWidget)
+{
+#if WITH_EDITOR
+	if (UserWidget)
+	{
+		Message->AddToken(FUObjectToken::Create(UserWidget->GetClass(), FText::FromName(UserWidget->GetFName()))
+			->OnMessageTokenActivated(FOnMessageTokenActivated::CreateStatic(&OnMessageLogLinkActivated_UObjectToken))
+		);
+	}
+#else
+	if (UserWidget)
+	{
+		Message->AddToken(FTextToken::Create(FText::FromName(UserWidget->GetFName())));
+	}
+#endif
+}
+} //namespace Private
+
+FMessageLog::FMessageLog(const UUserWidget* InUserWidget)
+	: ::FMessageLog(LogName)
+	, UserWidget(InUserWidget)
+{}
+
+TSharedRef<FTokenizedMessage> FMessageLog::Message(EMessageSeverity::Type InSeverity, const FText& InMessage)
+{
+	TSharedRef<FTokenizedMessage> NewMessage = FTokenizedMessage::Create(InSeverity);
+	NewMessage->AddToken(FTextToken::Create(InMessage));
+	Private::AddUserWidget(NewMessage, UserWidget.Get());
+	AddMessage(NewMessage);
+	return NewMessage;
+}
+
+TSharedRef<FTokenizedMessage> FMessageLog::Error(const FText& InMessage)
+{
+	return Message(EMessageSeverity::Error, InMessage);
+}
+
+TSharedRef<FTokenizedMessage> FMessageLog::PerformanceWarning(const FText& InMessage)
+{
+	return Message(EMessageSeverity::PerformanceWarning, InMessage);
+}
+
+TSharedRef<FTokenizedMessage> FMessageLog::Warning(const FText& InMessage)
+{
+	return Message(EMessageSeverity::Warning, InMessage);
+}
+
+TSharedRef<FTokenizedMessage> FMessageLog::Info(const FText& InMessage)
+{
+	return Message(EMessageSeverity::Info, InMessage);
+}
+
+void FMessageLog::AddBindingToken(TSharedRef<FTokenizedMessage> NewMessage, const UMVVMViewClass* Class, const FMVVMViewClass_Binding& ClassBinding, FMVVMViewClass_BindingKey Key)
+{
+#if WITH_EDITOR
+	NewMessage->AddToken(
+		FActionToken::Create(
+			LOCTEXT("BindingToken", "binding")
+			, FText::FromString(ClassBinding.ToString(Class, FMVVMViewClass_Binding::FToStringArgs::Short()))
+			, FOnActionTokenExecuted::CreateStatic(&Private::OnMessageLogLinkActivated_Binding, UserWidget, ClassBinding.GetEditorId())
+		)
+	);
+#else
+	NewMessage->AddToken(
+		FTextToken::Create(
+			FText::Format(LOCTEXT("BindingTokenToken", "binding: {0}"), FText::AsNumber(Key.GetIndex()))
+		)
+	);
+#endif
+}
+
+} // namespace UE::MVVM
+
+#undef LOCTEXT_NAMESPACE
\ No newline at end of file
diff --git a/Engine/Plugins/Runtime/ModelViewViewModel/Source/ModelViewViewModel/Private/View/MVVMViewClass.cpp b/Engine/Plugins/Runtime/ModelViewViewModel/Source/ModelViewViewModel/Private/View/MVVMViewClass.cpp
--- a/Engine/Plugins/Runtime/ModelViewViewModel/Source/ModelViewViewModel/Private/View/MVVMViewClass.cpp
+++ b/Engine/Plugins/Runtime/ModelViewViewModel/Source/ModelViewViewModel/Private/View/MVVMViewClass.cpp
@@ -372,6 +372,9 @@ FMVVMViewClass_Binding::FToStringArgs FMVVMViewClass_Binding::FToStringArgs::Sho
 	FToStringArgs Result;
 	Result.bUseDisplayName = false;
 	Result.bAddFlags = false;
+	Result.bAddBindingId = false;
+	Result.bAddBindingFields = true;
+	Result.bAddSources = false;
 	return Result;
 }
 
@@ -384,12 +387,22 @@ FString FMVVMViewClass_Binding::ToString(const UMVVMViewClass* ViewClass, FToStr
 {
 	TStringBuilder<1024> StringBuilder;
 
+	bool bAddNewLine = false;
+	auto AddNewLine = [&StringBuilder, &bAddNewLine]()
+		{
+			if (bAddNewLine)
+			{
+				StringBuilder << TEXT("\n    ");
+			}
+			bAddNewLine = true;
+		};
+
 #if WITH_EDITOR
 	if (Args.bAddBindingId)
 	{
+		AddNewLine();
 		StringBuilder << TEXT("BindingId: ");
 		EditorId.AppendString(StringBuilder);
-		StringBuilder << TEXT("\n    ");
 	}
 #endif
 
@@ -397,6 +410,7 @@ FString FMVVMViewClass_Binding::ToString(const UMVVMViewClass* ViewClass, FToStr
 	{
 		if (GetBinding().IsValid())
 		{
+			AddNewLine();
 			StringBuilder << TEXT("Binding: ");
 
 			FString SourceString;
@@ -475,21 +489,23 @@ FString FMVVMViewClass_Binding::ToString(const UMVVMViewClass* ViewClass, FToStr
 			{
 				StringBuilder << TEXT("()");
 			}
-
-			StringBuilder << TEXT("\n    ");
 		}
 		else
 		{
 			StringBuilder << TEXT("!!Invalid Binding Error!!");
-			StringBuilder << TEXT("\n    ");
 		}
 	}
 
-	StringBuilder << TEXT("Sources: 0x");
-	StringBuilder << FString::Printf(TEXT("%x"), SourceBitField);
+	if (Args.bAddSources)
+	{
+		AddNewLine();
+		StringBuilder << TEXT("Sources: 0x");
+		StringBuilder << FString::Printf(TEXT("%x"), SourceBitField);
+	}
 
 	if (Args.bAddFlags)
 	{
+		AddNewLine();
 		if ((Flags & (uint8)EFlags::OverrideExecuteMode) != 0)
 		{
 			StringBuilder << TEXT(", Mode: ");
diff --git a/Engine/Plugins/Runtime/ModelViewViewModel/Source/ModelViewViewModel/Public/MVVMMessageLog.h b/Engine/Plugins/Runtime/ModelViewViewModel/Source/ModelViewViewModel/Public/MVVMMessageLog.h
--- a/Engine/Plugins/Runtime/ModelViewViewModel/Source/ModelViewViewModel/Public/MVVMMessageLog.h
+++ b/Engine/Plugins/Runtime/ModelViewViewModel/Source/ModelViewViewModel/Public/MVVMMessageLog.h
@@ -6,31 +6,34 @@
 
 #include "UObject/WeakObjectPtr.h"
 
+struct FMVVMViewClass_Binding;
+struct FMVVMViewClass_BindingKey;
+class UMVVMViewClass;
 class UUserWidget;
 
 MODELVIEWVIEWMODEL_API DECLARE_LOG_CATEGORY_EXTERN(LogMVVM, Log, All);
 
 namespace UE::MVVM
 {
-	class MODELVIEWVIEWMODEL_API FMessageLog : private ::FMessageLog
-	{
-	public:
-		static const FName LogName;
-
-		FMessageLog(const UUserWidget* InUserWidget)
-			: ::FMessageLog(LogName)
-		{}
-
-		using ::FMessageLog::AddMessage;
-		using ::FMessageLog::AddMessages;
-		using ::FMessageLog::Message;
-		using ::FMessageLog::CriticalError;
-		using ::FMessageLog::Error;
-		using ::FMessageLog::PerformanceWarning;
-		using ::FMessageLog::Warning;
-		using ::FMessageLog::Info;
-
-	private:
-		TWeakObjectPtr<UUserWidget> UserWidget;
-	};
+class MODELVIEWVIEWMODEL_API FMessageLog : private ::FMessageLog
+{
+public:
+	static const FName LogName;
+
+	FMessageLog(const UUserWidget* InUserWidget);
+
+	using ::FMessageLog::AddMessage;
+	using ::FMessageLog::AddMessages;
+	TSharedRef<FTokenizedMessage> Message(EMessageSeverity::Type InSeverity, const FText& InMessage);
+	using ::FMessageLog::CriticalError;
+	TSharedRef<FTokenizedMessage> Error(const FText& InMessage);
+	TSharedRef<FTokenizedMessage> PerformanceWarning(const FText& InMessage);
+	TSharedRef<FTokenizedMessage> Warning(const FText& InMessage);
+	TSharedRef<FTokenizedMessage> Info(const FText& InMessage);
+
+	void AddBindingToken(TSharedRef<FTokenizedMessage> NewMessage, const UMVVMViewClass* Class, const FMVVMViewClass_Binding& ClassBinding, FMVVMViewClass_BindingKey Key);
+
+private:
+	TWeakObjectPtr<const UUserWidget> UserWidget;
+};
 }
diff --git a/Engine/Plugins/Runtime/ModelViewViewModel/Source/ModelViewViewModel/Public/View/MVVMViewClass.h b/Engine/Plugins/Runtime/ModelViewViewModel/Source/ModelViewViewModel/Public/View/MVVMViewClass.h
--- a/Engine/Plugins/Runtime/ModelViewViewModel/Source/ModelViewViewModel/Public/View/MVVMViewClass.h
+++ b/Engine/Plugins/Runtime/ModelViewViewModel/Source/ModelViewViewModel/Public/View/MVVMViewClass.h
@@ -127,6 +127,7 @@ public:
 		bool bAddFlags = true;
 		bool bAddBindingId = true;
 		bool bAddBindingFields = true;
+		bool bAddSources = true;
 
 		MODELVIEWVIEWMODEL_API static FToStringArgs Short();
 		MODELVIEWVIEWMODEL_API static FToStringArgs All();
@@ -398,6 +399,7 @@ private:
 		HasTickBindings = 1 << 8,
 		IsViewModel = 1 << 9,
 		IsViewModelInstanceExposed = 1 << 10,
+		GlobalViewModelCollectionRetry = 1 << 11,
 	};
 
 	UPROPERTY(VisibleAnywhere, Category = "View")
