Fix build issue 432248 introduced with CL 31643104
#rnx

diff --git a/Engine/Plugins/Runtime/ModelViewViewModel/Source/ModelViewViewModel/Private/MVVMMessageLog.cpp b/Engine/Plugins/Runtime/ModelViewViewModel/Source/ModelViewViewModel/Private/MVVMMessageLog.cpp
--- a/Engine/Plugins/Runtime/ModelViewViewModel/Source/ModelViewViewModel/Private/MVVMMessageLog.cpp
+++ b/Engine/Plugins/Runtime/ModelViewViewModel/Source/ModelViewViewModel/Private/MVVMMessageLog.cpp
@@ -100,7 +100,7 @@ TSharedRef<FTokenizedMessage> FMessageLog::Info(const FText& InMessage)
 
 void FMessageLog::AddBindingToken(TSharedRef<FTokenizedMessage> NewMessage, const UMVVMViewClass* Class, const FMVVMViewClass_Binding& ClassBinding, FMVVMViewClass_BindingKey Key)
 {
-#if UE_WITH_MVVM_DEBUGGING
+#if UE_WITH_MVVM_DEBUGGING && WITH_EDITOR
 	NewMessage->AddToken(
 		FActionToken::Create(
 			LOCTEXT("BindingToken", "binding")
