MVVM: Fix false positive warning. It is valid for a backward binding to set the same property but it is not valid for a forward binding. Forward binding executes on initialization. The order of initialization is not deterministic and would lead to miss use.

diff --git a/Engine/Plugins/Runtime/ModelViewViewModel/Source/ModelViewViewModelBlueprint/Private/MVVMViewBlueprintCompiler.cpp b/Engine/Plugins/Runtime/ModelViewViewModel/Source/ModelViewViewModelBlueprint/Private/MVVMViewBlueprintCompiler.cpp
--- a/Engine/Plugins/Runtime/ModelViewViewModel/Source/ModelViewViewModelBlueprint/Private/MVVMViewBlueprintCompiler.cpp
+++ b/Engine/Plugins/Runtime/ModelViewViewModel/Source/ModelViewViewModelBlueprint/Private/MVVMViewBlueprintCompiler.cpp
@@ -1208,33 +1208,39 @@ void FMVVMViewBlueprintCompiler::CreateWriteFieldContexts(const FWidgetBlueprint
 			}
 
 			// Test if it already exist
-			const TArray<UE::MVVM::FMVVMConstFieldVariant>& SkeletalGeneratedFieldsResult = FieldContextResult.GetValue().SkeletalGeneratedFields;
+			TSharedPtr<FGeneratedWriteFieldPathContext> WriteFieldPath;
 			{
+				const TArray<UE::MVVM::FMVVMConstFieldVariant>& SkeletalGeneratedFieldsResult = FieldContextResult.GetValue().SkeletalGeneratedFields;
 				TSharedRef<FGeneratedWriteFieldPathContext>* Found = GeneratedWriteFieldPaths.FindByPredicate([&SkeletalGeneratedFieldsResult](const TSharedRef<FGeneratedWriteFieldPathContext>& Other)
 					{
 						return Other->SkeletalGeneratedFields == SkeletalGeneratedFieldsResult;
 					});
 				if (Found)
 				{
-					// todo Temporary removing this message until the assets are fixed.
-					//AddMessageForBinding(Binding
-					//	, FText::Format(LOCTEXT("PropertyPathAlreadyUsed", "The property path '{0}' is already used by another binding."), PropertyPathToText(NewSkeletonClass, BlueprintView.Get(), DestinationPropertyPath))
-					//	, EMessageType::Warning
-					//	, FMVVMBlueprintPinId()
-					//);
-					ValidBinding->WritePath = *Found;
-					continue;
+					WriteFieldPath = *Found;
+
+					// Backward binding are reactive and are not trigger on initialization.
+					//It is valid for a backward binding to reuse the same write field path.
+					if (ValidBinding->Key.bIsForwardBinding)
+					{
+						AddMessageForBinding(Binding
+							, FText::Format(LOCTEXT("PropertyPathAlreadyUsed", "The property path '{0}' is already used by another binding."), PropertyPathToText(NewSkeletonClass, BlueprintView.Get(), DestinationPropertyPath))
+							, EMessageType::Warning
+							, FMVVMBlueprintPinId()
+						);
+					}
+				}
+				else
+				{
+					WriteFieldPath = MakeWriteFieldPath(
+						DestinationPropertyPath.GetSource(WidgetBlueprintCompilerContext.WidgetBlueprint())
+						, MoveTemp(FieldContextResult.GetValue().GeneratedFields)
+						, MoveTemp(FieldContextResult.GetValue().SkeletalGeneratedFields));
+					GeneratedWriteFieldPaths.Add(WriteFieldPath.ToSharedRef());
 				}
 			}
 
-			TSharedRef<FGeneratedWriteFieldPathContext> WriteFieldPath = MakeShared<FGeneratedWriteFieldPathContext>();
-			GeneratedWriteFieldPaths.Add(WriteFieldPath);
 			WriteFieldPath->UsedByBindings.AddUnique(ValidBinding);
-			//WriteFieldPath->OptionalSource;
-			WriteFieldPath->GeneratedFields = MoveTemp(FieldContextResult.GetValue().GeneratedFields);
-			WriteFieldPath->SkeletalGeneratedFields = MoveTemp(FieldContextResult.GetValue().SkeletalGeneratedFields);
-			WriteFieldPath->GeneratedFrom = DestinationPropertyPath.GetSource(WidgetBlueprintCompilerContext.WidgetBlueprint());
-			WriteFieldPath->bCanBeSetInNative = CanBeSetInNative(WriteFieldPath->SkeletalGeneratedFields);
 			WriteFieldPath->bUseByNativeBinding = true; // the setter function can be in BP but the destination will be set in native
 
 			// Assign the Destination to the binding
@@ -1263,30 +1269,26 @@ void FMVVMViewBlueprintCompiler::CreateWriteFieldContexts(const FWidgetBlueprint
 			}
 
 			// Test if it already exist
-			const TArray<UE::MVVM::FMVVMConstFieldVariant>& SkeletalGeneratedFieldsResult = FieldContextResult.GetValue().SkeletalGeneratedFields;
+			TSharedPtr<FGeneratedWriteFieldPathContext> WriteFieldPath;
 			{
+				const TArray<UE::MVVM::FMVVMConstFieldVariant>& SkeletalGeneratedFieldsResult = FieldContextResult.GetValue().SkeletalGeneratedFields;
 				TSharedRef<FGeneratedWriteFieldPathContext>* Found = GeneratedWriteFieldPaths.FindByPredicate([&SkeletalGeneratedFieldsResult](const TSharedRef<FGeneratedWriteFieldPathContext>& Other) { return Other->SkeletalGeneratedFields == SkeletalGeneratedFieldsResult; });
 				if (Found)
 				{
-					(*Found)->UsedByEvents.AddUnique(ValidEvent);
-					ValidEvent->WritePath = *Found;
+					WriteFieldPath = (*Found);
 				}
 				else
 				{
-					TSharedRef<FGeneratedWriteFieldPathContext> WriteFieldPath = MakeShared<FGeneratedWriteFieldPathContext>();
-					GeneratedWriteFieldPaths.Add(WriteFieldPath);
-					WriteFieldPath->UsedByEvents.AddUnique(ValidEvent);
-					//WriteFieldPath->OptionalSource;
-					WriteFieldPath->GeneratedFields = MoveTemp(FieldContextResult.GetValue().GeneratedFields);
-					WriteFieldPath->SkeletalGeneratedFields = MoveTemp(FieldContextResult.GetValue().SkeletalGeneratedFields);
-					WriteFieldPath->GeneratedFrom = EventPtr->GetDestinationPath().GetSource(WidgetBlueprintCompilerContext.WidgetBlueprint());
-					WriteFieldPath->bCanBeSetInNative = CanBeSetInNative(WriteFieldPath->SkeletalGeneratedFields);
-					WriteFieldPath->bUseByNativeBinding = false;
-
-					// Assign the Destination to the events
-					ValidEvent->WritePath = WriteFieldPath;
+					WriteFieldPath = MakeWriteFieldPath(
+						EventPtr->GetDestinationPath().GetSource(WidgetBlueprintCompilerContext.WidgetBlueprint())
+						, MoveTemp(FieldContextResult.GetValue().GeneratedFields)
+						,MoveTemp(FieldContextResult.GetValue().SkeletalGeneratedFields));
+					GeneratedWriteFieldPaths.Add(WriteFieldPath.ToSharedRef());
 				}
 			}
+
+			WriteFieldPath->UsedByEvents.AddUnique(ValidEvent);
+			ValidEvent->WritePath = WriteFieldPath;
 		}
 	}
 }
@@ -3618,6 +3620,17 @@ bool FMVVMViewBlueprintCompiler::CanBeSetInNative(TArrayView<const FMVVMConstFie
 	return true;
 }
 
+TSharedRef<FMVVMViewBlueprintCompiler::FGeneratedWriteFieldPathContext> FMVVMViewBlueprintCompiler::MakeWriteFieldPath(EMVVMBlueprintFieldPathSource GeneratedFrom, TArray<UE::MVVM::FMVVMConstFieldVariant>&& GeneratedFields, TArray<UE::MVVM::FMVVMConstFieldVariant>&& SkeletalGeneratedFields)
+{
+	TSharedRef<FGeneratedWriteFieldPathContext> WriteFieldPath = MakeShared<FGeneratedWriteFieldPathContext>();
+	//WriteFieldPath->OptionalSource;
+	WriteFieldPath->GeneratedFields = MoveTemp(GeneratedFields);
+	WriteFieldPath->SkeletalGeneratedFields = MoveTemp(SkeletalGeneratedFields);
+	WriteFieldPath->GeneratedFrom = GeneratedFrom;
+	WriteFieldPath->bCanBeSetInNative = CanBeSetInNative(WriteFieldPath->SkeletalGeneratedFields);
+	return WriteFieldPath;
+}
+
 void FMVVMViewBlueprintCompiler::TestGenerateSetter(const UBlueprint* Context, FStringView ObjectName, FStringView FieldPath, FStringView FunctionName)
 {
 #if UE_BUILD_DEVELOPMENT || UE_BUILD_DEBUG
diff --git a/Engine/Plugins/Runtime/ModelViewViewModel/Source/ModelViewViewModelBlueprint/Private/MVVMViewBlueprintCompiler.h b/Engine/Plugins/Runtime/ModelViewViewModel/Source/ModelViewViewModelBlueprint/Private/MVVMViewBlueprintCompiler.h
--- a/Engine/Plugins/Runtime/ModelViewViewModel/Source/ModelViewViewModelBlueprint/Private/MVVMViewBlueprintCompiler.h
+++ b/Engine/Plugins/Runtime/ModelViewViewModel/Source/ModelViewViewModelBlueprint/Private/MVVMViewBlueprintCompiler.h
@@ -389,6 +389,7 @@ private:
 	static TArray<FMVVMConstFieldVariant> AppendBaseField(const UClass* Class, FName PropertyName, TArray<FMVVMConstFieldVariant> Properties);
 	static bool IsPropertyPathValid(const UBlueprint* Context, TArrayView<const FMVVMConstFieldVariant> PropertyPath);
 	static bool CanBeSetInNative(TArrayView<const FMVVMConstFieldVariant> PropertyPath);
+	static TSharedRef<FGeneratedWriteFieldPathContext> MakeWriteFieldPath(EMVVMBlueprintFieldPathSource GeneratedFrom, TArray<UE::MVVM::FMVVMConstFieldVariant>&& GeneratedFields, TArray<UE::MVVM::FMVVMConstFieldVariant>&& SkeletalGeneratedFields);
 };
 
 } //namespace
