Add option to filter the resolver to only returns the viewmodel desired class.

diff --git a/Engine/Plugins/Runtime/ModelViewViewModel/Source/ModelViewViewModel/Private/View/MVVMViewModelContextResolver.cpp b/Engine/Plugins/Runtime/ModelViewViewModel/Source/ModelViewViewModel/Private/View/MVVMViewModelContextResolver.cpp
new file mode 100644
--- /dev/null
+++ b/Engine/Plugins/Runtime/ModelViewViewModel/Source/ModelViewViewModel/Private/View/MVVMViewModelContextResolver.cpp
@@ -0,0 +1,49 @@
+// Copyright Epic Games, Inc. All Rights Reserved.
+
+#include "View/MVVMViewModelContextResolver.h"
+
+#if WITH_EDITOR
+#include "ClassViewerFilter.h" //ClassViewer
+#endif
+
+#include UE_INLINE_GENERATED_CPP_BY_NAME(MVVMViewModelContextResolver)
+
+#if WITH_EDITOR
+bool UMVVMViewModelContextResolver::DoesSupportViewModelClass(const UClass* Class) const
+{
+	if (Class == nullptr)
+	{
+		return false;
+	}
+	
+	auto ContainsByPredicate = [Class](const TArray<FSoftClassPath>& List)
+		{
+			return List.ContainsByPredicate(
+				[Class](const FSoftClassPath& Other)
+				{
+					const UClass* OtherClass = Other.ResolveClass();
+					return OtherClass ? Class->IsChildOf(OtherClass) : false;
+				});
+		};
+	
+	if (AllowedViewModelClasses.Num() > 0)
+	{
+		bool bContains = ContainsByPredicate(AllowedViewModelClasses);		
+		if (!bContains)
+		{
+			return false;
+		}
+	}
+	
+	if (DeniedViewModelClasses.Num() > 0)
+	{
+		bool bContains = ContainsByPredicate(DeniedViewModelClasses);		
+		if (bContains)
+		{
+			return false;
+		}
+	}
+
+	return true;
+}
+#endif
diff --git a/Engine/Plugins/Runtime/ModelViewViewModel/Source/ModelViewViewModel/Public/View/MVVMViewModelContextResolver.h b/Engine/Plugins/Runtime/ModelViewViewModel/Source/ModelViewViewModel/Public/View/MVVMViewModelContextResolver.h
--- a/Engine/Plugins/Runtime/ModelViewViewModel/Source/ModelViewViewModel/Public/View/MVVMViewModelContextResolver.h
+++ b/Engine/Plugins/Runtime/ModelViewViewModel/Source/ModelViewViewModel/Public/View/MVVMViewModelContextResolver.h
@@ -35,4 +35,20 @@ public:
 
 	UFUNCTION(BlueprintImplementableEvent, Category = "Viewmodel", meta = (DisplayName = "Destroy Instance"))
 	void K2_DestroyInstance(const UObject* ViewModel, const UMVVMView* View) const;
+
+public:
+#if WITH_EDITOR
+	virtual bool DoesSupportViewModelClass(const UClass* Class) const;
+#endif
+
+private:
+#if WITH_EDITORONLY_DATA
+	/** Viewmodel class that the resolver supports.*/
+	UPROPERTY(EditDefaultsOnly, Category = "Viewmodel", meta=(MustImplement="/Script/FieldNotification.NotifyFieldValueChanged", DisallowedClasses="/Script/UMG.Widget"))
+	TArray<FSoftClassPath> AllowedViewModelClasses;
+
+	/** Viewmodel class that the resolver explicitly does not support. */
+	UPROPERTY(EditDefaultsOnly, Category = "Viewmodel", meta=(MustImplement="/Script/FieldNotification.NotifyFieldValueChanged", DisallowedClasses="/Script/UMG.Widget"))
+	TArray<FSoftClassPath> DeniedViewModelClasses;
+#endif
 };
diff --git a/Engine/Plugins/Runtime/ModelViewViewModel/Source/ModelViewViewModelEditor/Private/Customizations/MVVMBlueprintViewModelContextCustomization.cpp b/Engine/Plugins/Runtime/ModelViewViewModel/Source/ModelViewViewModelEditor/Private/Customizations/MVVMBlueprintViewModelContextCustomization.cpp
--- a/Engine/Plugins/Runtime/ModelViewViewModel/Source/ModelViewViewModelEditor/Private/Customizations/MVVMBlueprintViewModelContextCustomization.cpp
+++ b/Engine/Plugins/Runtime/ModelViewViewModel/Source/ModelViewViewModelEditor/Private/Customizations/MVVMBlueprintViewModelContextCustomization.cpp
@@ -15,6 +15,8 @@
 #include "MVVMDeveloperProjectSettings.h"
 #include "MVVMEditorSubsystem.h"
 #include "PropertyHandle.h"
+#include "PropertyRestriction.h"
+#include "View/MVVMViewModelContextResolver.h"
 
 #include "Framework/MultiBox/MultiBoxBuilder.h"
 #include "WidgetBlueprintEditor.h"
@@ -136,6 +138,37 @@ namespace Private
 		}
 		return nullptr;
 	}
+
+	class FResolverClassFilter : public IClassViewerFilter
+	{
+	public:
+		const UClass* ViewModelClass = nullptr;
+
+		virtual bool IsClassAllowed(const FClassViewerInitializationOptions& InInitOptions, const UClass* InClass, TSharedRef< FClassViewerFilterFuncs > InFilterFuncs) override
+		{
+			if (ViewModelClass && InClass->IsChildOf(UMVVMViewModelContextResolver::StaticClass()))
+			{
+				return InClass->GetDefaultObject<UMVVMViewModelContextResolver>()->DoesSupportViewModelClass(ViewModelClass);
+			}
+
+			return false;
+		}
+
+		virtual bool IsUnloadedClassAllowed(const FClassViewerInitializationOptions& InInitOptions, const TSharedRef<const IUnloadedBlueprintData> InBlueprint, TSharedRef< FClassViewerFilterFuncs > InFilterFuncs) override
+		{
+			if (ViewModelClass && InBlueprint->IsChildOf(UMVVMViewModelContextResolver::StaticClass()))
+			{
+				// Load the Blueprint
+				FSoftObjectPath BlueprintPath = FSoftObjectPath(InBlueprint->GetClassPathName());
+				if (UClass* LoadedClass = Cast<UClass>(BlueprintPath.TryLoad()))
+				{
+					return LoadedClass->GetDefaultObject<UMVVMViewModelContextResolver>()->DoesSupportViewModelClass(ViewModelClass);
+				}
+			}
+			return false;
+		}
+	};
+
 }
 
 FBlueprintViewModelContextDetailCustomization::FBlueprintViewModelContextDetailCustomization(TWeakPtr<FWidgetBlueprintEditor> InEditor)
@@ -155,6 +188,7 @@ void FBlueprintViewModelContextDetailCustomization::CustomizeChildren(TSharedRef
 		return;
 	}
 
+	UClass* ViewModelClass = nullptr;
 	FGuid ViewModelContextId = ContextPtr->GetViewModelId();
 	FName ViewModelPropertyName = ContextPtr->GetViewModelName();
 	bool bCanEdit = ContextPtr->bCanEdit;
@@ -170,7 +204,7 @@ void FBlueprintViewModelContextDetailCustomization::CustomizeChildren(TSharedRef
 		UObject* Object = nullptr;
 		if (NotifyFieldValueClassHandle->GetValue(Object) == FPropertyAccess::Success)
 		{
-			UClass* ViewModelClass = Cast<UClass>(Object);
+			ViewModelClass = Cast<UClass>(Object);
 			if (ViewModelClass)
 			{
 				AllowedCreationTypes = GetAllowedContextCreationType(ViewModelClass);
@@ -322,6 +356,13 @@ void FBlueprintViewModelContextDetailCustomization::CustomizeChildren(TSharedRef
 		TSharedPtr<IPropertyHandle> ResolverHandle = PropertyHandle->GetChildHandle(GET_MEMBER_NAME_CHECKED(FMVVMBlueprintViewModelContext, Resolver), false);
 		if (ensure(ResolverHandle))
 		{
+			TSharedRef<Private::FResolverClassFilter> ClassFilter = MakeShared<Private::FResolverClassFilter>();
+			ClassFilter->ViewModelClass = ViewModelClass;
+			TSharedRef<FPropertyRestriction> Restriction = MakeShared<FPropertyRestriction>(LOCTEXT("ResolverPropertyRestriction", "Resolver Property Restriction"));
+			Restriction->AddClassFilter(ClassFilter);
+
+			ResolverHandle->AddRestriction(Restriction);
+
 			ChildBuilder.AddProperty(ResolverHandle.ToSharedRef())
 				.IsEnabled(bCanEdit)
 				.Visibility(MakeAttributeLambda([ContextPtr]()
