MVVM: fix bindings not fired at initialization if they have complex conversion functions which use the same source for different arguments.

diff --git a/Engine/Plugins/Runtime/ModelViewViewModel/Source/ModelViewViewModelBlueprint/Private/MVVMViewBlueprintCompiler.cpp b/Engine/Plugins/Runtime/ModelViewViewModel/Source/ModelViewViewModelBlueprint/Private/MVVMViewBlueprintCompiler.cpp
--- a/Engine/Plugins/Runtime/ModelViewViewModel/Source/ModelViewViewModelBlueprint/Private/MVVMViewBlueprintCompiler.cpp
+++ b/Engine/Plugins/Runtime/ModelViewViewModel/Source/ModelViewViewModelBlueprint/Private/MVVMViewBlueprintCompiler.cpp
@@ -2849,12 +2849,15 @@ void FMVVMViewBlueprintCompiler::CompileBindings(const FCompiledBindingLibraryCo
 			//Remove the flag on all the binding on all the sources ecept the last one.
 			for (int32 Index = 0; Index < SharedExecuteAtInitializationBindings.Num() - 1; ++Index)
 			{
-				FMVVMViewClass_Source& ClassSource = ViewExtension->Sources[SharedExecuteAtInitializationBindings[Index].GetIndex()];
-				for (FMVVMViewClass_SourceBinding& SourceBinding : ClassSource.Bindings)
+				if (SharedExecuteAtInitializationBindings[Index] != SharedExecuteAtInitializationBindings.Last())
 				{
-					if (SourceBinding.GetBindingKey() == BindingKey)
+					FMVVMViewClass_Source& ClassSource = ViewExtension->Sources[SharedExecuteAtInitializationBindings[Index].GetIndex()];
+					for (FMVVMViewClass_SourceBinding& SourceBinding : ClassSource.Bindings)
 					{
-						SourceBinding.Flags &= ~(uint8)FMVVMViewClass_SourceBinding::EFlags::ExecuteAtInitialization;
+						if (SourceBinding.GetBindingKey() == BindingKey)
+						{
+							SourceBinding.Flags &= ~(uint8)FMVVMViewClass_SourceBinding::EFlags::ExecuteAtInitialization;
+						}
 					}
 				}
 			}
