MVVM: Add Conversion Function Library.
The functions are collected when on module load and when the BP class compiles. A function can be a Node. The BP registry do not collect all function types, like opcode functions. This library collects K2_Node and all function types supported by MVVM.

diff --git a/Engine/Plugins/Runtime/ModelViewViewModel/Source/ModelViewViewModelBlueprint/Private/MVVMDeveloperProjectSettings.cpp b/Engine/Plugins/Runtime/ModelViewViewModel/Source/ModelViewViewModelBlueprint/Private/MVVMDeveloperProjectSettings.cpp
--- a/Engine/Plugins/Runtime/ModelViewViewModel/Source/ModelViewViewModelBlueprint/Private/MVVMDeveloperProjectSettings.cpp
+++ b/Engine/Plugins/Runtime/ModelViewViewModel/Source/ModelViewViewModelBlueprint/Private/MVVMDeveloperProjectSettings.cpp
@@ -10,6 +10,9 @@
 #include "Types/MVVMExecutionMode.h"
 #include "UObject/UnrealType.h"
 
+#include "K2Node_FormatText.h"
+#include "Kismet/BlueprintFunctionLibrary.h"
+
 #define LOCTEXT_NAMESPACE "MVVMDeveloperProjectSettings"
 
 UMVVMDeveloperProjectSettings::UMVVMDeveloperProjectSettings()
@@ -25,6 +28,9 @@ UMVVMDeveloperProjectSettings::UMVVMDeveloperProjectSettings()
 	AllowedContextCreationType.Add(EMVVMBlueprintViewModelContextCreationType::PropertyPath);
 	AllowedContextCreationType.Add(EMVVMBlueprintViewModelContextCreationType::PropertyPath);
 	AllowedContextCreationType.Add(EMVVMBlueprintViewModelContextCreationType::Resolver);
+
+	AllowedClassForConversionFunctions.Add(UBlueprintFunctionLibrary::StaticClass());
+	AllowedClassForConversionFunctions.Add(UK2Node_FormatText::StaticClass());
 }
 
 FName UMVVMDeveloperProjectSettings::GetCategoryName() const
@@ -188,6 +194,45 @@ bool UMVVMDeveloperProjectSettings::IsFunctionAllowed(const UBlueprint* Generati
 	return true;
 }
 
+namespace UE::MVVM::Private
+{
+bool IsConversionFunctionAllowed(const TSet<FSoftClassPath>& AllowedClasses, const TSet<FSoftClassPath>& DeniedClasses, const TSet<FName>& DeniedModules, UClass* CurrentClass)
+{
+	bool bIsModuleDenied = DeniedModules.Contains(CurrentClass->GetClassPathName().GetPackageName());
+	if (bIsModuleDenied)
+	{
+		return false;
+	}
+	while (CurrentClass)
+	{
+		TStringBuilder<512> FunctionClassPath;
+		CurrentClass->GetPathName(nullptr, FunctionClassPath);
+		TStringBuilder<512> AllowedClassPath;
+		for (const FSoftClassPath& SoftClass : AllowedClasses)
+		{
+			SoftClass.ToString(AllowedClassPath);
+			if (AllowedClassPath.ToView() == FunctionClassPath.ToView())
+			{
+				return true;
+			}
+			AllowedClassPath.Reset();
+		}
+		for (const FSoftClassPath& SoftClass : DeniedClasses)
+		{
+			SoftClass.ToString(AllowedClassPath);
+			if (AllowedClassPath.ToView() == FunctionClassPath.ToView())
+			{
+				return false;
+			}
+			AllowedClassPath.Reset();
+		}
+
+		CurrentClass = CurrentClass->GetSuperClass();
+	}
+	return false;
+}
+} //namespace
+
 bool UMVVMDeveloperProjectSettings::IsConversionFunctionAllowed(const UBlueprint* GeneratingFor, const UFunction* Function) const
 {
 	if (ConversionFunctionFilter == EMVVMDeveloperConversionFunctionFilterType::BlueprintActionRegistry)
@@ -198,26 +243,15 @@ bool UMVVMDeveloperProjectSettings::IsConversionFunctionAllowed(const UBlueprint
 	{
 		check(ConversionFunctionFilter == EMVVMDeveloperConversionFunctionFilterType::AllowedList);
 
+		// Optimization. Static are for functions inside the AllowedClassForConversionFunctions.
 		if (Function->HasAllFunctionFlags(FUNC_Static))
 		{
-			TStringBuilder<512> FunctionClassPath;
-			Function->GetOwnerClass()->GetPathName(nullptr, FunctionClassPath);
-			TStringBuilder<512> AllowedClassPath;
-			for (const FSoftClassPath& SoftClass : AllowedClassForConversionFunctions)
-			{
-				SoftClass.ToString(AllowedClassPath);
-				if (AllowedClassPath.ToView() == FunctionClassPath.ToView())
-				{
-					return true;
-				}
-				AllowedClassPath.Reset();
-			}
-
-			return false;
+			UClass* CurrentClass = Function->GetOwnerClass();
+			return UE::MVVM::Private::IsConversionFunctionAllowed(AllowedClassForConversionFunctions, DeniedClassForConversionFunctions, DeniedModuleForConversionFunctions, CurrentClass);
 		}
 		else
 		{
-			// The function is on self and may have been filtered.
+			// The function is on self (WidgetBlueprint) and may be filtered.
 			return IsFunctionAllowed(GeneratingFor, Function->GetOwnerClass(), Function);
 		}
 	}
@@ -233,19 +267,7 @@ bool UMVVMDeveloperProjectSettings::IsConversionFunctionAllowed(const UBlueprint
 	{
 		check(ConversionFunctionFilter == EMVVMDeveloperConversionFunctionFilterType::AllowedList);
 
-		TStringBuilder<512> FunctionClassPath;
-		Function.Get()->GetPathName(nullptr, FunctionClassPath);
-		TStringBuilder<512> AllowedClassPath;
-		for (const FSoftClassPath& SoftClass : AllowedClassForConversionFunctions)
-		{
-			SoftClass.ToString(AllowedClassPath);
-			if (AllowedClassPath.ToView() == FunctionClassPath.ToView())
-			{
-				return true;
-			}
-			AllowedClassPath.Reset();
-		}
-		return false;
+		return UE::MVVM::Private::IsConversionFunctionAllowed(AllowedClassForConversionFunctions, DeniedClassForConversionFunctions, DeniedModuleForConversionFunctions, Function.Get());
 	}
 }
 
@@ -263,4 +285,18 @@ TArray<const UClass*> UMVVMDeveloperProjectSettings::GetAllowedConversionFunctio
 	return Result;
 }
 
+TArray<const UClass*> UMVVMDeveloperProjectSettings::GetDeniedConversionFunctionClasses() const
+{
+	TArray<const UClass*> Result;
+	for (const FSoftClassPath& SoftClass : DeniedClassForConversionFunctions)
+	{
+		if (UClass* Class = SoftClass.ResolveClass())
+		{
+			Result.Add(Class);
+		}
+	}
+
+	return Result;
+}
+
 #undef LOCTEXT_NAMESPACE
diff --git a/Engine/Plugins/Runtime/ModelViewViewModel/Source/ModelViewViewModelBlueprint/Public/MVVMDeveloperProjectSettings.h b/Engine/Plugins/Runtime/ModelViewViewModel/Source/ModelViewViewModelBlueprint/Public/MVVMDeveloperProjectSettings.h
--- a/Engine/Plugins/Runtime/ModelViewViewModel/Source/ModelViewViewModelBlueprint/Public/MVVMDeveloperProjectSettings.h
+++ b/Engine/Plugins/Runtime/ModelViewViewModel/Source/ModelViewViewModelBlueprint/Public/MVVMDeveloperProjectSettings.h
@@ -52,7 +52,7 @@ UENUM()
 enum class EMVVMDeveloperConversionFunctionFilterType : uint8
 {
 	BlueprintActionRegistry,
-	AllowedList,
+	AllowedList UMETA(DisplayName="Conversion Function Library"),
 };
 
 
@@ -93,6 +93,7 @@ public:
 	}
 
 	TArray<const UClass*> GetAllowedConversionFunctionClasses() const;
+	TArray<const UClass*> GetDeniedConversionFunctionClasses() const;
 
 private:
 	/** Permission list for filtering which properties are visible in UI. */
@@ -155,10 +156,18 @@ public:
 	UPROPERTY(EditAnywhere, config, Category = "Viewmodel")
 	EMVVMDeveloperConversionFunctionFilterType ConversionFunctionFilter = EMVVMDeveloperConversionFunctionFilterType::BlueprintActionRegistry;
 
-	/** Individual class that are allowed to be uses as conversion functions. */
+	/** Classes to include in conversion function list. It includes the child class. */
 	UPROPERTY(EditAnywhere, config, Category = "Viewmodel", meta = (EditCondition = "ConversionFunctionFilter == EMVVMDeveloperConversionFunctionFilterType::AllowedList"))
 	TSet<FSoftClassPath> AllowedClassForConversionFunctions;
 
+	/** Classes excluded for conversion function list. */
+	UPROPERTY(EditAnywhere, config, Category = "Viewmodel", meta = (EditCondition = "ConversionFunctionFilter == EMVVMDeveloperConversionFunctionFilterType::AllowedList"))
+	TSet<FSoftClassPath> DeniedClassForConversionFunctions;
+
+	/** Modules excluded for conversion function list. ie. "/Script/MyModule" */
+	UPROPERTY(EditAnywhere, config, Category = "Viewmodel", meta = (EditCondition = "ConversionFunctionFilter == EMVVMDeveloperConversionFunctionFilterType::AllowedList"))
+	TSet<FName> DeniedModuleForConversionFunctions;
+
 	/** Settings for filtering the list of available properties and functions on binding creation. */
 	UPROPERTY(EditAnywhere, config, Category = "Viewmodel")
 	FMVVMViewBindingFilterSettings FilterSettings;
diff --git a/Engine/Plugins/Runtime/ModelViewViewModel/Source/ModelViewViewModelEditor/Private/Framework/MVVMConversionFunctionLibrary.cpp b/Engine/Plugins/Runtime/ModelViewViewModel/Source/ModelViewViewModelEditor/Private/Framework/MVVMConversionFunctionLibrary.cpp
new file mode 100644
--- /dev/null
+++ b/Engine/Plugins/Runtime/ModelViewViewModel/Source/ModelViewViewModelEditor/Private/Framework/MVVMConversionFunctionLibrary.cpp
@@ -0,0 +1,727 @@
+// Copyright Epic Games, Inc. All Rights Reserved.
+
+#include "Framework/MVVMConversionFunctionLibrary.h"
+
+#include "Bindings/MVVMBindingHelper.h"
+#include "Bindings/MVVMConversionFunctionHelper.h"
+#include "Blueprint/UserWidget.h"
+#include "MVVMDeveloperProjectSettings.h"
+#include "MVVMEditorSubsystem.h"
+#include "WidgetBlueprint.h"
+
+#include "AssetRegistry/AssetRegistryModule.h"
+#include "BlueprintActionDatabase.h"
+#include "Editor.h"
+#include "EdGraphSchema_K2.h"
+#include "Editor/EditorEngine.h"
+#include "Kismet/BlueprintFunctionLibrary.h"
+#include "Kismet2/BlueprintEditorUtils.h"
+#include "Kismet2/KismetEditorUtilities.h"
+#include "Misc/ScopedSlowTask.h"
+#include "UObject/Class.h"
+
+
+#define LOCTEXT_NAMESPACE "MVVMConversionFunctionLibrary"
+
+extern UNREALED_API class UEditorEngine* GEditor;
+
+namespace UE::MVVM::ConversionFunctionLibrary
+{
+
+namespace Private
+{
+/**
+ *
+ */
+struct FFunctionEntry_UFunction : public FFunctionEntry
+{
+	using FFunctionEntry::FFunctionEntry;
+
+public:
+	static TSharedPtr<FFunctionEntry_UFunction> Create(const UFunction* InFunction, const FProperty* InReturnValue, TArray<const FProperty*>&& InArguments)
+	{
+		if (InFunction == nullptr)
+		{
+			return TSharedPtr<FFunctionEntry_UFunction>();
+		}
+
+		TSharedRef<FFunctionEntry_UFunction> Item = MakeShared<FFunctionEntry_UFunction>(FConversionFunctionValue(InFunction));
+		Item->ReturnValue = InReturnValue;
+		Item->Arguments = InArguments;
+
+		return Item;
+	}
+	
+	TArray<const FProperty*> Arguments;
+	const FProperty* ReturnValue = nullptr;
+};
+
+/**
+ *
+ */
+struct FFunctionEntry_Node : public FFunctionEntry
+{
+	using FFunctionEntry::FFunctionEntry;
+
+public:
+	static TSharedPtr<FFunctionEntry_Node> Create(const TSubclassOf<UK2Node> InFunction, const UEdGraphPin* InReturnValue, TArray<UEdGraphPin*>&& InArguments)
+	{
+		if (InFunction == nullptr)
+		{
+			return TSharedPtr<FFunctionEntry_Node>();
+		}
+
+		TSharedRef<FFunctionEntry_Node> Item = MakeShared<FFunctionEntry_Node>(FConversionFunctionValue(InFunction));
+		Item->ReturnValue = InReturnValue;
+		Item->Arguments = InArguments;
+
+		return Item;
+	}
+	
+	TArray<UEdGraphPin*> Arguments;
+	const UEdGraphPin* ReturnValue = nullptr;
+};
+
+} //namespace Private
+
+
+FFunctionEntry::FFunctionEntry(FConversionFunctionValue Function)
+	: ConversionFunctionValue(Function)
+{}
+
+FCollection::FCollection()
+{
+	FCoreUObjectDelegates::OnAssetLoaded.AddRaw(this, &FCollection::HandleObjectLoaded);
+
+	IAssetRegistry& AssetRegistry = FModuleManager::LoadModuleChecked<FAssetRegistryModule>(TEXT("AssetRegistry")).Get();
+	AssetRegistry.OnAssetAdded().AddRaw(this, &FCollection::HandleAssetAdded);
+	AssetRegistry.OnAssetRemoved().AddRaw(this, &FCollection::HandleAssetRemoved);
+	AssetRegistry.OnAssetRenamed().AddRaw(this, &FCollection::HandleAssetRenamed);
+
+	FEditorDelegates::OnAssetsPreDelete.AddRaw(this, &FCollection::HandleObjectPendingDelete);
+	FKismetEditorUtilities::OnBlueprintUnloaded.AddRaw(this, &FCollection::HandleBlueprintUnloaded);
+
+	FModuleManager::Get().OnModulesChanged().AddRaw(this, &FCollection::HandleModulesChanged);
+	FCoreUObjectDelegates::ReloadCompleteDelegate.AddRaw(this, &FCollection::HandleReloadComplete);
+}
+
+FCollection::~FCollection()
+{
+	FCoreUObjectDelegates::OnAssetLoaded.RemoveAll(this);
+
+	if (FAssetRegistryModule* AssetRegistryModule = FModuleManager::GetModulePtr<FAssetRegistryModule>(TEXT("AssetRegistry")))
+	{
+		IAssetRegistry& AssetRegistry = AssetRegistryModule->Get();
+		AssetRegistry.OnAssetAdded().RemoveAll(this);
+		AssetRegistry.OnAssetRemoved().RemoveAll(this);
+		AssetRegistry.OnAssetRenamed().RemoveAll(this);
+	}
+
+	FEditorDelegates::OnAssetsPreDelete.RemoveAll(this);
+	FKismetEditorUtilities::OnBlueprintUnloaded.RemoveAll(this);
+
+	FModuleManager::Get().OnModulesChanged().RemoveAll(this);
+	FCoreUObjectDelegates::ReloadCompleteDelegate.RemoveAll(this);
+
+	if (UObjectInitialized() && !IsEngineExitRequested())
+	{
+		for (TPair<FObjectKey, FFunctionContainer> PairClass : ClassOrBlueprintToFunctions)
+		{
+			if (PairClass.Value.bIsUserWidget)
+			{
+				UnegisterBlueprintCallback(PairClass.Key.ResolveObjectPtr());
+			}
+		}
+	}
+}
+
+void FCollection::RefreshIfNeeded()
+{
+	if (bRefreshAll)
+	{
+		for (TPair<FObjectKey, FFunctionContainer> PairClass : ClassOrBlueprintToFunctions)
+		{
+			if (PairClass.Value.bIsUserWidget)
+			{
+				UnegisterBlueprintCallback(PairClass.Key.ResolveObjectPtr());
+			}
+		}
+
+		ClassOrBlueprintToFunctions.Empty();
+		ObjectToRefresh.Empty();
+		ModuleToRefresh.Empty();
+
+		NumberOfFunctions = 0;
+		bRefreshAll = false;
+		Build();
+	}
+
+	if (ModuleToRefresh.Num() > 0 || ObjectToRefresh.Num() > 0)
+	{
+		TArray<const UClass*> AllowedClasses = GetDefault<UMVVMDeveloperProjectSettings>()->GetAllowedConversionFunctionClasses();
+		TArray<const UClass*> DeniedClasses = GetDefault<UMVVMDeveloperProjectSettings>()->GetDeniedConversionFunctionClasses();
+		for (FName ModuleName : ModuleToRefresh)
+		{		
+			const FName ModuleScriptPackageName = FPackageName::GetModuleScriptPackageName(ModuleName);
+			if (GetDefault<UMVVMDeveloperProjectSettings>()->DeniedModuleForConversionFunctions.Contains(ModuleScriptPackageName))
+			{
+				continue;
+			}
+			
+			if (const UPackage* ModuleScriptPackage = FindPackage(nullptr, *ModuleScriptPackageName.ToString()))
+			{
+				TArray<UObject*> ObjectsToProcess;
+				const bool bIncludeNestedObjects = false;
+				GetObjectsWithPackage(ModuleScriptPackage, ObjectsToProcess, bIncludeNestedObjects, RF_ClassDefaultObject);
+				for (UObject* Object : ObjectsToProcess)
+				{
+					Build_Class(AllowedClasses, DeniedClasses, Object);
+				}
+			}
+		}
+		ModuleToRefresh.Empty();
+
+		for (FObjectKey ObjectKey : ObjectToRefresh)
+		{
+			// Remove the class
+			FFunctionContainer FunctionEntriesToRemove;
+			ClassOrBlueprintToFunctions.RemoveAndCopyValue(ObjectKey, FunctionEntriesToRemove);
+			NumberOfFunctions -= FunctionEntriesToRemove.Functions.Num();
+
+			// Add the class
+			if (UObject* Object = ObjectKey.ResolveObjectPtr())
+			{
+				if (FunctionEntriesToRemove.bIsUserWidget)
+				{
+					UnegisterBlueprintCallback(Object);
+				}
+
+				Build_Class(AllowedClasses, DeniedClasses, Object);
+			}
+		}
+		ObjectToRefresh.Empty();
+	}
+}
+
+bool FCollection::IsClassSupported(const TArray<const UClass*>& AllowClasses, const TArray<const UClass*>& DenyClasses, const UClass* Class) const
+{
+	auto InDenyList = [&DenyClasses, Class]()
+	{
+		for (const UClass* DenyClass : DenyClasses)
+		{
+			if (Class->IsChildOf(DenyClass))
+			{
+				return true;
+			}
+		}
+		return false;
+	};
+
+	FName ModuleName = Class->GetClassPathName().GetPackageName();
+	bool bIsModuleDenied = GetDefault<UMVVMDeveloperProjectSettings>()->DeniedModuleForConversionFunctions.Contains(ModuleName);
+	if (bIsModuleDenied)
+	{
+		return false;
+	}
+
+	if (!Class->HasAnyClassFlags(CLASS_Deprecated | CLASS_Hidden | CLASS_NewerVersionExists) && Class->GetPackage() != GetTransientPackage())
+	{
+		if (Class->IsChildOf(UUserWidget::StaticClass()))
+		{
+			return !InDenyList();
+		}
+		else
+		{
+			// Is it a child of an allowed class
+			for (const UClass* AllowClass : AllowClasses)
+			{
+				if (Class->IsChildOf(AllowClass))
+				{
+					// Confirm that it's not in the deny
+					return !InDenyList();
+				}
+			}
+		}
+	}
+	return false;
+}
+
+void FCollection::Build()
+{
+	FScopedSlowTask SlowTask = FScopedSlowTask(0, LOCTEXT("BuildingConversionFunctionLibrary", "Loading Conversion function library"));
+
+	TArray<const UClass*> AllSupportedClass;
+	AllSupportedClass.Reserve(128);
+
+	// Make sure the base class are loaded.
+	TArray<const UClass*> AllowClasses;
+	TArray<const UClass*> DenyClasses = GetDefault<UMVVMDeveloperProjectSettings>()->GetDeniedConversionFunctionClasses();
+
+	{
+		AllowClasses.Reserve(GetDefault<UMVVMDeveloperProjectSettings>()->AllowedClassForConversionFunctions.Num());
+		for (const FSoftClassPath& SoftClass : GetDefault<UMVVMDeveloperProjectSettings>()->AllowedClassForConversionFunctions)
+		{
+			if (UClass* Class = SoftClass.TryLoadClass<UObject>())
+			{
+				AllowClasses.Add(Class);
+			}
+		}
+	}
+
+	for (TObjectIterator<UClass> It; It; ++It)
+	{
+		const UClass* Class = *It;
+		if (IsClassSupported(AllowClasses, DenyClasses, Class))
+		{
+			AllSupportedClass.Add(Class);
+		}
+	}
+
+	for (const UClass* Class : AllSupportedClass)
+	{
+		if (Class->IsChildOf(UK2Node::StaticClass()))
+		{
+			AddNode(TSubclassOf<UK2Node>(const_cast<UClass*>(Class)));
+		}
+		else
+		{
+			AddClassFunctions(Class);
+		}
+	}
+}
+
+void FCollection::Build_Class(const TArray<const UClass*>& AllowClasses, const TArray<const UClass*>& DenyClasses, const UObject* Object)
+{
+	const UClass* Class = Cast<const UClass>(Object);
+	if (const UBlueprint* Blueprint = Cast<const UBlueprint>(Object))
+	{
+		Class = Blueprint->GeneratedClass ? Blueprint->GeneratedClass : Blueprint->SkeletonGeneratedClass;
+	}
+
+	if (Class)
+	{
+		bool bIsAllowed = IsClassSupported(AllowClasses, DenyClasses, Class);
+		if (bIsAllowed)
+		{
+			if (Class->IsChildOf(UK2Node::StaticClass()))
+			{
+				AddNode(TSubclassOf<UK2Node>(const_cast<UClass*>(Class)));
+			}
+			else
+			{
+				AddClassFunctions(Class);
+			}
+		}
+	}
+}
+
+void FCollection::AddClassFunctions(const UClass* Class)
+{
+	RegisterBlueprintCallback(Class);
+
+	UBlueprint* BpOwner = Cast<UBlueprint>(Class->ClassGeneratedBy);
+	auto IsInheritedBlueprintFunction = [BpOwner](const UFunction* Function)
+		{
+			if (BpOwner)
+			{
+				if (UClass* ParentClass = BpOwner->ParentClass)
+				{
+					FName FuncName = Function->GetFName();
+					return ParentClass->FindFunctionByName(FuncName, EIncludeSuperFlag::IncludeSuper) != nullptr;
+				}
+			}
+			return false;
+		};
+
+	const bool bIsUserWidget = Class->IsChildOf<UUserWidget>();
+	FFunctionContainer* FunctionContainer = nullptr;
+	for (TFieldIterator<UFunction> FunctionIt(Class, EFieldIteratorFlags::ExcludeSuper); FunctionIt; ++FunctionIt)
+	{
+		const UFunction* Function = *FunctionIt;
+		// Any functions in a WidgetBlueprint, or the functions have to be static functions in a BlueprintFunctionLibrary.
+		bool bIsFromWidgetBlueprint = bIsUserWidget && Function->HasAllFunctionFlags(FUNC_BlueprintPure | FUNC_Const);
+		bool bFromBlueprintFunctionLibrary = !bIsUserWidget && Function->HasAllFunctionFlags(FUNC_Static | FUNC_BlueprintPure);
+		if (!bIsFromWidgetBlueprint && !bFromBlueprintFunctionLibrary)
+		{
+			continue;
+		}
+
+		// As one return value.
+		const FProperty* ReturnProperty = UE::MVVM::BindingHelper::GetReturnProperty(Function);
+		if (ReturnProperty == nullptr)
+		{
+			continue;
+		}
+
+		// As at least one argument
+		bool bValidProperty = false;
+		TArray<const FProperty*> ArgumentsResult = UE::MVVM::BindingHelper::GetAllArgumentProperties(Function);
+		if (ArgumentsResult.Num() == 0)
+		{
+			continue;
+		}
+
+		// Is function virtual and overridden by this class.
+		if (IsInheritedBlueprintFunction(Function))
+		{
+			continue;
+		}
+
+		// Is it deprecated
+		if (!UEdGraphSchema_K2::CanUserKismetCallFunction(Function))
+		{
+			continue;
+		}
+
+		// Apply general filtering for functions
+		if (!FBlueprintActionDatabase::IsFunctionAllowed(Function, FBlueprintActionDatabase::EPermissionsContext::Node))
+		{
+			continue;
+		}
+
+		// Create the function entry
+		TSharedPtr<Private::FFunctionEntry_UFunction> FunctionItem = Private::FFunctionEntry_UFunction::Create(Function, ReturnProperty, MoveTemp(ArgumentsResult));
+		if (FunctionItem)
+		{
+			if (FunctionContainer == nullptr)
+			{
+				FunctionContainer = &(ClassOrBlueprintToFunctions.FindOrAdd(FObjectKey(Class->ClassGeneratedBy ? Class->ClassGeneratedBy.Get() : Class)));
+				FunctionContainer->bIsUserWidget = bIsUserWidget;
+				check(FunctionContainer->bIsNode == false);
+			}
+			FunctionContainer->Functions.Add(FunctionItem.ToSharedRef());
+			++NumberOfFunctions;
+		}
+	}
+}
+
+void FCollection::AddNode(TSubclassOf<UK2Node> Function)
+{
+	UEdGraph* OwnerOfTemporaries = NewObject<UEdGraph>(GetTransientPackage());
+	UK2Node* NewNode = NewObject<UK2Node>(OwnerOfTemporaries, Function.Get(), NAME_None, RF_Transient);
+	NewNode->AllocateDefaultPins();
+	NewNode->PostPlacedNewNode();
+
+ 	TArray<UEdGraphPin*> InputPins = UE::MVVM::ConversionFunctionHelper::FindInputPins(NewNode);
+	if (InputPins.Num() == 0)
+	{
+		return;
+	}
+
+	UEdGraphPin* OutputPin = UE::MVVM::ConversionFunctionHelper::FindOutputPin(NewNode);
+	if (OutputPin == nullptr)
+	{
+		return;
+	}
+
+	TSharedPtr<Private::FFunctionEntry_Node> FunctionItem = Private::FFunctionEntry_Node::Create(Function, OutputPin, MoveTemp(InputPins));
+	if (FunctionItem)
+	{
+		FFunctionContainer& FunctionContainer = ClassOrBlueprintToFunctions.FindOrAdd(FObjectKey(Function.Get()));
+		FunctionContainer.Functions.Add(FunctionItem.ToSharedRef());
+		FunctionContainer.bIsUserWidget = false;
+		FunctionContainer.bIsNode = true;
+		check(FunctionContainer.Functions.Num() == 1);
+	}
+}
+
+void FCollection::RegisterBlueprintCallback(const UClass* Class)
+{
+	if (UBlueprint* BlueprintAsset = Cast<UBlueprint>(Class->ClassGeneratedBy))
+	{
+		UBlueprint::FChangedEvent& OnBPChanged = BlueprintAsset->OnChanged();
+		UBlueprint::FCompiledEvent& OnBPCompiled = BlueprintAsset->OnCompiled();
+		if (!OnBPChanged.IsBoundToObject(this))
+		{
+			OnBPChanged.AddRaw(this, &FCollection::HandleBlueprintChanged);
+		}
+		if (!OnBPCompiled.IsBoundToObject(this))
+		{
+			OnBPCompiled.AddRaw(this, &FCollection::HandleBlueprintChanged);
+		}
+	}
+}
+
+void FCollection::UnegisterBlueprintCallback(UObject* ObjectKey)
+{
+	if (UBlueprint* Blueprint = Cast<UBlueprint>(ObjectKey))
+	{
+		Blueprint->OnChanged().RemoveAll(this);
+		Blueprint->OnCompiled().RemoveAll(this);
+	}
+}
+
+TArray<::UE::MVVM::FConversionFunctionValue> FCollection::GetFunctions(const UWidgetBlueprint* WidgetBlueprint) const
+{
+	return GetFunctions(WidgetBlueprint, nullptr, nullptr);
+}
+
+TArray<::UE::MVVM::FConversionFunctionValue> FCollection::GetFunctions(const UWidgetBlueprint* WidgetBlueprint, const FProperty* ArgumentType, const FProperty* ReturnType) const
+{
+	const_cast<FCollection*>(this)->RefreshIfNeeded();
+	
+	const bool bHasArgumentsTest = ArgumentType != nullptr;
+	const bool bHasResultTest = ReturnType != nullptr;
+	
+	TArray<::UE::MVVM::FConversionFunctionValue> Result;
+	Result.Reserve(NumberOfFunctions);
+	
+	check(WidgetBlueprint);
+	const UClass* WidgetBlueprintClass = WidgetBlueprint->GeneratedClass ? WidgetBlueprint->GeneratedClass : WidgetBlueprint->SkeletonGeneratedClass;
+	check(WidgetBlueprintClass);
+	
+	for (const auto& FunctionKeyValue : ClassOrBlueprintToFunctions)
+	{
+		if (FunctionKeyValue.Value.bIsUserWidget)
+		{
+			// test if it's this userwidget before adding the functions
+			const UObject* Object = FunctionKeyValue.Key.ResolveObjectPtr();
+			if (ensure(Object))
+			{
+				// Can be a UWidgetBlueprintGeneratedClass or a UWidgetBlueprint
+				if (const UWidgetBlueprint* WidgetBlueprintKey = Cast<const UWidgetBlueprint>(Object))
+				{
+					const UClass* WidgetBlueprintKeyClass = WidgetBlueprintKey->GeneratedClass ? WidgetBlueprintKey->GeneratedClass : WidgetBlueprintKey->SkeletonGeneratedClass;
+					if (!WidgetBlueprintClass->IsChildOf(WidgetBlueprintKeyClass))
+					{
+						continue;
+					}
+				}
+				else if (const UClass* ClassKey = Cast<const UClass>(Object))
+				{
+					if (!WidgetBlueprintClass->IsChildOf(ClassKey))
+					{
+						continue;
+					}
+				}
+				else
+				{
+					ensure(false);
+					continue;
+				}
+			}
+		}
+				
+		if (!bHasArgumentsTest && !bHasResultTest)
+		{
+			for (const TSharedRef<FFunctionEntry>& FunctionEntry : FunctionKeyValue.Value.Functions)
+			{
+				Result.Add(FunctionEntry->ConversionFunctionValue);
+			}
+		}
+		else
+		{
+			FEdGraphPinType ReturnPinType;
+			if (ReturnType)
+			{
+				GetDefault<UEdGraphSchema_K2>()->ConvertPropertyToPinType(ReturnType, ReturnPinType);
+			}
+			FEdGraphPinType ArgumentPinType;
+			if (ArgumentType)
+			{
+				GetDefault<UEdGraphSchema_K2>()->ConvertPropertyToPinType(ArgumentType, ArgumentPinType);
+			}
+
+			if (!FunctionKeyValue.Value.bIsNode)
+			{
+				for (const TSharedRef<FFunctionEntry>& FunctionEntryIt : FunctionKeyValue.Value.Functions)
+				{
+					const TSharedRef<Private::FFunctionEntry_UFunction>& FunctionEntry = StaticCastSharedRef<Private::FFunctionEntry_UFunction>(FunctionEntryIt);
+					if (ReturnType && !UE::MVVM::BindingHelper::ArePropertiesCompatible(FunctionEntry->ReturnValue, ReturnType))
+					{
+						continue;
+					}
+
+					if (bHasArgumentsTest)
+					{
+						bool bFound = false;
+						for (const FProperty* Property : FunctionEntry->Arguments)
+						{
+							if (UE::MVVM::BindingHelper::ArePropertiesCompatible(Property, ArgumentType))
+							{
+								bFound = true;
+								break;
+							}
+						}
+						if (!bFound)
+						{
+							continue;
+						}
+					}
+
+					Result.Add(FunctionEntry->ConversionFunctionValue);
+				}
+			}
+			else
+			{
+				check(FunctionKeyValue.Value.Functions.Num() == 1);
+				const TSharedRef<Private::FFunctionEntry_Node>& FunctionEntry = StaticCastSharedRef<Private::FFunctionEntry_Node>(FunctionKeyValue.Value.Functions[0]);
+				bool bIgnoreArray = true;
+
+				if (ReturnType && !GetDefault<UEdGraphSchema_K2>()->ArePinTypesCompatible(ReturnPinType, FunctionEntry->ReturnValue->PinType, WidgetBlueprintClass, bIgnoreArray))
+				{
+					continue;
+				}
+
+				if (bHasArgumentsTest)
+				{
+					bool bFound = false;
+					for (const UEdGraphPin* PropertyPinType : FunctionEntry->Arguments)
+					{
+						if (!GetDefault<UEdGraphSchema_K2>()->ArePinTypesCompatible(ArgumentPinType, PropertyPinType->PinType, WidgetBlueprintClass, bIgnoreArray))
+						{
+							bFound = true;
+							break;
+						}
+					}
+					if (!bFound)
+					{
+						continue;
+					}
+				}
+
+				Result.Add(FunctionEntry->ConversionFunctionValue);
+			}
+		}
+	}
+	
+	return Result;
+}
+
+const TSharedPtr<FFunctionEntry> FCollection::FindFunction(::UE::MVVM::FConversionFunctionValue FunctionValue) const
+{
+	const_cast<FCollection*>(this)->RefreshIfNeeded();
+
+	if (FunctionValue.IsFunction())
+	{
+		const UFunction* Function = FunctionValue.GetFunction();
+
+		// Find the function's Blueprint or the function's native class.
+		const FFunctionContainer* FoundCountainer = nullptr;
+		if (Function->GetClass()->ClassGeneratedBy)
+		{
+			FoundCountainer = ClassOrBlueprintToFunctions.Find(FObjectKey(Function->GetOuterUClass()->ClassGeneratedBy));
+		}
+		else
+		{
+			FoundCountainer = ClassOrBlueprintToFunctions.Find(FObjectKey(Function->GetOuterUClass()));
+		}
+		if (FoundCountainer == nullptr)
+		{
+			return TSharedPtr<FFunctionEntry>();
+		}
+
+		for (const TSharedRef<FFunctionEntry>& Entry : FoundCountainer->Functions)
+		{
+			if (Entry->GetFunction() == Function)
+			{
+				return Entry;
+			}
+		}
+	}
+	else if (FunctionValue.IsNode())
+	{
+		TSubclassOf<UK2Node> Function = FunctionValue.GetNode();
+
+		const FFunctionContainer* FoundCountainer = ClassOrBlueprintToFunctions.Find(FObjectKey(Function.Get()));
+		return FoundCountainer && FoundCountainer->Functions.Num() == 1 ? FoundCountainer->Functions[0] : TSharedPtr<FFunctionEntry>();
+	}
+
+	return TSharedPtr<FFunctionEntry>();
+}
+
+namespace Private
+{
+bool IsObjectValidForCollection(UObject* Object)
+{
+	return Object 
+		&& !Object->GetOutermost()->HasAnyPackageFlags(PKG_PlayInEditor | PKG_ForDiffing)
+		&& (Cast<UBlueprint>(Object) || Cast<UClass>(Object)) 
+		&& Object->IsAsset();
+}
+} // namespace
+
+void FCollection::AddObjectToRefresh(UObject* Object)
+{
+	if (Private::IsObjectValidForCollection(Object))
+	{
+		ObjectToRefresh.Add(FObjectKey(Object));
+	}
+}
+
+void FCollection::HandleBlueprintChanged(UBlueprint* Blueprint)
+{
+	AddObjectToRefresh(Blueprint);
+}
+
+void FCollection::HandleBlueprintUnloaded(UBlueprint* Blueprint)
+{
+	AddObjectToRefresh(Blueprint);
+}
+
+void FCollection::HandleObjectLoaded(UObject* Object)
+{
+	AddObjectToRefresh(Object);
+}
+
+void FCollection::HandleObjectPendingDelete(TArray<UObject*> const& ObjectsForDelete)
+{
+	FBlueprintActionDatabase& ActionDatabase = FBlueprintActionDatabase::Get();
+	for (UObject* DeletingObject : ObjectsForDelete)
+	{
+		AddObjectToRefresh(DeletingObject);
+	}
+}
+
+void FCollection::HandleAssetAdded(FAssetData const& NewAssetInfo)
+{
+	if (NewAssetInfo.IsAssetLoaded())
+	{
+		UObject* AssetObject = NewAssetInfo.GetAsset();
+		AddObjectToRefresh(AssetObject);
+	}
+}
+
+void FCollection::HandleAssetRemoved(FAssetData const& NewAssetInfo)
+{
+	if (NewAssetInfo.IsAssetLoaded())
+	{
+		UObject* AssetObject = NewAssetInfo.GetAsset();
+		AddObjectToRefresh(AssetObject);
+	}
+}
+
+void FCollection::HandleAssetRenamed(FAssetData const& NewAssetInfo, const FString& OldName)
+{
+	if (NewAssetInfo.IsAssetLoaded())
+	{
+		UObject* AssetObject = NewAssetInfo.GetAsset();
+		AddObjectToRefresh(AssetObject);
+	}
+}
+
+void FCollection::HandleModulesChanged(FName ModuleName, EModuleChangeReason ModuleChangeReason)
+{
+	switch (ModuleChangeReason)
+	{
+	case EModuleChangeReason::ModuleLoaded:
+		ModuleToRefresh.Add(ModuleName);
+		break;
+
+	case EModuleChangeReason::ModuleUnloaded:
+		bRefreshAll = true;
+		break;
+
+	default:
+		break;
+	}
+}
+
+void FCollection::HandleReloadComplete(EReloadCompleteReason Reason)
+{
+	bRefreshAll = true;
+}
+
+} // namespace UE::MVVM::BindingEntry
+
+#undef LOCTEXT_NAMESPACE
diff --git a/Engine/Plugins/Runtime/ModelViewViewModel/Source/ModelViewViewModelEditor/Private/Framework/MVVMConversionFunctionLibrary.h b/Engine/Plugins/Runtime/ModelViewViewModel/Source/ModelViewViewModelEditor/Private/Framework/MVVMConversionFunctionLibrary.h
new file mode 100644
--- /dev/null
+++ b/Engine/Plugins/Runtime/ModelViewViewModel/Source/ModelViewViewModelEditor/Private/Framework/MVVMConversionFunctionLibrary.h
@@ -0,0 +1,97 @@
+// Copyright Epic Games, Inc. All Rights Reserved.
+
+#pragma once
+
+#include "CoreMinimal.h"
+#include "EdGraph/EdGraphPin.h"
+#include "Types/MVVMConversionFunctionValue.h"
+#include "UObject/ObjectKey.h"
+
+class UBlueprint;
+class UFunction;
+class UK2Node;
+class UWidgetBlueprint;
+
+namespace UE::MVVM::ConversionFunctionLibrary
+{
+
+struct FCategoryEntry;
+class FCollection;
+
+/**
+ *
+ */
+struct FFunctionEntry
+{
+private:
+	friend FCollection;
+
+public:
+	FFunctionEntry() = default;
+	FFunctionEntry(FConversionFunctionValue Function);
+
+	const ::UE::MVVM::FConversionFunctionValue& GetFunction() const
+	{
+		return ConversionFunctionValue;
+	}
+
+private:
+	::UE::MVVM::FConversionFunctionValue ConversionFunctionValue;
+};
+
+
+/**
+ * Collection of all available conversion functions.
+ * The collection rebuild when a new object is loaded and when the WidgetBlueprint is compiled.
+ */
+class FCollection
+{
+public:
+	FCollection();
+	~FCollection();
+
+public:
+	[[nodiscard]] TArray<::UE::MVVM::FConversionFunctionValue> GetFunctions(const UWidgetBlueprint* Blueprint) const;
+	[[nodiscard]] TArray<::UE::MVVM::FConversionFunctionValue> GetFunctions(const UWidgetBlueprint* Blueprint, const FProperty* ArgumentType, const FProperty* ReturnType) const;
+
+	[[nodiscard]] const TSharedPtr<FFunctionEntry> FindFunction(::UE::MVVM::FConversionFunctionValue) const;
+
+private:
+	void RefreshIfNeeded();
+	void Build();
+	void Build_Class(const TArray<const UClass*>& AllowClasses, const TArray<const UClass*>& DenyClasses, const UObject* Class);
+	bool IsClassSupported(const TArray<const UClass*>& AllowClasses, const TArray<const UClass*>& DenyClasses, const UClass* Class) const;
+	void AddClassFunctions(const UClass* Class);
+	void AddNode(TSubclassOf<UK2Node> Function);
+	void RegisterBlueprintCallback(const UClass* Class);
+	void UnegisterBlueprintCallback(UObject* Class);
+
+private:
+	void AddObjectToRefresh(UObject* Object);
+
+	void HandleBlueprintChanged(UBlueprint* Blueprint);
+	void HandleBlueprintUnloaded(UBlueprint* Blueprint);
+	void HandleObjectLoaded(UObject* NewObject);
+	void HandleObjectPendingDelete(TArray<UObject*> const& ObjectsForDelete);
+	void HandleAssetAdded(FAssetData const& AssetInfo);
+	void HandleAssetRemoved(FAssetData const& AssetInfo);
+	void HandleAssetRenamed(FAssetData const& AssetInfo, const FString& InOldName);
+	void HandleModulesChanged(FName InModuleName, EModuleChangeReason InModuleChangeReason);
+	void HandleReloadComplete(EReloadCompleteReason Reason);
+
+private:
+	struct FFunctionContainer
+	{
+		TArray<TSharedRef<FFunctionEntry>> Functions;
+		bool bIsUserWidget = false;
+		bool bIsNode = false;
+	};
+	TMap<FObjectKey, FFunctionContainer> ClassOrBlueprintToFunctions;
+
+	TSet<FObjectKey> ObjectToRefresh; // can only be GeneratedClass or Blueprint
+	TSet<FName> ModuleToRefresh;
+	int32 NumberOfFunctions = 0;
+	bool bRefreshAll = true;
+};
+
+} // namespace UE::MVVM::ConversionFunctionLibrary
diff --git a/Engine/Plugins/Runtime/ModelViewViewModel/Source/ModelViewViewModelEditor/Private/MVVMEditorSubsystem.cpp b/Engine/Plugins/Runtime/ModelViewViewModel/Source/ModelViewViewModelEditor/Private/MVVMEditorSubsystem.cpp
--- a/Engine/Plugins/Runtime/ModelViewViewModel/Source/ModelViewViewModelEditor/Private/MVVMEditorSubsystem.cpp
+++ b/Engine/Plugins/Runtime/ModelViewViewModel/Source/ModelViewViewModelEditor/Private/MVVMEditorSubsystem.cpp
@@ -9,6 +9,7 @@
 #include "BlueprintActionFilter.h"
 #include "BlueprintNodeSpawner.h"
 #include "Engine/Engine.h"
+#include "Framework/MVVMConversionFunctionLibrary.h"
 #include "Kismet/BlueprintFunctionLibrary.h"
 #include "Kismet2/BlueprintEditorUtils.h"
 #include "K2Node_BreakStruct.h"
@@ -715,6 +716,15 @@ void UMVVMEditorSubsystem::SetCompileForEvent(UMVVMBlueprintViewEvent* Event, bo
 	}
 }
 
+namespace UE::MVVM::Private
+{
+UClass* GetCallingContext(const UWidgetBlueprint* WidgetBlueprint)
+{
+	return WidgetBlueprint->GeneratedClass ? WidgetBlueprint->GeneratedClass :
+		WidgetBlueprint->SkeletonGeneratedClass ? WidgetBlueprint->SkeletonGeneratedClass : WidgetBlueprint->ParentClass;
+}
+}
+
 bool UMVVMEditorSubsystem::IsValidConversionFunction(const UWidgetBlueprint* WidgetBlueprint, const UFunction* Function, const FMVVMBlueprintPropertyPath& Source, const FMVVMBlueprintPropertyPath& Destination) const
 {
 	if (WidgetBlueprint == nullptr || Function == nullptr)
@@ -722,64 +732,73 @@ bool UMVVMEditorSubsystem::IsValidConversionFunction(const UWidgetBlueprint* Wid
 		return false;
 	}
 
-	if (!UMVVMBlueprintViewConversionFunction::IsValidConversionFunction(WidgetBlueprint, Function))
+	UClass* CallingContext = UE::MVVM::Private::GetCallingContext(WidgetBlueprint);
+
+	const FProperty* SourceProperty = nullptr;
+	TArray<UE::MVVM::FMVVMConstFieldVariant> SourceFields = Source.GetFields(CallingContext);
+	if (SourceFields.Num() > 0)
 	{
-		return false;
+		SourceProperty = SourceFields.Last().IsProperty() ? SourceFields.Last().GetProperty() : UE::MVVM::BindingHelper::GetReturnProperty(SourceFields.Last().GetFunction());
 	}
 
+	const FProperty* DestinationProperty = nullptr;
+	TArray<UE::MVVM::FMVVMConstFieldVariant> DestFields = Destination.GetFields(CallingContext);
+	if (DestFields.Num() > 0)
 	{
-		const FProperty* SourceProperty = nullptr;
-		TArray<UE::MVVM::FMVVMConstFieldVariant> SourceFields = Source.GetFields(WidgetBlueprint->SkeletonGeneratedClass);
-		if (SourceFields.Num() > 0)
+		if (!DestFields.Last().IsEmpty())
 		{
-			SourceProperty = SourceFields.Last().IsProperty() ? SourceFields.Last().GetProperty() : UE::MVVM::BindingHelper::GetReturnProperty(SourceFields.Last().GetFunction());
+			DestinationProperty = DestFields.Last().IsProperty() ? DestFields.Last().GetProperty() : UE::MVVM::BindingHelper::GetFirstArgumentProperty(DestFields.Last().GetFunction());
 		}
+	}
 
-		// check that at least one source -> argument binding is compatible
-		if (SourceProperty)
-		{
-			TValueOrError<TArray<const FProperty*>, FText> ArgumentsResult = UE::MVVM::BindingHelper::TryGetArgumentsForConversionFunction(Function);
-			check(ArgumentsResult.HasValue());
+	return IsValidConversionFunction(WidgetBlueprint, Function, SourceProperty, DestinationProperty);
+}
 
-			bool bAnyCompatible = false;
+bool UMVVMEditorSubsystem::IsValidConversionFunction(const UWidgetBlueprint* WidgetBlueprint, const UFunction* Function, const FProperty* ExpectedArgumentType, const FProperty* ExptectedReturnType) const
+{
+	if (WidgetBlueprint == nullptr || Function == nullptr)
+	{
+		return false;
+	}
 
-			const TArray<const FProperty*>& ConversionArgProperties = ArgumentsResult.GetValue();
-			for (const FProperty* ArgumentProperty : ConversionArgProperties)
-			{
-				if (UE::MVVM::BindingHelper::ArePropertiesCompatible(SourceProperty, ArgumentProperty))
-				{
-					bAnyCompatible = true;
-					break;
-				}
-			}
-			if (!bAnyCompatible)
-			{
-				return false;
-			}
-		}
+	if (!UMVVMBlueprintViewConversionFunction::IsValidConversionFunction(WidgetBlueprint, Function))
+	{
+		return false;
 	}
 
+
+	// check that at least one source -> argument binding is compatible
+	if (ExpectedArgumentType)
 	{
-		const FProperty* DestinationProperty = nullptr;
-		TArray<UE::MVVM::FMVVMConstFieldVariant> DestFields = Destination.GetFields(WidgetBlueprint->SkeletonGeneratedClass);
-		if (DestFields.Num() > 0)
+		TValueOrError<TArray<const FProperty*>, FText> ArgumentsResult = UE::MVVM::BindingHelper::TryGetArgumentsForConversionFunction(Function);
+		check(ArgumentsResult.HasValue());
+
+		bool bAnyCompatible = false;
+
+		const TArray<const FProperty*>& ConversionArgProperties = ArgumentsResult.GetValue();
+		for (const FProperty* ArgumentProperty : ConversionArgProperties)
 		{
-			if (!DestFields.Last().IsEmpty())
+			if (UE::MVVM::BindingHelper::ArePropertiesCompatible(ExpectedArgumentType, ArgumentProperty))
 			{
-				DestinationProperty = DestFields.Last().IsProperty() ? DestFields.Last().GetProperty() : UE::MVVM::BindingHelper::GetFirstArgumentProperty(DestFields.Last().GetFunction());
+				bAnyCompatible = true;
+				break;
 			}
 		}
-
-		if (DestinationProperty)
+		if (!bAnyCompatible)
 		{
-			TValueOrError<const FProperty*, FText> ReturnResult = UE::MVVM::BindingHelper::TryGetReturnTypeForConversionFunction(Function);
-			check(ReturnResult.HasValue());
+			return false;
+		}
+	}
 
-			// check that the return -> dest is valid
-			if (!UE::MVVM::BindingHelper::ArePropertiesCompatible(ReturnResult.GetValue(), DestinationProperty))
-			{
-				return false;
-			}
+	if (ExptectedReturnType)
+	{
+		TValueOrError<const FProperty*, FText> ReturnResult = UE::MVVM::BindingHelper::TryGetReturnTypeForConversionFunction(Function);
+		check(ReturnResult.HasValue());
+
+		// check that the return -> dest is valid
+		if (!UE::MVVM::BindingHelper::ArePropertiesCompatible(ReturnResult.GetValue(), ExptectedReturnType))
+		{
+			return false;
 		}
 	}
 
@@ -793,92 +812,114 @@ bool UMVVMEditorSubsystem::IsValidConversionNode(const UWidgetBlueprint* WidgetB
 		return false;
 	}
 
+	UClass* CallingContext = UE::MVVM::Private::GetCallingContext(WidgetBlueprint);
+
+	const FProperty* SourceProperty = nullptr;
+	TArray<UE::MVVM::FMVVMConstFieldVariant> SourceFields = Source.GetFields(CallingContext);
+	if (SourceFields.Num() > 0)
+	{
+		SourceProperty = SourceFields.Last().IsProperty() ? SourceFields.Last().GetProperty() : UE::MVVM::BindingHelper::GetReturnProperty(SourceFields.Last().GetFunction());
+	}
+
+	const FProperty* DestinationProperty = nullptr;
+	TArray<UE::MVVM::FMVVMConstFieldVariant> DestFields = Destination.GetFields(CallingContext);
+	if (DestFields.Num() > 0)
+	{
+		if (!DestFields.Last().IsEmpty())
+		{
+			DestinationProperty = DestFields.Last().IsProperty() ? DestFields.Last().GetProperty() : UE::MVVM::BindingHelper::GetFirstArgumentProperty(DestFields.Last().GetFunction());
+		}
+	}
+
+	return IsValidConversionNode(WidgetBlueprint, Function, SourceProperty, DestinationProperty);
+}
+
+bool UMVVMEditorSubsystem::IsValidConversionNode(const UWidgetBlueprint* WidgetBlueprint, const TSubclassOf<UK2Node> Function, const FProperty* ExpectedArgumentType, const FProperty* ExptectedReturnType) const
+{
+	if (WidgetBlueprint == nullptr || Function.Get() == nullptr)
+	{
+		return false;
+	}
+
 	if (!UMVVMBlueprintViewConversionFunction::IsValidConversionNode(WidgetBlueprint, Function))
 	{
 		return false;
 	}
 
-	UClass* CallingContext = WidgetBlueprint->GeneratedClass ? WidgetBlueprint->GeneratedClass : WidgetBlueprint->ParentClass;
+	UClass* CallingContext = UE::MVVM::Private::GetCallingContext(WidgetBlueprint);
+
+	// check that at least one source -> argument binding is compatible
+	if (ExpectedArgumentType)
 	{
-		const FProperty* SourceProperty = nullptr;
-		TArray<UE::MVVM::FMVVMConstFieldVariant> SourceFields = Source.GetFields(CallingContext);
-		if (SourceFields.Num() > 0)
+		FEdGraphPinType SourcePinType;
+		if (!GetDefault<UEdGraphSchema_K2>()->ConvertPropertyToPinType(ExpectedArgumentType, SourcePinType))
 		{
-			SourceProperty = SourceFields.Last().IsProperty() ? SourceFields.Last().GetProperty() : UE::MVVM::BindingHelper::GetReturnProperty(SourceFields.Last().GetFunction());
+			return false;
 		}
 
-		// check that at least one source -> argument binding is compatible
-		if (SourceProperty)
+		TArray<UEdGraphPin*> InputPins = UE::MVVM::ConversionFunctionHelper::FindInputPins(Function.GetDefaultObject());
+		if (InputPins.Num() == 0)
 		{
-			FEdGraphPinType SourcePinType;
-			if (!GetDefault<UEdGraphSchema_K2>()->ConvertPropertyToPinType(SourceProperty, SourcePinType))
-			{
-				return false;
-			}
-
-			TArray<UEdGraphPin*> InputPins = UE::MVVM::ConversionFunctionHelper::FindInputPins(Function.GetDefaultObject());
-			if (InputPins.Num() == 0)
-			{
-				return false;
-			}
+			return false;
+		}
 
-			bool bAnyCompatible = false;
-			for (const UEdGraphPin* ArgumentPin : InputPins)
-			{
-				bool bIgnoreArray = true;
-				const bool bTypesMatch = GetDefault<UEdGraphSchema_K2>()->ArePinTypesCompatible(SourcePinType, ArgumentPin->PinType, CallingContext, bIgnoreArray);
-				if (bTypesMatch)
-				{
-					bAnyCompatible = true;
-					break;
-				}
-			}
-			if (!bAnyCompatible)
+		bool bAnyCompatible = false;
+		for (const UEdGraphPin* ArgumentPin : InputPins)
+		{
+			bool bIgnoreArray = true;
+			const bool bTypesMatch = GetDefault<UEdGraphSchema_K2>()->ArePinTypesCompatible(SourcePinType, ArgumentPin->PinType, CallingContext, bIgnoreArray);
+			if (bTypesMatch)
 			{
-				return false;
+				bAnyCompatible = true;
+				break;
 			}
 		}
+		if (!bAnyCompatible)
+		{
+			return false;
+		}
 	}
 
+	if (ExptectedReturnType)
 	{
-		const FProperty* DestinationProperty = nullptr;
-		TArray<UE::MVVM::FMVVMConstFieldVariant> DestFields = Destination.GetFields(CallingContext);
-		if (DestFields.Num() > 0)
+		FEdGraphPinType DestinationPinType;
+		if (!GetDefault<UEdGraphSchema_K2>()->ConvertPropertyToPinType(ExptectedReturnType, DestinationPinType))
 		{
-			if (!DestFields.Last().IsEmpty())
-			{
-				DestinationProperty = DestFields.Last().IsProperty() ? DestFields.Last().GetProperty() : UE::MVVM::BindingHelper::GetFirstArgumentProperty(DestFields.Last().GetFunction());
-			}
+			return false;
 		}
 
-		if (DestinationProperty)
+		UEdGraphPin* ReturnPin = UE::MVVM::ConversionFunctionHelper::FindOutputPin(Function.GetDefaultObject());
+		if (ReturnPin == nullptr)
 		{
-			FEdGraphPinType DestinationPinType;
-			if (!GetDefault<UEdGraphSchema_K2>()->ConvertPropertyToPinType(DestinationProperty, DestinationPinType))
-			{
-				return false;
-			}
-
-			UEdGraphPin* ReturnPin = UE::MVVM::ConversionFunctionHelper::FindOutputPin(Function.GetDefaultObject());
-			if (ReturnPin == nullptr)
-			{
-				return false;
-			}
+			return false;
+		}
 
-			// check that the return -> dest is valid
-			bool bIgnoreArray = true;
-			const bool bTypesMatch = GetDefault<UEdGraphSchema_K2>()->ArePinTypesCompatible(ReturnPin->PinType, DestinationPinType, CallingContext, bIgnoreArray);
-			if (!bTypesMatch)
-			{
-				return false;
-			}
+		// check that the return -> dest is valid
+		bool bIgnoreArray = true;
+		const bool bTypesMatch = GetDefault<UEdGraphSchema_K2>()->ArePinTypesCompatible(ReturnPin->PinType, DestinationPinType, CallingContext, bIgnoreArray);
+		if (!bTypesMatch)
+		{
+			return false;
 		}
 	}
 
 	return true;
 }
 
-bool UMVVMEditorSubsystem::IsSimpleConversionFunctionA(const UFunction* Function) const
+bool UMVVMEditorSubsystem::IsValidConversionFunction(const UWidgetBlueprint* WidgetBlueprint, UE::MVVM::FConversionFunctionValue Function, const FProperty* ExpectedArgumentType, const FProperty* ExptectedReturnType) const
+{
+	if (Function.IsFunction())
+	{
+		return IsValidConversionFunction(WidgetBlueprint, Function.GetFunction(), ExpectedArgumentType, ExptectedReturnType);
+	}
+	else if (Function.IsNode())
+	{
+		return IsValidConversionNode(WidgetBlueprint, Function.GetNode(), ExpectedArgumentType, ExptectedReturnType);
+	}
+	return false;
+}
+
+bool UMVVMEditorSubsystem::IsSimpleConversionFunction(const UFunction* Function) const
 {
 	TValueOrError<const FProperty*, FText> ReturnResult = UE::MVVM::BindingHelper::TryGetReturnTypeForConversionFunction(Function);
 	if (ReturnResult.HasError())
@@ -917,19 +958,60 @@ UK2Node_CallFunction* UMVVMEditorSubsystem::GetConversionFunctionNode(const UWid
 
 TArray<UFunction*> UMVVMEditorSubsystem::GetAvailableConversionFunctions(const UWidgetBlueprint* WidgetBlueprint, const FMVVMBlueprintPropertyPath& Source, const FMVVMBlueprintPropertyPath& Destination) const
 {
-	TArray<UFunction*> ConversionFunctions;
+	if (WidgetBlueprint == nullptr)
+	{
+		return TArray<UFunction*>();
+	}
+
+	UClass* CallingContext = UE::MVVM::Private::GetCallingContext(WidgetBlueprint);
 
-	auto AddFunction = [this, &Source, &Destination, WidgetBlueprint, &ConversionFunctions](const UFunction* Function)
+	const FProperty* SourceProperty = nullptr;
+	TArray<UE::MVVM::FMVVMConstFieldVariant> SourceFields = Source.GetFields(CallingContext);
+	if (SourceFields.Num() > 0)
 	{
-		if (IsValidConversionFunction(WidgetBlueprint, Function, Source, Destination))
+		SourceProperty = SourceFields.Last().IsProperty() ? SourceFields.Last().GetProperty() : UE::MVVM::BindingHelper::GetReturnProperty(SourceFields.Last().GetFunction());
+	}
+
+	const FProperty* DestinationProperty = nullptr;
+	TArray<UE::MVVM::FMVVMConstFieldVariant> DestFields = Destination.GetFields(CallingContext);
+	if (DestFields.Num() > 0)
+	{
+		if (!DestFields.Last().IsEmpty())
 		{
-			ConversionFunctions.Add(const_cast<UFunction*>(Function));
+			DestinationProperty = DestFields.Last().IsProperty() ? DestFields.Last().GetProperty() : UE::MVVM::BindingHelper::GetFirstArgumentProperty(DestFields.Last().GetFunction());
 		}
-	};
+	}
+
+	TArray<UE::MVVM::FConversionFunctionValue> NewResult = GetConversionFunctions(WidgetBlueprint, SourceProperty, DestinationProperty);
+	TArray<UFunction*> ConversionFunctions;
+	ConversionFunctions.Reserve(NewResult.Num());
+
+	Algo::TransformIf(NewResult, ConversionFunctions
+		, [](const UE::MVVM::FConversionFunctionValue& Other)->bool{ return Other.IsFunction(); }
+		, [](const UE::MVVM::FConversionFunctionValue& Other)->UFunction*{ return const_cast<UFunction*>(Other.GetFunction()); });
 
+	ConversionFunctions.Sort([](const UFunction& A, const UFunction& B) -> bool
+		{
+			return A.GetFName().LexicalLess(B.GetFName());
+		});
+	return ConversionFunctions;
+}
+
+TArray<UE::MVVM::FConversionFunctionValue> UMVVMEditorSubsystem::GetConversionFunctions(const UWidgetBlueprint* WidgetBlueprint, const FProperty* ExpectedArgumentType, const FProperty* ExptectedReturnType) const
+{
 	EMVVMDeveloperConversionFunctionFilterType FilterType = GetDefault<UMVVMDeveloperProjectSettings>()->GetConversionFunctionFilter();
 	if (FilterType == EMVVMDeveloperConversionFunctionFilterType::BlueprintActionRegistry)
 	{
+		TArray<UE::MVVM::FConversionFunctionValue> ConversionFunctions;
+
+		auto AddFunction = [this, &ExptectedReturnType, &ExpectedArgumentType, WidgetBlueprint, &ConversionFunctions](const UFunction* Function)
+		{
+			if (IsValidConversionFunction(WidgetBlueprint, UE::MVVM::FConversionFunctionValue(Function), ExpectedArgumentType, ExptectedReturnType))
+			{
+				ConversionFunctions.Add(UE::MVVM::FConversionFunctionValue(Function));
+			}
+		};
+
 		const FBlueprintActionDatabase::FActionRegistry& ActionRegistry = FBlueprintActionDatabase::Get().GetAllActions();
 		for (auto It = ActionRegistry.CreateConstIterator(); It; ++It)
 		{
@@ -946,59 +1028,18 @@ TArray<UFunction*> UMVVMEditorSubsystem::GetAvailableConversionFunctions(const U
 				}
 			}
 		}
+		return ConversionFunctions;
 	}
 	else if (FilterType == EMVVMDeveloperConversionFunctionFilterType::AllowedList)
 	{
-		auto IsInheritedBlueprintFunction = [](const UFunction* Function)
-		{
-			bool bIsBpInheritedFunc = false;
-			if (UClass* FuncClass = Function->GetOwnerClass())
-			{
-				if (UBlueprint* BpOwner = Cast<UBlueprint>(FuncClass->ClassGeneratedBy))
-				{
-					FName FuncName = Function->GetFName();
-					if (UClass* ParentClass = BpOwner->ParentClass)
-					{
-						bIsBpInheritedFunc = (ParentClass->FindFunctionByName(FuncName, EIncludeSuperFlag::IncludeSuper) != nullptr);
-					}
-				}
-			}
-			return bIsBpInheritedFunc;
-		};
-
-		TArray<const UClass*> Classes = GetDefault<UMVVMDeveloperProjectSettings>()->GetAllowedConversionFunctionClasses();
-		Classes.Add(WidgetBlueprint->GeneratedClass);
-		for (const UClass* Class : Classes)
+		if (!ConversionFunctionCollection.IsValid())
 		{
-			if (Class->IsChildOf(UK2Node::StaticClass()))
-			{
-
-			}
-			else
-			{
-				for (TFieldIterator<UFunction> FunctionIt(Class, EFieldIteratorFlags::ExcludeSuper); FunctionIt; ++FunctionIt)
-				{
-					UFunction* Function = *FunctionIt;
-					if (IsInheritedBlueprintFunction(Function)) // remove override virtual function
-					{
-						continue;
-					}
-					if (!Function->HasAllFunctionFlags(FUNC_BlueprintCallable))
-					{
-						continue;
-					}
-
-					AddFunction(Function);
-				}
-			}
+			ConversionFunctionCollection = MakeUnique<UE::MVVM::ConversionFunctionLibrary::FCollection>();
 		}
+		return ConversionFunctionCollection->GetFunctions(WidgetBlueprint, ExpectedArgumentType, ExptectedReturnType);
 	}
 
-	ConversionFunctions.Sort([](const UFunction& A, const UFunction& B) -> bool
-		{
-			return A.GetFName().LexicalLess(B.GetFName());
-		});
-	return ConversionFunctions;
+	return TArray<UE::MVVM::FConversionFunctionValue>();
 }
 
 FMVVMBlueprintPropertyPath UMVVMEditorSubsystem::GetPathForConversionFunctionArgument(const UWidgetBlueprint* WidgetBlueprint, const FMVVMBlueprintViewBinding& Binding,  const FMVVMBlueprintPinId& ParameterId, bool bSourceToDestination) const
diff --git a/Engine/Plugins/Runtime/ModelViewViewModel/Source/ModelViewViewModelEditor/Private/Types/MVVMConversionFunctionValue.cpp b/Engine/Plugins/Runtime/ModelViewViewModel/Source/ModelViewViewModelEditor/Private/Types/MVVMConversionFunctionValue.cpp
new file mode 100644
--- /dev/null
+++ b/Engine/Plugins/Runtime/ModelViewViewModel/Source/ModelViewViewModelEditor/Private/Types/MVVMConversionFunctionValue.cpp
@@ -0,0 +1,74 @@
+// Copyright Epic Games, Inc. All Rights Reserved.
+
+#include "Types/MVVMConversionFunctionValue.h"
+#include "EdGraphSchema_K2.h"
+
+namespace UE::MVVM
+{
+
+FString FConversionFunctionValue::GetName() const
+{
+	return ConversionFunction ? ConversionFunction->GetName() : ConversionNode ? ConversionNode->GetName() : FString();
+}
+
+FName FConversionFunctionValue::GetFName() const
+{
+	return ConversionFunction ? ConversionFunction->GetFName() : ConversionNode ? ConversionNode->GetFName() : FName();
+}
+
+FString FConversionFunctionValue::GetFullGroupName(bool bStartWithOuter) const
+{
+	return ConversionFunction ? ConversionFunction->GetFullGroupName(bStartWithOuter) : ConversionNode ? ConversionNode->GetFullGroupName(bStartWithOuter) : FString();
+}
+
+FText FConversionFunctionValue::GetDisplayName() const
+{
+	return ConversionFunction ? ConversionFunction->GetDisplayNameText() : ConversionNode ? ConversionNode.GetDefaultObject()->GetNodeTitle(ENodeTitleType::MenuTitle) : FText();
+}
+
+FText FConversionFunctionValue::GetTooltip() const
+{
+	return ConversionFunction ? ConversionFunction->GetToolTipText() : ConversionNode ? ConversionNode->GetToolTipText() : FText();
+}
+
+FText FConversionFunctionValue::GetCategory() const
+{
+	static const FName NAME_Category = "Category";
+	return ConversionFunction ? ConversionFunction->GetMetaDataText(NAME_Category) : ConversionNode ? ConversionNode->GetMetaDataText(NAME_Category) : FText();
+}
+
+namespace Private
+{
+	template<typename T>
+	void AddResult(TArray<FString>& Result, T* Value)
+	{
+		Result.Add(Value->GetName());
+		const FString& DisplayName = Value->GetMetaData(FBlueprintMetadata::MD_DisplayName);
+		if (DisplayName.Len() > 0)
+		{
+			Result.Add(DisplayName);
+		}
+
+		FString MetadataKeywords = Value->GetMetaDataText(FBlueprintMetadata::MD_FunctionKeywords, TEXT("UObjectKeywords"), Value->GetFullGroupName(false)).ToString();
+		if (MetadataKeywords.Len() > 0)
+		{
+			Result.Add(MoveTemp(MetadataKeywords));
+		}
+	}
+}
+
+TArray<FString> FConversionFunctionValue::GetSearchKeywords() const
+{
+	TArray<FString> Result;
+	if (ConversionFunction)
+	{
+		Private::AddResult(Result, ConversionFunction);
+	}
+	else if (ConversionNode.Get())
+	{
+		Private::AddResult(Result, ConversionNode.Get());
+	}
+	return Result;
+}
+
+} // namespace UE::MVVM
\ No newline at end of file
diff --git a/Engine/Plugins/Runtime/ModelViewViewModel/Source/ModelViewViewModelEditor/Private/Types/MVVMPinLinkedValue.cpp b/Engine/Plugins/Runtime/ModelViewViewModel/Source/ModelViewViewModelEditor/Private/Types/MVVMPinLinkedValue.cpp
--- a/Engine/Plugins/Runtime/ModelViewViewModel/Source/ModelViewViewModelEditor/Private/Types/MVVMPinLinkedValue.cpp
+++ b/Engine/Plugins/Runtime/ModelViewViewModel/Source/ModelViewViewModelEditor/Private/Types/MVVMPinLinkedValue.cpp
@@ -20,6 +20,11 @@ FMVVMLinkedPinValue::FMVVMLinkedPinValue(const UBlueprint* InBlueprint, FMVVMBlu
 	}
 }
 
+FMVVMLinkedPinValue::FMVVMLinkedPinValue(UE::MVVM::FConversionFunctionValue Function)
+	: ConversionFunction(Function.GetFunction())
+	, ConversionNode(Function.GetNode())
+{}
+
 FMVVMLinkedPinValue::FMVVMLinkedPinValue(const UFunction* Function)
 	: ConversionFunction(Function)
 {}
diff --git a/Engine/Plugins/Runtime/ModelViewViewModel/Source/ModelViewViewModelEditor/Private/Widgets/SMVVMFieldSelectorMenu.cpp b/Engine/Plugins/Runtime/ModelViewViewModel/Source/ModelViewViewModelEditor/Private/Widgets/SMVVMFieldSelectorMenu.cpp
--- a/Engine/Plugins/Runtime/ModelViewViewModel/Source/ModelViewViewModelEditor/Private/Widgets/SMVVMFieldSelectorMenu.cpp
+++ b/Engine/Plugins/Runtime/ModelViewViewModel/Source/ModelViewViewModelEditor/Private/Widgets/SMVVMFieldSelectorMenu.cpp
@@ -241,7 +241,7 @@ FReply SFieldSelectorMenu::HandleSelectClicked()
 		FMVVMLinkedPinValue Value;
 		if (ConversionFunctionList.IsValid())
 		{
-			TArray<const UFunction*> Selection;
+			TArray<FConversionFunctionValue> Selection;
 			if (ConversionFunctionList.IsValid())
 			{
 				Selection = ConversionFunctionList->GetSelectedItems();
@@ -302,7 +302,7 @@ void SFieldSelectorMenu::SetPropertyPathSelection(const FMVVMBlueprintPropertyPa
 	OnSelectionChanged.ExecuteIfBound(FMVVMLinkedPinValue(SelectedPath));
 }
 
-void SFieldSelectorMenu::SetConversionFunctionSelection(const UFunction* SelectedFunction)
+void SFieldSelectorMenu::SetConversionFunctionSelection(const FConversionFunctionValue SelectedFunction)
 {
 	OnSelectionChanged.ExecuteIfBound(FMVVMLinkedPinValue(SelectedFunction));
 }
@@ -366,7 +366,7 @@ void SFieldSelectorMenu::HandleSearchBoxTextChanged(const FText& NewText)
 
 	if (ConversionFunctionList.IsValid())
 	{
-		TArray<const UFunction*> OldSelectedFunctions;
+		TArray<FConversionFunctionValue> OldSelectedFunctions;
 		ConversionFunctionList->GetSelectedItems(OldSelectedFunctions);
 
 		FilterConversionFunctions();
@@ -583,9 +583,9 @@ int32 SFieldSelectorMenu::SortConversionFunctionItemsRecursive(TArray<TSharedPtr
 			{
 				return false;
 			}
-			if (A->Function != nullptr && B->Function != nullptr)
+			if (A->Function.IsValid() && B->Function.IsValid())
 			{
-				return A->Function->GetDisplayNameText().CompareTo(B->Function->GetDisplayNameText()) <= 0;
+				return A->Function.GetDisplayName().CompareTo(B->Function.GetDisplayName()) <= 0;
 			}
 			return true;
 		});
@@ -605,40 +605,17 @@ int32 SFieldSelectorMenu::SortConversionFunctionItemsRecursive(TArray<TSharedPtr
 
 void SFieldSelectorMenu::GenerateConversionFunctionItems()
 {
+	const FProperty* AssignToProperty = SelectionContext.AssignableTo;
 	UMVVMEditorSubsystem* Subsystem = GEditor->GetEditorSubsystem<UMVVMEditorSubsystem>();
-	TArray<UFunction*> AllConversionFunctions = Subsystem->GetAvailableConversionFunctions(WidgetBlueprint.Get(), FMVVMBlueprintPropertyPath(), FMVVMBlueprintPropertyPath());
+	TArray<FConversionFunctionValue> AllConversionFunctions = Subsystem->GetConversionFunctions(WidgetBlueprint.Get(), nullptr, AssignToProperty);
 
-	// remove all incompatible conversion functions
-	for (int32 Index = AllConversionFunctions.Num() - 1; Index >= 0; --Index)
-	{
-		const UFunction* Function = AllConversionFunctions[Index];
-		const FProperty* ReturnProperty = BindingHelper::GetReturnProperty(Function);
-		const FProperty* AssignToProperty = SelectionContext.AssignableTo;
-
-		if (AssignToProperty != nullptr && !BindingHelper::ArePropertiesCompatible(ReturnProperty, AssignToProperty))
-		{
-			AllConversionFunctions.RemoveAtSwap(Index);
-		}
-	}
-
-	auto AddFunctionToItem = [](const UFunction* Function, const TSharedPtr<FConversionFunctionItem>& Parent)
+	auto AddFunctionToItem = [](FConversionFunctionValue Function, const TSharedPtr<FConversionFunctionItem>& Parent)
 	{
 		TSharedPtr<FConversionFunctionItem>& Item = Parent->Children.Add_GetRef(MakeShared<FConversionFunctionItem>());
 		Item->Function = Function;
+		Item->SearchKeywords = Function.GetSearchKeywords();
 		Item->NumFunctions = 1;
 		Parent->NumFunctions += 1;
-
-		Item->SearchKeywords.Add(Function->GetName());
-		const FString& DisplayName = Function->GetMetaData(FBlueprintMetadata::MD_DisplayName);
-		if (DisplayName.Len() > 0)
-		{
-			Item->SearchKeywords.Add(DisplayName);
-		}
-		FString MetadataKeywords = Function->GetMetaDataText(FBlueprintMetadata::MD_FunctionKeywords, TEXT("UObjectKeywords"), Function->GetFullGroupName(false)).ToString();
-		if (MetadataKeywords.Len() > 0)
-		{
-			Item->SearchKeywords.Add(MoveTemp(MetadataKeywords));
-		}
 	};
 
 	TSharedPtr<FConversionFunctionItem> CurrentSelectedItem;
@@ -649,10 +626,9 @@ void SFieldSelectorMenu::GenerateConversionFunctionItems()
 	TSharedPtr<FConversionFunctionItem>& RootItem = ConversionFunctionRoot.Add_GetRef(MakeShared<FConversionFunctionItem>());
 	RootItem->CategoryPath = { TEXT("Conversion Functions") };
 
-	FName NAME_Category = "Category";
-	for (const UFunction* Function : AllConversionFunctions)
+	for (const FConversionFunctionValue& Function : AllConversionFunctions)
 	{
-		const FText& CategoryName = Function->GetMetaDataText(NAME_Category);
+		const FText& CategoryName = Function.GetCategory();
 		if (CategoryName.IsEmpty())
 		{
 			AddFunctionToItem(Function, RootItem);
@@ -688,11 +664,11 @@ void SFieldSelectorMenu::GenerateConversionFunctionItems()
 	ensure(NumItems == RootItem->NumFunctions);
 }
 
-TSharedPtr<SFieldSelectorMenu::FConversionFunctionItem> SFieldSelectorMenu::ExpandFunctionCategoryTreeToItem(const UFunction* Function)
+TSharedPtr<SFieldSelectorMenu::FConversionFunctionItem> SFieldSelectorMenu::ExpandFunctionCategoryTreeToItem(const FConversionFunctionValue Function)
 {
 	TArray<TSharedPtr<FConversionFunctionItem>> Path;
 
-	FText FullCategoryName = Function->GetMetaDataText("Category");
+	FText FullCategoryName = Function.GetCategory();
 	if (FullCategoryName.IsEmpty())
 	{
 		Path.Add(FilteredConversionFunctionRoot[0]);
@@ -889,10 +865,12 @@ TSharedRef<SWidget> SFieldSelectorMenu::CreateBindingContextPanel(const FArgumen
 			.OnSelectionChanged(this, &SFieldSelectorMenu::HandleConversionFunctionCategorySelected)
 			.OnGetChildren(this, &SFieldSelectorMenu::HandleGetConversionFunctionCategoryChildren);
 
-		if (CurrentFieldSelectedLinkedValue.IsConversionFunction())
+		if (CurrentFieldSelectedLinkedValue.IsConversionFunction() || CurrentFieldSelectedLinkedValue.IsConversionNode())
 		{
-			const UFunction* ConversionFunction = CurrentFieldSelectedLinkedValue.GetConversionFunction();
-			if (ConversionFunction)
+			const FConversionFunctionValue ConversionFunction = CurrentFieldSelectedLinkedValue.IsConversionFunction()
+				? FConversionFunctionValue(CurrentFieldSelectedLinkedValue.GetConversionFunction())
+				: FConversionFunctionValue(CurrentFieldSelectedLinkedValue.GetConversionNode());
+			if (ConversionFunction.IsValid())
 			{
 				TSharedPtr<FConversionFunctionItem> FunctionItem = ExpandFunctionCategoryTreeToItem(ConversionFunction);
 				if (FunctionItem)
@@ -902,10 +880,6 @@ TSharedRef<SWidget> SFieldSelectorMenu::CreateBindingContextPanel(const FArgumen
 				}
 			}
 		}
-		else if (CurrentFieldSelectedLinkedValue.IsConversionNode())
-		{
-			check(false); // not supported yet
-		}
 
 		StackedSourcePicker->AddSlot()
 			.AutoHeight()
@@ -958,7 +932,7 @@ TSharedRef<SWidget> SFieldSelectorMenu::CreateBindingListPanel(const FArguments&
 		BindingListVBox->AddSlot()
 			.AutoHeight()
 			[
-				SAssignNew(ConversionFunctionList, SListView<const UFunction*>)
+				SAssignNew(ConversionFunctionList, SListView<FConversionFunctionValue>)
 				.SelectionMode(ESelectionMode::Single)
 				.ListItemsSource(&FilteredConversionFunctions)
 				.OnMouseButtonDoubleClick(this, &SFieldSelectorMenu::SetConversionFunctionSelection)
@@ -967,8 +941,10 @@ TSharedRef<SWidget> SFieldSelectorMenu::CreateBindingListPanel(const FArguments&
 
 		if (bValidSelectedFunction)
 		{
-			check(InArgs._CurrentSelected.GetValue().IsConversionFunction()); // node not implemented yet
-			ConversionFunctionList->SetItemSelection(InArgs._CurrentSelected.GetValue().GetConversionFunction(), true);
+			FConversionFunctionValue ConversionFunctionValue = InArgs._CurrentSelected.GetValue().IsConversionFunction()
+				? FConversionFunctionValue(InArgs._CurrentSelected.GetValue().GetConversionFunction())
+				: FConversionFunctionValue(InArgs._CurrentSelected.GetValue().GetConversionNode());
+			ConversionFunctionList->SetItemSelection(ConversionFunctionValue, true);
 		}
 	}
 
@@ -1073,9 +1049,9 @@ void SFieldSelectorMenu::HandleGetConversionFunctionCategoryChildren(TSharedPtr<
 		});
 }
 
-TSharedRef<ITableRow> SFieldSelectorMenu::HandleGenerateConversionFunctionRow(const UFunction* Function, const TSharedRef<STableViewBase>& OwnerTable)
+TSharedRef<ITableRow> SFieldSelectorMenu::HandleGenerateConversionFunctionRow(const FConversionFunctionValue Function, const TSharedRef<STableViewBase>& OwnerTable)
 {
-	return SNew(STableRow<const UFunction*>, OwnerTable)
+	return SNew(STableRow<FConversionFunctionValue>, OwnerTable)
 		[
 			SNew(SHorizontalBox)
 			+ SHorizontalBox::Slot()
@@ -1093,23 +1069,23 @@ TSharedRef<ITableRow> SFieldSelectorMenu::HandleGenerateConversionFunctionRow(co
 			.VAlign(VAlign_Center)
 			[
 				SNew(STextBlock)
-				.Text(Function->GetDisplayNameText())
-				.ToolTipText(Function->GetToolTipText())
+				.Text(Function.GetDisplayName())
+				.ToolTipText(Function.GetTooltip())
 				.HighlightText_Lambda([this]() { return SearchBox.IsValid() ? SearchBox->GetText() : FText::GetEmpty(); })
 			]
 		];
 }
 
-void SFieldSelectorMenu::AddConversionFunctionChildrenRecursive(const TSharedPtr<FConversionFunctionItem>& Parent, TArray<const UFunction*>& OutFunctions)
+void SFieldSelectorMenu::AddConversionFunctionChildrenRecursive(const TSharedPtr<FConversionFunctionItem>& Parent, TArray<FConversionFunctionValue>& OutFunctions)
 {
 	for (const TSharedPtr<FConversionFunctionItem>& Item : Parent->Children)
 	{
-		if (Item->Function != nullptr)
+		if (Item->Function.IsValid())
 		{
 			int32 Index = 0;
 			for (; Index < OutFunctions.Num(); ++Index)
 			{
-				if (OutFunctions[Index]->GetFName().Compare(Item->Function->GetFName()) > 0)
+				if (OutFunctions[Index].GetFName().Compare(Item->Function.GetFName()) > 0)
 				{
 					break;
 				}
@@ -1139,11 +1115,15 @@ void SFieldSelectorMenu::FilterConversionFunctions()
 	}
 
 	FilteredConversionFunctions.Reset();
-	for (const UFunction* Function : ConversionFunctions)
+	for (const FConversionFunctionValue& Function : ConversionFunctions)
 	{
-		FString FunctionName = Function->GetName();
-		const FString& DisplayName = Function->GetMetaData(FBlueprintMetadata::MD_DisplayName);
-		FString MetadataKeywords = Function->GetMetaDataText(FBlueprintMetadata::MD_FunctionKeywords, TEXT("UObjectKeywords"), Function->GetFullGroupName(false)).ToString();
+		FString FunctionName = Function.GetName();
+		FString DisplayName = Function.GetDisplayName().ToString();
+		FString MetadataKeywords;
+		if (Function.IsFunction())
+		{
+			MetadataKeywords = Function.GetFunction()->GetMetaDataText(FBlueprintMetadata::MD_FunctionKeywords, TEXT("UObjectKeywords"), Function.GetFullGroupName(false)).ToString();
+		}
 
 		bool bMatches = true;
 		for (const FString& Filter : FilterStrings)
@@ -1195,7 +1175,7 @@ void SFieldSelectorMenu::HandleConversionFunctionCategorySelected(TSharedPtr<FCo
 			AddConversionFunctionChildrenRecursive(Item, ConversionFunctions);
 		}
 
-		Algo::SortBy(ConversionFunctions, &UFunction::GetDisplayNameText, FText::FSortPredicate());
+		Algo::SortBy(ConversionFunctions, &FConversionFunctionValue::GetDisplayName, FText::FSortPredicate());
 
 		FilterConversionFunctions();
 	}
diff --git a/Engine/Plugins/Runtime/ModelViewViewModel/Source/ModelViewViewModelEditor/Private/Widgets/SMVVMFieldSelectorMenu.h b/Engine/Plugins/Runtime/ModelViewViewModel/Source/ModelViewViewModelEditor/Private/Widgets/SMVVMFieldSelectorMenu.h
--- a/Engine/Plugins/Runtime/ModelViewViewModel/Source/ModelViewViewModelEditor/Private/Widgets/SMVVMFieldSelectorMenu.h
+++ b/Engine/Plugins/Runtime/ModelViewViewModel/Source/ModelViewViewModelEditor/Private/Widgets/SMVVMFieldSelectorMenu.h
@@ -21,6 +21,57 @@ namespace UE::MVVM {class SSourceBindingList; }
 
 class UWidgetBlueprint;
 
+template<>
+struct TListTypeTraits<UE::MVVM::FConversionFunctionValue>
+{
+public:
+	using NullableType = UE::MVVM::FConversionFunctionValue;
+	using MapKeyFuncs = TDefaultMapHashableKeyFuncs<UE::MVVM::FConversionFunctionValue, TSharedRef<ITableRow>, false>;
+	using MapKeyFuncsSparse = TDefaultMapHashableKeyFuncs<UE::MVVM::FConversionFunctionValue, FSparseItemInfo, false>;
+	using SetKeyFuncs = DefaultKeyFuncs<UE::MVVM::FConversionFunctionValue>;
+
+	template<typename U>
+	static void AddReferencedObjects(FReferenceCollector&, TArray<UE::MVVM::FConversionFunctionValue>&, TSet<UE::MVVM::FConversionFunctionValue>&, TMap< const U*, UE::MVVM::FConversionFunctionValue>&)
+	{
+	}
+
+	static bool IsPtrValid(UE::MVVM::FConversionFunctionValue InPtr)
+	{
+		return InPtr.IsValid();
+	}
+
+	static void ResetPtr(UE::MVVM::FConversionFunctionValue& InPtr)
+	{
+		InPtr = UE::MVVM::FConversionFunctionValue();
+	}
+
+	static  UE::MVVM::FConversionFunctionValue MakeNullPtr()
+	{
+		return UE::MVVM::FConversionFunctionValue();
+	}
+
+	static  UE::MVVM::FConversionFunctionValue NullableItemTypeConvertToItemType(UE::MVVM::FConversionFunctionValue InPtr)
+	{
+		return InPtr;
+	}
+
+	static FString DebugDump( UE::MVVM::FConversionFunctionValue InPtr)
+	{
+		return InPtr.GetName();
+	}
+
+	class SerializerType {};
+};
+
+template <>
+struct TIsValidListItem<UE::MVVM::FConversionFunctionValue>
+{
+	enum
+	{
+		Value = true
+	};
+};
+
 namespace UE::MVVM
 {
 
@@ -62,14 +113,14 @@ private:
 
 		TArray<FString> CategoryPath;
 		TArray<FString> SearchKeywords;
-		const UFunction* Function = nullptr;
+		UE::MVVM::FConversionFunctionValue Function;
 		TArray<TSharedPtr<FConversionFunctionItem>> Children;
 		int32 NumFunctions = 0;
 	};
 
 private:
 	void SetPropertyPathSelection(const FMVVMBlueprintPropertyPath& SelectedPath);
-	void SetConversionFunctionSelection(const UFunction* SelectedFunction);
+	void SetConversionFunctionSelection(const UE::MVVM::FConversionFunctionValue SelectedFunction);
 
 	TSharedRef<SWidget> CreateBindingContextPanel(const FArguments& InArgs);
 	TSharedRef<SWidget> CreateBindingListPanel(const FArguments& InArgs, const FProperty* AssignableToProperty);
@@ -85,15 +136,15 @@ private:
 	void FilterConversionFunctions();
 	/** Recursively filter the items in SourceArray and place them into DestArray. Returns true if any items were added. */
 	int32 FilterConversionFunctionCategoryChildren(const TArray<FString>& FilterStrings, const TArray<TSharedPtr<FConversionFunctionItem>>& SourceArray, TArray<TSharedPtr<FConversionFunctionItem>>& OutDestArray);
-	void AddConversionFunctionChildrenRecursive(const TSharedPtr<FConversionFunctionItem>& Parent, TArray<const UFunction*>& OutFunctions);
+	void AddConversionFunctionChildrenRecursive(const TSharedPtr<FConversionFunctionItem>& Parent, TArray<UE::MVVM::FConversionFunctionValue>& OutFunctions);
 	TSharedPtr<FConversionFunctionItem> FindOrCreateItemForCategory(TArray<TSharedPtr<FConversionFunctionItem>>& Items, TArrayView<FString> CategoryPath);
 	TSharedPtr<FConversionFunctionItem> FindConversionFunctionCategory(const TArray<TSharedPtr<FConversionFunctionItem>>& Items, TArrayView<FString> CategoryNameParts) const;
 	void HandleGetConversionFunctionCategoryChildren(TSharedPtr<FConversionFunctionItem> Item, TArray<TSharedPtr<FConversionFunctionItem>>& OutItems) const;
 	void HandleConversionFunctionCategorySelected(TSharedPtr<FConversionFunctionItem> Item, ESelectInfo::Type);
 	TSharedRef<ITableRow> HandleGenerateConversionFunctionCategoryRow(TSharedPtr<FConversionFunctionItem> Item, const TSharedRef<STableViewBase>& OwnerTable);
-	TSharedRef<ITableRow> HandleGenerateConversionFunctionRow(const UFunction* Function, const TSharedRef<STableViewBase>& OwnerTable);
+	TSharedRef<ITableRow> HandleGenerateConversionFunctionRow(const UE::MVVM::FConversionFunctionValue Function, const TSharedRef<STableViewBase>& OwnerTable);
 
-	TSharedPtr<FConversionFunctionItem> ExpandFunctionCategoryTreeToItem(const UFunction* Function);
+	TSharedPtr<FConversionFunctionItem> ExpandFunctionCategoryTreeToItem(const UE::MVVM::FConversionFunctionValue Function);
 	void ExpandFunctionCategoryTree(const TArray<TSharedPtr<FConversionFunctionItem>>& Items, bool bRecursive);
 
 	void FilterViewModels(const FText& NewText);
@@ -132,9 +183,9 @@ private:
 	TArray<TSharedPtr<FConversionFunctionItem>> ConversionFunctionRoot;
 
 	//~ functions (selection panel)
-	TSharedPtr<SListView<const UFunction*>> ConversionFunctionList;
-	TArray<const UFunction*> ConversionFunctions;
-	TArray<const UFunction*> FilteredConversionFunctions;
+	TSharedPtr<SListView<UE::MVVM::FConversionFunctionValue>> ConversionFunctionList;
+	TArray<UE::MVVM::FConversionFunctionValue> ConversionFunctions;
+	TArray<UE::MVVM::FConversionFunctionValue> FilteredConversionFunctions;
 
 	bool bIsMenuInitialized = false;
 	bool bIsClearEnabled = false;
diff --git a/Engine/Plugins/Runtime/ModelViewViewModel/Source/ModelViewViewModelEditor/Public/MVVMEditorSubsystem.h b/Engine/Plugins/Runtime/ModelViewViewModel/Source/ModelViewViewModelEditor/Public/MVVMEditorSubsystem.h
--- a/Engine/Plugins/Runtime/ModelViewViewModel/Source/ModelViewViewModelEditor/Public/MVVMEditorSubsystem.h
+++ b/Engine/Plugins/Runtime/ModelViewViewModel/Source/ModelViewViewModelEditor/Public/MVVMEditorSubsystem.h
@@ -4,6 +4,7 @@
 
 #include "EditorSubsystem.h"
 #include "MVVMBlueprintPin.h"
+#include "Types/MVVMConversionFunctionValue.h"
 #include "UObject/Package.h"
 
 #include "MVVMEditorSubsystem.generated.h"
@@ -13,6 +14,7 @@ class UMVVMBlueprintView;
 enum class EMVVMBindingMode : uint8;
 enum class EMVVMExecutionMode : uint8;
 namespace UE::MVVM { struct FBindingSource; }
+namespace UE::MVVM::ConversionFunctionLibrary { class FCollection; }
 struct FMVVMAvailableBinding;
 struct FMVVMBlueprintFunctionReference;
 struct FMVVMBlueprintPropertyPath;
@@ -20,7 +22,6 @@ struct FMVVMBlueprintViewBinding;
 template <typename T> class TSubclassOf;
 
 class UEdGraph;
-class UK2Node;
 class UK2Node_CallFunction;
 class UMVVMBlueprintViewEvent;
 class UWidgetBlueprint;
@@ -93,10 +94,14 @@ public:
 
 	UFUNCTION(BlueprintCallable, Category = "Viewmodel")
 	bool IsValidConversionFunction(const UWidgetBlueprint* WidgetBlueprint, const UFunction* Function, const FMVVMBlueprintPropertyPath& Source, const FMVVMBlueprintPropertyPath& Destination) const;
+	bool IsValidConversionFunction(const UWidgetBlueprint* WidgetBlueprint, const UFunction* Function, const FProperty* ExpectedArgumentType, const FProperty* ExptectedReturnType) const;
+	bool IsValidConversionFunction(const UWidgetBlueprint* WidgetBlueprint, UE::MVVM::FConversionFunctionValue Function, const FProperty* ExpectedArgumentType, const FProperty* ExptectedReturnType) const;
+
 	bool IsValidConversionNode(const UWidgetBlueprint* WidgetBlueprint, const TSubclassOf<UK2Node> Function, const FMVVMBlueprintPropertyPath& Source, const FMVVMBlueprintPropertyPath& Destination) const;
+	bool IsValidConversionNode(const UWidgetBlueprint* WidgetBlueprint, const TSubclassOf<UK2Node> Function, const FProperty* ExpectedArgumentType, const FProperty* ExptectedReturnType) const;
 
 	UFUNCTION(BlueprintCallable, Category = "Viewmodel")
-	bool IsSimpleConversionFunctionA(const UFunction* Function) const;
+	bool IsSimpleConversionFunction(const UFunction* Function) const;
 
 	UFUNCTION(BlueprintCallable, Category = "Viewmodel")
 	UEdGraph* GetConversionFunctionGraph(const UWidgetBlueprint* WidgetBlueprint, const FMVVMBlueprintViewBinding& Binding, bool bSourceToDestination) const;
@@ -109,9 +114,12 @@ public:
 	UFUNCTION(BlueprintCallable, Category = "Viewmodel")
 	UK2Node_CallFunction* GetConversionFunctionNode(const UWidgetBlueprint* WidgetBlueprint, const FMVVMBlueprintViewBinding& Binding, bool bSourceToDestination) const;
 
+	UE_DEPRECATED(5.5, "GetAvailableConversionFunctions return value changes.")
 	UFUNCTION(BlueprintCallable, Category = "Viewmodel")
 	TArray<UFunction*> GetAvailableConversionFunctions(const UWidgetBlueprint* WidgetBlueprint, const FMVVMBlueprintPropertyPath& Source, const FMVVMBlueprintPropertyPath& Destination) const;
 
+	TArray<UE::MVVM::FConversionFunctionValue> GetConversionFunctions(const UWidgetBlueprint* WidgetBlueprint, const FProperty* ExpectedArgumentType, const FProperty* ExptectedReturnType) const;
+
 	FMVVMBlueprintPropertyPath GetPathForConversionFunctionArgument(const UWidgetBlueprint* WidgetBlueprint, const FMVVMBlueprintViewBinding& Binding, const FMVVMBlueprintPinId& PinId, bool bSourceToDestination) const;
 	void SetPathForConversionFunctionArgument(UWidgetBlueprint* WidgetBlueprint, FMVVMBlueprintViewBinding& Binding, const FMVVMBlueprintPinId& PinId, const FMVVMBlueprintPropertyPath& Path, bool bSourceToDestination) const;
 	UEdGraphPin* GetConversionFunctionArgumentPin(const UWidgetBlueprint* WidgetBlueprint, const FMVVMBlueprintViewBinding& Binding, const FMVVMBlueprintPinId& PinId, bool bSourceToDestination) const;
@@ -137,6 +145,9 @@ public:
 	TArray<UE::MVVM::FBindingSource> GetAllViewModels(const UWidgetBlueprint* WidgetBlueprint) const;
 
 	FGuid GetFirstBindingThatUsesViewModel(const UWidgetBlueprint* WidgetBlueprint, FGuid ViewModelId) const;
+
+private:
+	mutable TUniquePtr<UE::MVVM::ConversionFunctionLibrary::FCollection> ConversionFunctionCollection;
 };
 
 #if UE_ENABLE_INCLUDE_ORDER_DEPRECATED_IN_5_2
diff --git a/Engine/Plugins/Runtime/ModelViewViewModel/Source/ModelViewViewModelEditor/Public/Types/MVVMConversionFunctionValue.h b/Engine/Plugins/Runtime/ModelViewViewModel/Source/ModelViewViewModelEditor/Public/Types/MVVMConversionFunctionValue.h
new file mode 100644
--- /dev/null
+++ b/Engine/Plugins/Runtime/ModelViewViewModel/Source/ModelViewViewModelEditor/Public/Types/MVVMConversionFunctionValue.h
@@ -0,0 +1,87 @@
+// Copyright Epic Games, Inc. All Rights Reserved.
+
+#pragma once
+
+#include "K2Node.h"
+#include "Templates/SubclassOf.h"
+
+class UFunction;
+class UK2Node;
+
+namespace UE::MVVM
+{
+/** */
+struct MODELVIEWVIEWMODELEDITOR_API FConversionFunctionValue
+{
+public:
+
+	FConversionFunctionValue() = default;
+	explicit FConversionFunctionValue(const UFunction* Function)
+		: ConversionFunction(Function)
+	{}
+	explicit FConversionFunctionValue(TSubclassOf<UK2Node> Node)
+		: ConversionNode(Node)
+	{}
+
+public:
+	bool IsFunction() const
+	{
+		return ConversionFunction != nullptr;
+	}
+
+	const UFunction* GetFunction() const
+	{
+		return ConversionFunction;
+	}
+
+	bool IsNode() const
+	{
+		return ConversionNode.Get() != nullptr;
+	}
+
+	TSubclassOf<UK2Node> GetNode() const
+	{
+		return ConversionNode;
+	}
+
+	FString GetName() const;
+	FName GetFName() const;
+	FString GetFullGroupName(bool bStartWithOuter) const;
+	FText GetDisplayName() const;
+	FText GetTooltip() const;
+	FText GetCategory() const;
+	TArray<FString> GetSearchKeywords() const;
+
+public:
+	bool IsValid() const
+	{
+		return ConversionFunction != nullptr || ConversionNode.Get() != nullptr;
+	}
+
+	bool operator== (const FConversionFunctionValue& Other) const
+	{
+		return ConversionNode == Other.ConversionNode
+			&& ConversionFunction == Other.ConversionFunction;
+	}
+	bool operator== (const UFunction* Other) const
+	{
+		return ConversionNode.Get() == nullptr
+			&& ConversionFunction == Other;
+	}
+	bool operator== (const TSubclassOf<UK2Node> Other) const
+	{
+		return ConversionNode == Other
+			&& ConversionFunction == nullptr;
+	}
+
+	friend int32 GetTypeHash(const FConversionFunctionValue& Value)
+	{
+		return HashCombine(GetTypeHash(Value.ConversionNode), GetTypeHash(Value.ConversionFunction));
+	}
+
+private:
+	const UFunction* ConversionFunction = nullptr;
+	TSubclassOf<UK2Node> ConversionNode;
+};
+
+} // namespace UE::MVVM
diff --git a/Engine/Plugins/Runtime/ModelViewViewModel/Source/ModelViewViewModelEditor/Public/Types/MVVMLinkedPinValue.h b/Engine/Plugins/Runtime/ModelViewViewModel/Source/ModelViewViewModelEditor/Public/Types/MVVMLinkedPinValue.h
--- a/Engine/Plugins/Runtime/ModelViewViewModel/Source/ModelViewViewModelEditor/Public/Types/MVVMLinkedPinValue.h
+++ b/Engine/Plugins/Runtime/ModelViewViewModel/Source/ModelViewViewModelEditor/Public/Types/MVVMLinkedPinValue.h
@@ -4,6 +4,7 @@
 
 #include "MVVMBlueprintFunctionReference.h"
 #include "MVVMPropertyPath.h"
+#include "Types/MVVMConversionFunctionValue.h"
 #include "Templates/SubclassOf.h"
 
 #include "MVVMLinkedPinValue.generated.h"
@@ -21,6 +22,7 @@ public:
 	FMVVMLinkedPinValue() = default;
 	explicit FMVVMLinkedPinValue(FMVVMBlueprintPropertyPath InPath);
 	explicit FMVVMLinkedPinValue(const UBlueprint* InBlueprint, FMVVMBlueprintFunctionReference InConversion);
+	explicit FMVVMLinkedPinValue(UE::MVVM::FConversionFunctionValue Function);
 	explicit FMVVMLinkedPinValue(const UFunction* Function);
 	explicit FMVVMLinkedPinValue(TSubclassOf<UK2Node> Node);
 
