MVVM: Add an option to auto update the view's viewmodel when it's creation context is the global viewmodel collection. The viewmodel can be added/remove/modified from the global viewmodel collection at anytime and it will update the view.

diff --git a/Engine/Plugins/Runtime/ModelViewViewModel/Source/ModelViewViewModel/Private/View/MVVMView.cpp b/Engine/Plugins/Runtime/ModelViewViewModel/Source/ModelViewViewModel/Private/View/MVVMView.cpp
--- a/Engine/Plugins/Runtime/ModelViewViewModel/Source/ModelViewViewModel/Private/View/MVVMView.cpp
+++ b/Engine/Plugins/Runtime/ModelViewViewModel/Source/ModelViewViewModel/Private/View/MVVMView.cpp
@@ -9,10 +9,13 @@
 #include "Blueprint/UserWidget.h"
 #include "Debugging/MVVMDebugging.h"
 #include "Engine/Engine.h"
+#include "Engine/GameInstance.h"
+#include "MVVMGameSubsystem.h"
 #include "MVVMMessageLog.h"
 #include "ModelViewViewModelModule.h"
 #include "Templates/ValueOrError.h"
 #include "Types/MVVMFieldContext.h"
+#include "Types/MVVMViewModelCollection.h"
 
 #include UE_INLINE_GENERATED_CPP_BY_NAME(MVVMView)
 
@@ -31,6 +34,21 @@ DECLARE_CYCLE_STAT(TEXT("ExecuteBinding Tick"), STAT_UMG_Viewmodel_ExecuteBindin
 // 
 ///////////////////////////////////////////////////////////////////////
 
+namespace UE::MVVM::Private
+{
+UMVVMViewModelCollectionObject* GetGlobalCollection(UUserWidget* UserWidget)
+{
+	if (const UWorld* World = UserWidget->GetWorld())
+	{
+		if (const UGameInstance* GameInstance = World->GetGameInstance())
+		{
+			return GameInstance->GetSubsystem<UMVVMGameSubsystem>()->GetViewModelCollection();
+		}
+	}
+	return nullptr;
+}
+} // namespace
+
 void UMVVMView::ConstructView(const UMVVMViewClass* InClassExtension)
 {
 	ensure(ClassExtension == nullptr);
@@ -100,6 +118,14 @@ void UMVVMView::InitializeSources()
 	check(Sources.Num() == NumberOfSources);
 	check(NumberOfSources <= 64); // the max number of source the bitfield can hold.
 
+	if (ClassExtension->DoesListenToViewModelCollectionChanged())
+	{
+		if (UMVVMViewModelCollectionObject* Collection = UE::MVVM::Private::GetGlobalCollection(GetUserWidget()))
+		{
+			Collection->OnCollectionChanged().AddUObject(this, &UMVVMView::HandleViewModelCollectionChanged);
+		}
+	}
+
 	for (int32 Index = 0; Index < NumberOfSources; ++Index)
 	{
 		InitializeSource(FMVVMView_SourceKey(Index));
@@ -135,6 +161,14 @@ void UMVVMView::UninitializeSources()
 		UninitializeBindings();
 	}
 
+	if (ClassExtension->DoesListenToViewModelCollectionChanged())
+	{
+		if (UMVVMViewModelCollectionObject* Collection = UE::MVVM::Private::GetGlobalCollection(GetUserWidget()))
+		{
+			Collection->OnCollectionChanged().RemoveAll(this);
+		}
+	}
+
 	bSourcesInitialized = false;
 
 	const int32 NumberOfSources = Sources.Num();
@@ -829,6 +863,28 @@ bool UMVVMView::SetSourceInternal(FMVVMViewClass_SourceKey ClassSourceKey, TScri
 }
 
 
+void UMVVMView::HandleViewModelCollectionChanged()
+{
+	if (ensure(ClassExtension))
+	{
+		const int32 NumberOfSources = ClassExtension->GetSources().Num();
+		for (int32 Index = 0; Index < NumberOfSources; ++Index)
+		{
+			FMVVMView_SourceKey SourceKey(Index);
+			FMVVMView_Source& ViewSource = Sources[SourceKey.GetIndex()];
+			const FMVVMViewClass_Source& ClassSource = ClassExtension->GetSource(ViewSource.ClassKey);
+			if (ViewSource.bSourceInitialized && !ViewSource.bSetManually && ClassSource.RequireGlobalViewModelCollectionUpdate())
+			{
+				TScriptInterface<INotifyFieldValueChanged> NewValue = ClassSource.GetGlobalCollectionViewModel(GetUserWidget());
+				if (NewValue.GetObject() != ViewSource.Source)
+				{
+					SetSourceInternal(ViewSource.ClassKey, NewValue, true);
+				}
+			}
+		}
+	}
+}
+
 void UMVVMView::InitializeEvents()
 {
 	if (bEventsInitialized)
diff --git a/Engine/Plugins/Runtime/ModelViewViewModel/Source/ModelViewViewModel/Private/View/MVVMViewClass.cpp b/Engine/Plugins/Runtime/ModelViewViewModel/Source/ModelViewViewModel/Private/View/MVVMViewClass.cpp
--- a/Engine/Plugins/Runtime/ModelViewViewModel/Source/ModelViewViewModel/Private/View/MVVMViewClass.cpp
+++ b/Engine/Plugins/Runtime/ModelViewViewModel/Source/ModelViewViewModel/Private/View/MVVMViewClass.cpp
@@ -26,6 +26,35 @@
 
 #define LOCTEXT_NAMESPACE "MVVMViewClass"
 
+namespace UE::MVVM::Private
+{
+UMVVMViewModelBase* GetGlobalCollectionViewModel(UUserWidget* InUserWidget, const FMVVMViewModelContext& GlobalViewModelInstance, UMVVMViewModelCollectionObject*& Collection)
+{
+	Collection = nullptr;
+	UMVVMViewModelBase* FoundViewModelInstance = nullptr;
+	if (GlobalViewModelInstance.IsValid())
+	{
+		if (const UWorld* World = InUserWidget->GetWorld())
+		{
+			if (const UGameInstance* GameInstance = World->GetGameInstance())
+			{
+				Collection = GameInstance->GetSubsystem<UMVVMGameSubsystem>()->GetViewModelCollection();
+				if (Collection)
+				{
+					FoundViewModelInstance = Collection->FindViewModelInstance(GlobalViewModelInstance);
+				}
+			}
+		}
+	}
+
+	if (FoundViewModelInstance != nullptr)
+	{
+		ensureMsgf(FoundViewModelInstance->IsA(GlobalViewModelInstance.ContextClass), TEXT("The Global View Model Instance is not of the expected type."));
+	}
+	return FoundViewModelInstance;
+}
+} // namespace
+
 ///////////////////////////////////////////////////////////////////////
 // 
 ///////////////////////////////////////////////////////////////////////
@@ -103,18 +132,7 @@ UObject* FMVVMViewClass_Source::GetOrCreateInstance(const UMVVMViewClass* InView
 	else if (GlobalViewModelInstance.IsValid())
 	{
 		UMVVMViewModelCollectionObject* Collection = nullptr;
-		UMVVMViewModelBase* FoundViewModelInstance = nullptr;
-		if (const UWorld* World = InUserWidget->GetWorld())
-		{
-			if (const UGameInstance* GameInstance = World->GetGameInstance())
-			{
-				Collection = GameInstance->GetSubsystem<UMVVMGameSubsystem>()->GetViewModelCollection();
-				if (Collection)
-				{
-					FoundViewModelInstance = Collection->FindViewModelInstance(GlobalViewModelInstance);
-				}
-			}
-		}
+		UMVVMViewModelBase* FoundViewModelInstance = UE::MVVM::Private::GetGlobalCollectionViewModel(InUserWidget, GlobalViewModelInstance, Collection);
 
 		if (FoundViewModelInstance != nullptr)
 		{
@@ -154,6 +172,13 @@ void FMVVMViewClass_Source::ReleaseInstance(const UObject* ViewModel, const UMVV
 	}
 }
 
+TScriptInterface<INotifyFieldValueChanged> FMVVMViewClass_Source::GetGlobalCollectionViewModel(UUserWidget* InUserWidget) const
+{
+	UMVVMViewModelCollectionObject* Collection = nullptr;
+	UMVVMViewModelBase* FoundViewModelInstance = UE::MVVM::Private::GetGlobalCollectionViewModel(InUserWidget, GlobalViewModelInstance, Collection);
+	return FoundViewModelInstance ? TScriptInterface<INotifyFieldValueChanged>(FoundViewModelInstance) : TScriptInterface<INotifyFieldValueChanged>();
+}
+
 #if UE_WITH_MVVM_DEBUGGING
 FMVVMViewClass_Source::FToStringArgs FMVVMViewClass_Source::FToStringArgs::Short()
 {
diff --git a/Engine/Plugins/Runtime/ModelViewViewModel/Source/ModelViewViewModel/Public/View/MVVMView.h b/Engine/Plugins/Runtime/ModelViewViewModel/Source/ModelViewViewModel/Public/View/MVVMView.h
--- a/Engine/Plugins/Runtime/ModelViewViewModel/Source/ModelViewViewModel/Public/View/MVVMView.h
+++ b/Engine/Plugins/Runtime/ModelViewViewModel/Source/ModelViewViewModel/Public/View/MVVMView.h
@@ -178,6 +178,7 @@ private:
 
 	//~ evaluate source
 	bool EvaluateSource(FMVVMViewClass_SourceKey SourceIndex);
+	void HandleViewModelCollectionChanged();
 
 	//~ events
 	void BindEvent(const FMVVMViewClass_Event& ClassItem, FMVVMViewClass_EventKey KeyForLog);
diff --git a/Engine/Plugins/Runtime/ModelViewViewModel/Source/ModelViewViewModel/Public/View/MVVMViewClass.h b/Engine/Plugins/Runtime/ModelViewViewModel/Source/ModelViewViewModel/Public/View/MVVMViewClass.h
--- a/Engine/Plugins/Runtime/ModelViewViewModel/Source/ModelViewViewModel/Public/View/MVVMViewClass.h
+++ b/Engine/Plugins/Runtime/ModelViewViewModel/Source/ModelViewViewModel/Public/View/MVVMViewClass.h
@@ -239,6 +239,9 @@ public:
 	/** The source is not needed anymore. */
 	MODELVIEWVIEWMODEL_API void ReleaseInstance(const UObject* ViewModel, const UMVVMView* View) const;
 
+	/** Returns the viewmodel inside the global viewmodel collection. */
+	MODELVIEWVIEWMODEL_API TScriptInterface<INotifyFieldValueChanged> GetGlobalCollectionViewModel(UUserWidget* UserWidget) const;
+
 	/** The expected class of the source. */
 	UClass* GetSourceClass() const
 	{
@@ -290,6 +293,12 @@ public:
 		return (Flags & (uint16)EFlags::HasEvaluatedBindings) != 0;
 	}
 	
+	/** The source has at least one evaluate binding. */
+	bool RequireGlobalViewModelCollectionUpdate() const
+	{
+		return (Flags & (uint16)EFlags::GlobalViewModelCollectionUpdate) != 0;
+	}
+	
 	/**
 	 * The source GetOrCreateInstance can fail.
 	 * The view will not warn if a binding can't execute because the source is invalid.
@@ -399,7 +408,7 @@ private:
 		HasTickBindings = 1 << 8,
 		IsViewModel = 1 << 9,
 		IsViewModelInstanceExposed = 1 << 10,
-		GlobalViewModelCollectionRetry = 1 << 11,
+		GlobalViewModelCollectionUpdate = 1 << 11,
 	};
 
 	UPROPERTY(VisibleAnywhere, Category = "View")
@@ -504,6 +513,12 @@ public:
 	{
 		return bInitializeEventsOnConstruct;
 	}
+	
+	/** Should the view listen to modification made in the global viewmodel collection. */
+	[[nodiscard]] bool DoesListenToViewModelCollectionChanged() const
+	{
+		return bListenToViewModelCollectionChanged;
+	}
 
 	/** Get the container of all the bindings. */
 	[[nodiscard]] const FMVVMCompiledBindingLibrary& GetBindingLibrary() const
@@ -616,6 +631,9 @@ private:
 
 	UPROPERTY()
 	bool bInitializeEventsOnConstruct = true;
+	
+	UPROPERTY()
+	bool bListenToViewModelCollectionChanged = false;
 
 #if WITH_EDITORONLY_DATA
 	FDelegateHandle BluerpintCompiledHandle;
diff --git a/Engine/Plugins/Runtime/ModelViewViewModel/Source/ModelViewViewModelBlueprint/Private/MVVMViewBlueprintCompiler.cpp b/Engine/Plugins/Runtime/ModelViewViewModel/Source/ModelViewViewModelBlueprint/Private/MVVMViewBlueprintCompiler.cpp
--- a/Engine/Plugins/Runtime/ModelViewViewModel/Source/ModelViewViewModelBlueprint/Private/MVVMViewBlueprintCompiler.cpp
+++ b/Engine/Plugins/Runtime/ModelViewViewModel/Source/ModelViewViewModelBlueprint/Private/MVVMViewBlueprintCompiler.cpp
@@ -1459,6 +1459,8 @@ bool FMVVMViewBlueprintCompiler::Compile(UWidgetBlueprintGeneratedClass* Class,
 		ViewExtension->bInitializeSourcesOnConstruct = BlueprintView->GetSettings()->bInitializeSourcesOnConstruct;
 		ViewExtension->bInitializeBindingsOnConstruct = ViewExtension->bInitializeSourcesOnConstruct ? BlueprintView->GetSettings()->bInitializeBindingsOnConstruct : false;
 		ViewExtension->bInitializeEventsOnConstruct = BlueprintView->GetSettings()->bInitializeEventsOnConstruct;
+		ViewExtension->bListenToViewModelCollectionChanged = Algo::AnyOf(ViewExtension->Sources, [](const FMVVMViewClass_Source& Source){ return Source.RequireGlobalViewModelCollectionUpdate(); });
+
 	}
 	{
 		ViewExtension->OptionalSources = 0;
@@ -2178,6 +2180,8 @@ void FMVVMViewBlueprintCompiler::CompileSources(const FCompiledBindingLibraryCom
 		bool bIsOptional = SourceCreatorContext.Source->bIsOptional;
 		bool bCreateInstance = false;
 		bool bIsUserWidgetProperty = !SourceCreatorContext.DynamicContext.IsValid();
+		bool bExposeInstanceInEditor = false;
+		bool bGlobalViewModelCollectionUpdate = false;
 
 		if (ViewModelContext.CreationType == EMVVMBlueprintViewModelContextCreationType::Manual)
 		{
@@ -2187,8 +2191,9 @@ void FMVVMViewBlueprintCompiler::CompileSources(const FCompiledBindingLibraryCom
 		else if (ViewModelContext.CreationType == EMVVMBlueprintViewModelContextCreationType::CreateInstance)
 		{
 			bCreateInstance = true;
+			bExposeInstanceInEditor = ViewModelContext.bExposeInstanceInEditor;
 
-			if (ViewModelContext.bExposeInstanceInEditor)
+			if (bExposeInstanceInEditor)
 			{
 				const FMVVMVCompiledFieldPath* CompiledFieldPath = CompileResult.FieldPaths.Find(SourceCreatorContext.ReadPropertyPathHandle);
 				if (CompiledFieldPath == nullptr)
@@ -2237,6 +2242,7 @@ void FMVVMViewBlueprintCompiler::CompileSources(const FCompiledBindingLibraryCom
 			}
 
 			CompiledSourceCreator.GlobalViewModelInstance = MoveTemp(GlobalViewModelInstance);
+			bGlobalViewModelCollectionUpdate = ViewModelContext.bGlobalViewModelCollectionUpdate;
 			ensure(bIsOptional == ViewModelContext.bOptional);
 		}
 		else if (ViewModelContext.CreationType == EMVVMBlueprintViewModelContextCreationType::Resolver)
@@ -2284,9 +2290,10 @@ void FMVVMViewBlueprintCompiler::CompileSources(const FCompiledBindingLibraryCom
 		CompiledSourceCreator.Flags |= bIsUserWidgetProperty ? (uint16)FMVVMViewClass_Source::EFlags::SetUserWidgetProperty : 0;
 		CompiledSourceCreator.Flags |= bIsOptional ? (uint16)FMVVMViewClass_Source::EFlags::IsOptional : 0;
 		CompiledSourceCreator.Flags |= bCanBeSet ? (uint16)FMVVMViewClass_Source::EFlags::CanBeSet : 0;
-		CompiledSourceCreator.Flags |= bCanBeEvaluated ? (uint16)FMVVMViewClass_Source::EFlags::CanBeEvaluated : 0;
+		CompiledSourceCreator.Flags |= bCanBeEvaluated || bGlobalViewModelCollectionUpdate ? (uint16)FMVVMViewClass_Source::EFlags::CanBeEvaluated : 0;
 		CompiledSourceCreator.Flags |= (uint16)FMVVMViewClass_Source::EFlags::IsViewModel;
-		CompiledSourceCreator.Flags |= ViewModelContext.bExposeInstanceInEditor ? (uint16)FMVVMViewClass_Source::EFlags::IsViewModelInstanceExposed : 0;
+		CompiledSourceCreator.Flags |= bExposeInstanceInEditor ? (uint16)FMVVMViewClass_Source::EFlags::IsViewModelInstanceExposed : 0;
+		CompiledSourceCreator.Flags |= bGlobalViewModelCollectionUpdate ? (uint16)FMVVMViewClass_Source::EFlags::GlobalViewModelCollectionUpdate : 0;
 
 		{
 			FSortData SortData;
diff --git a/Engine/Plugins/Runtime/ModelViewViewModel/Source/ModelViewViewModelBlueprint/Public/MVVMBlueprintViewModelContext.h b/Engine/Plugins/Runtime/ModelViewViewModel/Source/ModelViewViewModelBlueprint/Public/MVVMBlueprintViewModelContext.h
--- a/Engine/Plugins/Runtime/ModelViewViewModel/Source/ModelViewViewModelBlueprint/Public/MVVMBlueprintViewModelContext.h
+++ b/Engine/Plugins/Runtime/ModelViewViewModel/Source/ModelViewViewModelBlueprint/Public/MVVMBlueprintViewModelContext.h
@@ -159,6 +159,10 @@ public:
 	/** Expose the viewmodel instance on every instance of the user widget for modification in editor. */
 	UPROPERTY(EditAnywhere, Category = "Viewmodel", AdvancedDisplay)
 	bool bExposeInstanceInEditor = false;
+	
+	/** Auto update the instance when the viewmodel is added/removed/modifed from the global viewmodel collection. */
+	UPROPERTY(EditAnywhere, Category = "Viewmodel", AdvancedDisplay)
+	bool bGlobalViewModelCollectionUpdate = false;
 
 	/** Can change the name in the editor. */
 	UPROPERTY()
diff --git a/Engine/Plugins/Runtime/ModelViewViewModel/Source/ModelViewViewModelEditor/Private/Customizations/MVVMBlueprintViewModelContextCustomization.cpp b/Engine/Plugins/Runtime/ModelViewViewModel/Source/ModelViewViewModelEditor/Private/Customizations/MVVMBlueprintViewModelContextCustomization.cpp
--- a/Engine/Plugins/Runtime/ModelViewViewModel/Source/ModelViewViewModelEditor/Private/Customizations/MVVMBlueprintViewModelContextCustomization.cpp
+++ b/Engine/Plugins/Runtime/ModelViewViewModel/Source/ModelViewViewModelEditor/Private/Customizations/MVVMBlueprintViewModelContextCustomization.cpp
@@ -405,6 +405,18 @@ void FBlueprintViewModelContextDetailCustomization::CustomizeChildren(TSharedRef
 					}));
 		}
 
+		TSharedPtr<IPropertyHandle> GlobalViewModelCollectionUpdateHandle = PropertyHandle->GetChildHandle(GET_MEMBER_NAME_CHECKED(FMVVMBlueprintViewModelContext, bGlobalViewModelCollectionUpdate), false);
+		if (ensure(GlobalViewModelCollectionUpdateHandle))
+		{
+			ChildBuilder.AddProperty(GlobalViewModelCollectionUpdateHandle.ToSharedRef())
+				.IsEnabled(bCanEdit)
+				.Visibility(MakeAttributeLambda([ContextPtr]()
+					{
+						bool bResult = ContextPtr->CreationType == EMVVMBlueprintViewModelContextCreationType::GlobalViewModelCollection;
+						return bResult ? EVisibility::Visible : EVisibility::Collapsed;
+					}));
+		}
+
 		TSharedPtr<IPropertyHandle> OptionalHandle = PropertyHandle->GetChildHandle(GET_MEMBER_NAME_CHECKED(FMVVMBlueprintViewModelContext, bOptional), false);
 		if (ensure(OptionalHandle))
 		{
