MVVM: Add the source needed to run the event. When running the event, early exist if the sources are not valid. Some sources are optional and should not trigger an error. Binding do the same but it's the view that decides when to trigger the binding. Since the event can be trigger from anywhere the event itself has to check.

diff --git a/Engine/Plugins/Runtime/ModelViewViewModel/Source/ModelViewViewModel/Private/View/MVVMView.cpp b/Engine/Plugins/Runtime/ModelViewViewModel/Source/ModelViewViewModel/Private/View/MVVMView.cpp
--- a/Engine/Plugins/Runtime/ModelViewViewModel/Source/ModelViewViewModel/Private/View/MVVMView.cpp
+++ b/Engine/Plugins/Runtime/ModelViewViewModel/Source/ModelViewViewModel/Private/View/MVVMView.cpp
@@ -1056,4 +1056,36 @@ void UMVVMView::ReinitializeEvents(FMVVMViewClass_SourceKey SourceKey, UObject*
 	}
 }
 
+
+bool UMVVMView::AreSourcesValidForEvent(int32 EventKeyIndex) const
+{
+	FMVVMViewClass_EventKey EventKey = FMVVMViewClass_EventKey(EventKeyIndex);
+	if (ClassExtension)
+	{
+		const FMVVMViewClass_Event& ClassEvent = ClassExtension->GetEvent(EventKey);
+		uint64 EventSources = ClassEvent.GetSources();
+		if ((EventSources & ValidSources) == EventSources)
+		{
+			return true;
+		}
+		
+		const uint64 MissingSources = EventSources & (~ValidSources);
+		if ((MissingSources & ClassExtension->GetOptionalSources()) != MissingSources)
+		{
+#if UE_WITH_MVVM_DEBUGGING
+			UE::MVVM::FMessageLog Log(GetUserWidget());
+			Log.Error(FText::Format(LOCTEXT("ExecuteEventFailInvalidSource", "The event '{0}' was not executed. There are invalid sources.")
+				, FText::FromString(ClassEvent.ToString(ClassExtension, FMVVMViewClass_Event::FToStringArgs::Short()))
+			));
+#else
+			UE::MVVM::FMessageLog Log(GetUserWidget());
+			Log.Error(FText::Format(LOCTEXT("ExecuteEventFailInvalidSource", "The event '{0}' was not executed. There are invalid sources.")
+				, FText::AsNumber(EventKey.GetIndex())
+			));
+#endif
+		}
+	}
+	return false;
+}
+
 #undef LOCTEXT_NAMESPACE
diff --git a/Engine/Plugins/Runtime/ModelViewViewModel/Source/ModelViewViewModel/Public/View/MVVMView.h b/Engine/Plugins/Runtime/ModelViewViewModel/Source/ModelViewViewModel/Public/View/MVVMView.h
--- a/Engine/Plugins/Runtime/ModelViewViewModel/Source/ModelViewViewModel/Public/View/MVVMView.h
+++ b/Engine/Plugins/Runtime/ModelViewViewModel/Source/ModelViewViewModel/Public/View/MVVMView.h
@@ -159,9 +159,18 @@ public:
 	 * The viewmodel needs to be settable and it should have a valid name.
 	 * If the view is initialized, all bindings that uses that viewmodel will be re-executed with the new viewmodel instance.
 	 */
-	UFUNCTION(BlueprintCallable, Category = "Viewmodel")
+	UFUNCTION(BlueprintCallable, Category = "View")
 	bool SetViewModelByClass(TScriptInterface<INotifyFieldValueChanged> NewValue);
 
+private:
+	/**
+	 * Set the viewmodel of the specified name.
+	 * The viewmodel needs to be settable and the type should match (child of the defined viewmodel).
+	 * If the view is initialized, all bindings that uses that viewmodel will be re-executed with the new viewmodel instance.
+	 */
+	UFUNCTION(BlueprintCallable, Category = "View", meta=(BlueprintInternalUseOnly="true"))
+	bool AreSourcesValidForEvent(int32 EventKey) const;
+
 private:
 	//~ Source
 	void InitializeSource(FMVVMView_SourceKey SourceKey);
@@ -201,7 +210,7 @@ private:
 	/** The event that are registered by the view to the sources. */
 	TArray<FBoundEvent> BoundEvents;
 
-	/** The view has at least one binding that need to be ticked every frame. */
+	/** Bitfield that represents the valid sources. */
 	UPROPERTY(VisibleAnywhere, Transient, Category = "View")
 	uint64 ValidSources = 0;
 
diff --git a/Engine/Plugins/Runtime/ModelViewViewModel/Source/ModelViewViewModel/Public/View/MVVMViewClass.h b/Engine/Plugins/Runtime/ModelViewViewModel/Source/ModelViewViewModel/Public/View/MVVMViewClass.h
--- a/Engine/Plugins/Runtime/ModelViewViewModel/Source/ModelViewViewModel/Public/View/MVVMViewClass.h
+++ b/Engine/Plugins/Runtime/ModelViewViewModel/Source/ModelViewViewModel/Public/View/MVVMViewClass.h
@@ -451,6 +451,16 @@ public:
 		return SourceToReevaluate;
 	}
 
+	/**
+	 * A view event may require more than one view sources to run the event.
+	 * A event will not execute if any view source is invalid.
+	 * It will not warn if the view source is make as optional.
+	 */
+	uint64 GetSources() const
+	{
+		return SourceBitField;
+	}
+
 #if UE_WITH_MVVM_DEBUGGING
 	struct FToStringArgs
 	{
@@ -472,6 +482,9 @@ private:
 
 	UPROPERTY()
 	FMVVMViewClass_SourceKey SourceToReevaluate;
+	
+	UPROPERTY()
+	uint64 SourceBitField;
 };
 
 
diff --git a/Engine/Plugins/Runtime/ModelViewViewModel/Source/ModelViewViewModelBlueprint/Private/Bindings/MVVMConversionFunctionHelper.cpp b/Engine/Plugins/Runtime/ModelViewViewModel/Source/ModelViewViewModelBlueprint/Private/Bindings/MVVMConversionFunctionHelper.cpp
--- a/Engine/Plugins/Runtime/ModelViewViewModel/Source/ModelViewViewModelBlueprint/Private/Bindings/MVVMConversionFunctionHelper.cpp
+++ b/Engine/Plugins/Runtime/ModelViewViewModel/Source/ModelViewViewModelBlueprint/Private/Bindings/MVVMConversionFunctionHelper.cpp
@@ -763,7 +763,7 @@ TValueOrError<FCreateGraphResult, FText> CreateSetterGraph(UBlueprint* Blueprint
 		return MakeError(CanCreateSetterGraphResult.StealError());
 	}
 
-	const bool bIsEditable = true;
+	const bool bIsEditable = false;
 	const bool bAddToBlueprint = !bTransient;
 	Private::FCreateGraphResult CreateGraphInternalResult = Private::CreateGraph(Blueprint, GraphName, Signature, bIsConst, bIsEditable, bAddToBlueprint);
 
@@ -1057,6 +1057,41 @@ FCreateGraphResult CreateGraph(UBlueprint* Blueprint, FName GraphName, const UFu
 	return { NewGraph.FunctionGraph , CallFunctionNode };
 }
 
+UK2Node* InsertEarlyExitBranchNode(UEdGraph* Graph, TSubclassOf<UK2Node> BranchNodeType)
+{
+	if (Graph == nullptr || BranchNodeType.Get() == nullptr)
+	{
+		return nullptr;
+	}
+
+	UK2Node_FunctionEntry* FunctionEntry = Private::FindFunctionEntry(Graph);
+	if (FunctionEntry == nullptr)
+	{
+		return nullptr;
+	}
+
+	UK2Node* BranchNode = nullptr;
+	{
+		FGraphNodeCreator<UK2Node> BranchNodeCreator(*Graph);
+		BranchNode = BranchNodeCreator.CreateNode(true, BranchNodeType.Get());
+		BranchNode->NodePosX = FunctionEntry->NodePosX;
+		BranchNode->NodePosY = FunctionEntry->NodePosY+100;
+		BranchNodeCreator.Finalize();
+	}
+
+	const UEdGraphSchema* GraphSchema = GetDefault<UMVVMConversionFunctionGraphSchema>();
+	{
+		UEdGraphPin* EntryThenPin = FunctionEntry->FindPinChecked(UEdGraphSchema_K2::PN_Then);
+		UEdGraphPin* BranchThenPin = BranchNode->FindPinChecked(UEdGraphSchema_K2::PN_Then, EEdGraphPinDirection::EGPD_Output);
+		GraphSchema->MovePinLinks(*EntryThenPin, *BranchThenPin, true, false);
+
+		UEdGraphPin* BranchInputPin = BranchNode->FindPinChecked(UEdGraphSchema_K2::PN_Execute, EEdGraphPinDirection::EGPD_Input);
+		GraphSchema->TryCreateConnection(EntryThenPin, BranchInputPin);
+	}
+
+	return BranchNode;
+}
+
 UK2Node* GetWrapperNode(const UEdGraph* Graph)
 {
 	if (Graph == nullptr)
diff --git a/Engine/Plugins/Runtime/ModelViewViewModel/Source/ModelViewViewModelBlueprint/Private/MVVMBlueprintViewEvent.cpp b/Engine/Plugins/Runtime/ModelViewViewModel/Source/ModelViewViewModelBlueprint/Private/MVVMBlueprintViewEvent.cpp
--- a/Engine/Plugins/Runtime/ModelViewViewModel/Source/ModelViewViewModelBlueprint/Private/MVVMBlueprintViewEvent.cpp
+++ b/Engine/Plugins/Runtime/ModelViewViewModel/Source/ModelViewViewModelBlueprint/Private/MVVMBlueprintViewEvent.cpp
@@ -8,6 +8,7 @@
 #include "Bindings/MVVMFieldPathHelper.h"
 #include "Engine/Blueprint.h"
 #include "Kismet2/BlueprintEditorUtils.h"
+#include "MVVMConversionFunctionGraphSchema.h"
 #include "MVVMDeveloperProjectSettings.h"
 #include "MVVMWidgetBlueprintExtension_View.h"
 #include "WidgetBlueprint.h"
@@ -20,6 +21,7 @@
 #include "K2Node_FunctionResult.h"
 #include "K2Node_VariableSet.h"
 #include "KismetCompiler.h"
+#include "Node/MVVMK2Node_AreSourcesValidForEvent.h"
 
 #include UE_INLINE_GENERATED_CPP_BY_NAME(MVVMBlueprintViewEvent)
 
@@ -64,7 +66,7 @@ void UMVVMBlueprintViewEvent::SetDestinationPath(FMVVMBlueprintPropertyPath InDe
 	SavePinValues();
 }
 
-void UMVVMBlueprintViewEvent::SetCachedWrapperGraphInternal(UEdGraph* Graph, UK2Node* Node)
+void UMVVMBlueprintViewEvent::SetCachedWrapperGraphInternal(UEdGraph* Graph, UK2Node* Node, UMVVMK2Node_AreSourcesValidForEvent* SourceNode)
 {
 	if (CachedWrapperNode && OnUserDefinedPinRenamedHandle.IsValid())
 	{
@@ -77,6 +79,7 @@ void UMVVMBlueprintViewEvent::SetCachedWrapperGraphInternal(UEdGraph* Graph, UK2
 
 	CachedWrapperGraph = Graph;
 	CachedWrapperNode = Node;
+	CachedSourceValidNode = SourceNode;
 	OnGraphChangedHandle.Reset();
 	OnUserDefinedPinRenamedHandle.Reset();
 
@@ -88,6 +91,7 @@ void UMVVMBlueprintViewEvent::SetCachedWrapperGraphInternal(UEdGraph* Graph, UK2
 	{
 		OnUserDefinedPinRenamedHandle = CachedWrapperNode->OnUserDefinedPinRenamed().AddUObject(this, &UMVVMBlueprintViewEvent::HandleUserDefinedPinRenamed);
 	}
+	UpdateEventKeyInternal();
 }
 	
 UEdGraph* UMVVMBlueprintViewEvent::GetOrCreateWrapperGraph()
@@ -106,7 +110,7 @@ void UMVVMBlueprintViewEvent::RemoveWrapperGraph()
 	if (CachedWrapperGraph)
 	{
 		FBlueprintEditorUtils::RemoveGraph(GetWidgetBlueprintInternal(), CachedWrapperGraph);
-		SetCachedWrapperGraphInternal(nullptr, nullptr);
+		SetCachedWrapperGraphInternal(nullptr, nullptr, nullptr);
 	}
 
 	Messages.Empty();
@@ -155,6 +159,23 @@ bool UMVVMBlueprintViewEvent::HasOrphanedPin() const
 	return false;
 }
 
+void UMVVMBlueprintViewEvent::UpdateEventKey(FMVVMViewClass_EventKey InEventKey)
+{
+	if (EventKey != InEventKey)
+	{
+		EventKey = InEventKey;
+		UpdateEventKeyInternal();
+	}
+}
+
+void UMVVMBlueprintViewEvent::UpdateEventKeyInternal()
+{
+	if (CachedSourceValidNode)
+	{
+		CachedSourceValidNode->EventKey = EventKey;
+	}
+}
+
 FMVVMBlueprintPropertyPath UMVVMBlueprintViewEvent::GetPinPath(const FMVVMBlueprintPinId& PinId) const
 {
 	const FMVVMBlueprintPin* ViewPin = SavedPins.FindByPredicate([&PinId](const FMVVMBlueprintPin& Other) { return PinId == Other.GetId(); });
@@ -249,12 +270,13 @@ UEdGraph* UMVVMBlueprintViewEvent::CreateWrapperGraphInternal()
 	TValueOrError<UE::MVVM::ConversionFunctionHelper::FCreateGraphResult, FText> CreateSetterGraphResult = UE::MVVM::ConversionFunctionHelper::CreateSetterGraph(WidgetBlueprint, GraphName, DelegateSignature, DestinationPath, bIsConst, bTransient, true);
 	if (CreateSetterGraphResult.HasError())
 	{
-		SetCachedWrapperGraphInternal(nullptr, nullptr);
+		SetCachedWrapperGraphInternal(nullptr, nullptr, nullptr);
 		return nullptr;
 	}
 	else
 	{
-		SetCachedWrapperGraphInternal(CreateSetterGraphResult.GetValue().NewGraph, CreateSetterGraphResult.GetValue().WrappedNode);
+		UMVVMK2Node_AreSourcesValidForEvent* BranchNode = Cast<UMVVMK2Node_AreSourcesValidForEvent>(UE::MVVM::ConversionFunctionHelper::InsertEarlyExitBranchNode(CreateSetterGraphResult.GetValue().NewGraph, UMVVMK2Node_AreSourcesValidForEvent::StaticClass()));
+		SetCachedWrapperGraphInternal(CreateSetterGraphResult.GetValue().NewGraph, CreateSetterGraphResult.GetValue().WrappedNode, BranchNode);
 		LoadPinValuesInternal();
 	}
 
diff --git a/Engine/Plugins/Runtime/ModelViewViewModel/Source/ModelViewViewModelBlueprint/Private/MVVMViewBlueprintCompiler.cpp b/Engine/Plugins/Runtime/ModelViewViewModel/Source/ModelViewViewModelBlueprint/Private/MVVMViewBlueprintCompiler.cpp
--- a/Engine/Plugins/Runtime/ModelViewViewModel/Source/ModelViewViewModelBlueprint/Private/MVVMViewBlueprintCompiler.cpp
+++ b/Engine/Plugins/Runtime/ModelViewViewModel/Source/ModelViewViewModelBlueprint/Private/MVVMViewBlueprintCompiler.cpp
@@ -1152,10 +1152,11 @@ void FMVVMViewBlueprintCompiler::CategorizeBindings(const FWidgetBlueprintCompil
 
 void FMVVMViewBlueprintCompiler::CategorizeEvents(const FWidgetBlueprintCompilerContext::FCreateFunctionContext& Context)
 {
-	for (TSharedRef<FCompilerEvent>& Event : ValidEvents)
-	{
-		Event->Type = FCompilerEvent::EType::Invalid;
+	TArray<TSharedRef<FCompilerEvent>> TemporaryEvent = MoveTemp(ValidEvents);
+	ValidEvents.Reset(TemporaryEvent.Num());
 
+	for (TSharedRef<FCompilerEvent>& Event : TemporaryEvent)
+	{
 		UMVVMBlueprintViewEvent* EventPtr = Event->Event.Get();
 		UEdGraph* WrapperGraph = EventPtr->GetOrCreateWrapperGraph();
 		if (WrapperGraph == nullptr)
@@ -1174,7 +1175,8 @@ void FMVVMViewBlueprintCompiler::CategorizeEvents(const FWidgetBlueprintCompiler
 			continue;
 		}
 
-		Event->Type = FCompilerEvent::EType::Valid;
+		int32 Index = ValidEvents.Add(Event);
+		EventPtr->UpdateEventKey(FMVVMViewClass_EventKey(Index));
 	}
 }
 
@@ -1254,42 +1256,39 @@ void FMVVMViewBlueprintCompiler::CreateWriteFieldContexts(const FWidgetBlueprint
 		UMVVMBlueprintViewEvent* EventPtr = ValidEvent->Event.Get();
 		check(EventPtr);
 
-		if (ValidEvent->Type == FCompilerEvent::EType::Valid)
+		TValueOrError<FCreateFieldsResult, FText> FieldContextResult = CreateFieldContext(NewSkeletonClass, EventPtr->GetDestinationPath(), false);
+		if (FieldContextResult.HasError())
 		{
-			TValueOrError<FCreateFieldsResult, FText> FieldContextResult = CreateFieldContext(NewSkeletonClass, EventPtr->GetDestinationPath(), false);
-			if (FieldContextResult.HasError())
+			AddMessageForEvent(EventPtr
+				, FText::Format(Private::PropertyPathIsInvalidFormat, PropertyPathToText(NewSkeletonClass, BlueprintView.Get(), EventPtr->GetDestinationPath()))
+				, EMessageType::Error
+				, FMVVMBlueprintPinId()
+			);
+			bIsCreateFunctionsStepValid = false;
+			continue;
+		}
+
+		// Test if it already exist
+		TSharedPtr<FGeneratedWriteFieldPathContext> WriteFieldPath;
+		{
+			const TArray<UE::MVVM::FMVVMConstFieldVariant>& SkeletalGeneratedFieldsResult = FieldContextResult.GetValue().SkeletalGeneratedFields;
+			TSharedRef<FGeneratedWriteFieldPathContext>* Found = GeneratedWriteFieldPaths.FindByPredicate([&SkeletalGeneratedFieldsResult](const TSharedRef<FGeneratedWriteFieldPathContext>& Other) { return Other->SkeletalGeneratedFields == SkeletalGeneratedFieldsResult; });
+			if (Found)
 			{
-				AddMessageForEvent(EventPtr
-					, FText::Format(Private::PropertyPathIsInvalidFormat, PropertyPathToText(NewSkeletonClass, BlueprintView.Get(), EventPtr->GetDestinationPath()))
-					, EMessageType::Error
-					, FMVVMBlueprintPinId()
-				);
-				bIsCreateFunctionsStepValid = false;
-				continue;
+				WriteFieldPath = (*Found);
 			}
-
-			// Test if it already exist
-			TSharedPtr<FGeneratedWriteFieldPathContext> WriteFieldPath;
+			else
 			{
-				const TArray<UE::MVVM::FMVVMConstFieldVariant>& SkeletalGeneratedFieldsResult = FieldContextResult.GetValue().SkeletalGeneratedFields;
-				TSharedRef<FGeneratedWriteFieldPathContext>* Found = GeneratedWriteFieldPaths.FindByPredicate([&SkeletalGeneratedFieldsResult](const TSharedRef<FGeneratedWriteFieldPathContext>& Other) { return Other->SkeletalGeneratedFields == SkeletalGeneratedFieldsResult; });
-				if (Found)
-				{
-					WriteFieldPath = (*Found);
-				}
-				else
-				{
-					WriteFieldPath = MakeWriteFieldPath(
-						EventPtr->GetDestinationPath().GetSource(WidgetBlueprintCompilerContext.WidgetBlueprint())
-						, MoveTemp(FieldContextResult.GetValue().GeneratedFields)
-						,MoveTemp(FieldContextResult.GetValue().SkeletalGeneratedFields));
-					GeneratedWriteFieldPaths.Add(WriteFieldPath.ToSharedRef());
-				}
+				WriteFieldPath = MakeWriteFieldPath(
+					EventPtr->GetDestinationPath().GetSource(WidgetBlueprintCompilerContext.WidgetBlueprint())
+					, MoveTemp(FieldContextResult.GetValue().GeneratedFields)
+					,MoveTemp(FieldContextResult.GetValue().SkeletalGeneratedFields));
+				GeneratedWriteFieldPaths.Add(WriteFieldPath.ToSharedRef());
 			}
-
-			WriteFieldPath->UsedByEvents.AddUnique(ValidEvent);
-			ValidEvent->WritePath = WriteFieldPath;
 		}
+
+		WriteFieldPath->UsedByEvents.AddUnique(ValidEvent);
+		ValidEvent->WritePath = WriteFieldPath;
 	}
 }
 
@@ -1390,16 +1389,13 @@ void FMVVMViewBlueprintCompiler::CreateIntermediateGraphFunctions(const FWidgetB
 	for (TSharedRef<FCompilerEvent>& Event : ValidEvents)
 	{
 		UMVVMBlueprintViewEvent* EventPtr = Event->Event.Get();
-		if (Event->Type == FCompilerEvent::EType::Valid)
-		{
-			UEdGraph* WrapperGraph = EventPtr->GetOrCreateWrapperGraph();
-			ensure(WrapperGraph);
+		UEdGraph* WrapperGraph = EventPtr->GetOrCreateWrapperGraph();
+		ensure(WrapperGraph);
 
-			bool bAlreadyContained = WidgetBlueprintCompilerContext.Blueprint->FunctionGraphs.Contains(WrapperGraph);
-			if (ensure(!bAlreadyContained))
-			{
-				Context.AddGeneratedFunctionGraph(WrapperGraph);
-			}
+		bool bAlreadyContained = WidgetBlueprintCompilerContext.Blueprint->FunctionGraphs.Contains(WrapperGraph);
+		if (ensure(!bAlreadyContained))
+		{
+			Context.AddGeneratedFunctionGraph(WrapperGraph);
 		}
 	}
 }
@@ -1811,36 +1807,33 @@ void FMVVMViewBlueprintCompiler::CreateReadFieldContexts(UWidgetBlueprintGenerat
 	{
 		UMVVMBlueprintViewEvent* EventPtr = ValidEvent->Event.Get();
 		check(EventPtr);
-		if (ValidEvent->Type == FCompilerEvent::EType::Valid)
+		for (const FMVVMBlueprintPin& Pin : EventPtr->GetPins())
 		{
-			for (const FMVVMBlueprintPin& Pin : EventPtr->GetPins())
+			if (Pin.UsedPathAsValue())
 			{
-				if (Pin.UsedPathAsValue())
+				TValueOrError<FCreateFieldsResult, FText> CreateSourceResult = CreateFieldContext(Class, Pin.GetPath(), true);
+				if (CreateSourceResult.HasError())
 				{
-					TValueOrError<FCreateFieldsResult, FText> CreateSourceResult = CreateFieldContext(Class, Pin.GetPath(), true);
-					if (CreateSourceResult.HasError())
-					{
-						AddMessageForEvent(ValidEvent
-							, FText::Format(Private::PropertyPathIsInvalidFormat, PropertyPathToText(Class, BlueprintView.Get(), Pin.GetPath()))
-							, EMessageType::Error
-							, Pin.GetId()
-						);
-						bIsPreCompileStepValid = false;
-					}
-
-					TSharedPtr<FGeneratedReadFieldPathContext> Found = AlreadyExist(CreateSourceResult.GetValue().SkeletalGeneratedFields);
-					if (!Found)
-					{
-						Found = MakeShared<FGeneratedReadFieldPathContext>();
-						Found->Source = MoveTemp(CreateSourceResult.GetValue().OptionalSource);
-						Found->GeneratedFields = MoveTemp(CreateSourceResult.GetValue().GeneratedFields);
-						Found->SkeletalGeneratedFields = MoveTemp(CreateSourceResult.GetValue().SkeletalGeneratedFields);
-						GeneratedReadFieldPaths.Add(Found.ToSharedRef());
-					}
+					AddMessageForEvent(ValidEvent
+						, FText::Format(Private::PropertyPathIsInvalidFormat, PropertyPathToText(Class, BlueprintView.Get(), Pin.GetPath()))
+						, EMessageType::Error
+						, Pin.GetId()
+					);
+					bIsPreCompileStepValid = false;
+				}
 
-					Found->UsedByEvents.AddUnique(ValidEvent);
-					ValidEvent->ReadPaths.Add(Found.ToSharedRef());
+				TSharedPtr<FGeneratedReadFieldPathContext> Found = AlreadyExist(CreateSourceResult.GetValue().SkeletalGeneratedFields);
+				if (!Found)
+				{
+					Found = MakeShared<FGeneratedReadFieldPathContext>();
+					Found->Source = MoveTemp(CreateSourceResult.GetValue().OptionalSource);
+					Found->GeneratedFields = MoveTemp(CreateSourceResult.GetValue().GeneratedFields);
+					Found->SkeletalGeneratedFields = MoveTemp(CreateSourceResult.GetValue().SkeletalGeneratedFields);
+					GeneratedReadFieldPaths.Add(Found.ToSharedRef());
 				}
+
+				Found->UsedByEvents.AddUnique(ValidEvent);
+				ValidEvent->ReadPaths.Add(Found.ToSharedRef());
 			}
 		}
 	}
@@ -3048,6 +3041,7 @@ void FMVVMViewBlueprintCompiler::PreCompileEvents(UWidgetBlueprintGeneratedClass
 
 void FMVVMViewBlueprintCompiler::CompileEvents(const FCompiledBindingLibraryCompiler::FCompileResult& CompileResult, UWidgetBlueprintGeneratedClass* Class, UMVVMViewClass* ViewExtension)
 {
+	// NB. The order is important. The index is used in the generated function to identify the event.
 	for (const TSharedRef<FCompilerEvent>& ValidEvent : ValidEvents)
 	{
 		const FMVVMVCompiledFieldPath* CompiledFieldPath = CompileResult.FieldPaths.Find(ValidEvent->DelegateFieldPathHandle);
@@ -3074,10 +3068,52 @@ void FMVVMViewBlueprintCompiler::CompileEvents(const FCompiledBindingLibraryComp
 				});
 		}
 
+		uint64 SourceBitField = 0;
+		{
+			if (ValidEvent->WritePath && ValidEvent->WritePath->OptionalSource)
+			{
+				int32 ViewExtensionSourceCreatorsIndex = ViewExtension->Sources.IndexOfByPredicate([LookFor = ValidEvent->WritePath->OptionalSource->Name](const FMVVMViewClass_Source& Other)
+					{
+						return Other.GetName() == LookFor;
+					});
+				if (ViewExtension->Sources.IsValidIndex(ViewExtensionSourceCreatorsIndex))
+				{
+
+					FMVVMViewClass_SourceKey FieldClassSourceKey = FMVVMViewClass_SourceKey(ViewExtensionSourceCreatorsIndex);
+					SourceBitField |= FieldClassSourceKey.GetBit();
+				}
+			}
+			for (TSharedPtr<FGeneratedReadFieldPathContext> ReadPath : ValidEvent->ReadPaths)
+			{
+				if (ReadPath->Source == nullptr)
+				{
+					AddMessageForEvent(ValidEvent->Event.Get(), LOCTEXT("InvalidEventSourceInternal", "Internal error. The event has an invalid source."), EMessageType::Error, FMVVMBlueprintPinId());
+					bIsCompileStepValid = false;
+					continue;
+				}
+
+				int32 ViewExtensionSourceCreatorsIndex = ViewExtension->Sources.IndexOfByPredicate([LookFor = ReadPath->Source->Name](const FMVVMViewClass_Source& Other)
+					{
+						return Other.GetName() == LookFor;
+					});
+				if (!ViewExtension->Sources.IsValidIndex(ViewExtensionSourceCreatorsIndex))
+				{
+					AddMessageForEvent(ValidEvent->Event.Get(), LOCTEXT("CompiledEventSourceCreatorNotGenerated", "Internal error. The source creator was not generated."), EMessageType::Error, FMVVMBlueprintPinId());
+					bIsCompileStepValid = false;
+					continue;
+				}
+
+				// Add the needed source.
+				FMVVMViewClass_SourceKey FieldClassSourceKey = FMVVMViewClass_SourceKey(ViewExtensionSourceCreatorsIndex);
+				SourceBitField |= FieldClassSourceKey.GetBit();
+			}
+		}
+
 		FMVVMViewClass_Event& NewBinding = ViewExtension->Events.AddDefaulted_GetRef();
 		NewBinding.FieldPath = *CompiledFieldPath;
 		NewBinding.UserWidgetFunctionName = ValidEvent->GeneratedGraphName;
 		NewBinding.SourceToReevaluate = FoundSourceIndex != INDEX_NONE ? FMVVMViewClass_SourceKey(FoundSourceIndex) : FMVVMViewClass_SourceKey();
+		NewBinding.SourceBitField = SourceBitField;
 	}
 }
 
diff --git a/Engine/Plugins/Runtime/ModelViewViewModel/Source/ModelViewViewModelBlueprint/Private/MVVMViewBlueprintCompiler.h b/Engine/Plugins/Runtime/ModelViewViewModel/Source/ModelViewViewModelBlueprint/Private/MVVMViewBlueprintCompiler.h
--- a/Engine/Plugins/Runtime/ModelViewViewModel/Source/ModelViewViewModelBlueprint/Private/MVVMViewBlueprintCompiler.h
+++ b/Engine/Plugins/Runtime/ModelViewViewModel/Source/ModelViewViewModelBlueprint/Private/MVVMViewBlueprintCompiler.h
@@ -321,14 +321,7 @@ private:
 	 */
 	struct FCompilerEvent
 	{
-		enum class EType : int32
-		{
-			Unknown = 0, // was not evaluated yet
-			Invalid = -1, // evaluated and not valid
-			Valid = 1,
-		};
 		TWeakObjectPtr<UMVVMBlueprintViewEvent> Event = nullptr;
-		EType Type = EType::Unknown;
 
 		TArray<TSharedPtr<FGeneratedReadFieldPathContext>> ReadPaths;
 		TSharedPtr<FGeneratedWriteFieldPathContext> WritePath;
diff --git a/Engine/Plugins/Runtime/ModelViewViewModel/Source/ModelViewViewModelBlueprint/Private/Node/MVVMK2Node_AreSourcesValidForEvent.cpp b/Engine/Plugins/Runtime/ModelViewViewModel/Source/ModelViewViewModelBlueprint/Private/Node/MVVMK2Node_AreSourcesValidForEvent.cpp
new file mode 100644
--- /dev/null
+++ b/Engine/Plugins/Runtime/ModelViewViewModel/Source/ModelViewViewModelBlueprint/Private/Node/MVVMK2Node_AreSourcesValidForEvent.cpp
@@ -0,0 +1,113 @@
+// Copyright Epic Games, Inc. All Rights Reserved.
+
+#include "Node/MVVMK2Node_AreSourcesValidForEvent.h"
+
+#include "BlueprintCompiledStatement.h"
+#include "BlueprintNodeSpawner.h"
+#include "EdGraph/EdGraphPin.h"
+#include "EdGraphSchema_K2.h"
+#include "Kismet2/CompilerResultsLog.h"
+#include "KismetCompiledFunctionContext.h"
+#include "KismetCompiler.h"
+#include "MVVMSubsystem.h"
+#include "Subsystems/SubsystemBlueprintLibrary.h"
+#include "View/MVVMView.h"
+
+#include "K2Node_CallFunction.h"
+#include "K2Node_IfThenElse.h"
+#include "K2Node_Self.h"
+
+#define LOCTEXT_NAMESPACE "MVVMK2Node_AreSourcesValidForEvent"
+
+//////////////////////////////////////////////////////////////////////////
+
+void UMVVMK2Node_AreSourcesValidForEvent::AllocateDefaultPins()
+{
+	const UEdGraphSchema_K2* K2Schema = GetDefault<UEdGraphSchema_K2>();
+
+	CreatePin(EGPD_Input, UEdGraphSchema_K2::PC_Exec, UEdGraphSchema_K2::PN_Execute);
+
+	UEdGraphPin* TruePin = CreatePin(EGPD_Output, UEdGraphSchema_K2::PC_Exec, UEdGraphSchema_K2::PN_Then);
+	TruePin->PinFriendlyName = LOCTEXT("true", "true");
+
+	UEdGraphPin* FalsePin = CreatePin(EGPD_Output, UEdGraphSchema_K2::PC_Exec, UEdGraphSchema_K2::PN_Else);
+	FalsePin->PinFriendlyName = LOCTEXT("false", "false");
+
+	Super::AllocateDefaultPins();
+}
+
+void UMVVMK2Node_AreSourcesValidForEvent::ExpandNode(FKismetCompilerContext& CompilerContext, UEdGraph* SourceGraph)
+{
+	Super::ExpandNode(CompilerContext, SourceGraph);
+
+	if (!EventKey.IsValid())
+	{
+		CompilerContext.MessageLog.Error(*LOCTEXT("NoEventKey", "Node @@ doesn't have a valid event key.").ToString(), this);
+		BreakAllNodeLinks();
+		return;
+	}
+
+	const UEdGraphSchema_K2* Schema = CompilerContext.GetSchema();
+
+	UK2Node_CallFunction* CallGetSubsystemNode = CompilerContext.SpawnIntermediateNode<UK2Node_CallFunction>(this, SourceGraph);
+	{
+		CallGetSubsystemNode->FunctionReference.SetExternalMember(GET_FUNCTION_NAME_CHECKED(USubsystemBlueprintLibrary, GetEngineSubsystem), USubsystemBlueprintLibrary::StaticClass());
+		CallGetSubsystemNode->AllocateDefaultPins();
+		UEdGraphPin* CallCreateClassTypePin = CallGetSubsystemNode->FindPinChecked(FName("Class"));
+		CallCreateClassTypePin->DefaultObject = UMVVMSubsystem::StaticClass();
+	}
+
+	UK2Node_CallFunction* CallGetViewNode = CompilerContext.SpawnIntermediateNode<UK2Node_CallFunction>(this, SourceGraph);
+	{
+		CallGetViewNode->FunctionReference.SetExternalMember(GET_FUNCTION_NAME_CHECKED(UMVVMSubsystem, K2_GetViewFromUserWidget), UMVVMSubsystem::StaticClass());
+		CallGetViewNode->AllocateDefaultPins();
+	}
+
+	UK2Node_Self* SelfNode = CompilerContext.SpawnIntermediateNode<UK2Node_Self>(this, SourceGraph);
+	{
+		SelfNode->AllocateDefaultPins();
+	}
+
+	UK2Node_CallFunction* CallAreSourcesValidForEventNode = CompilerContext.SpawnIntermediateNode<UK2Node_CallFunction>(this, SourceGraph);
+	{
+		CallAreSourcesValidForEventNode->FunctionReference.SetExternalMember(FName("AreSourcesValidForEvent"), UMVVMView::StaticClass());
+		CallAreSourcesValidForEventNode->AllocateDefaultPins();
+
+		UEdGraphPin* EventKeyPin = CallAreSourcesValidForEventNode->FindPinChecked(FName("EventKey"));
+		Schema->TrySetDefaultValue(*EventKeyPin, LexToString(EventKey.GetIndex()));
+	}
+
+	UK2Node_IfThenElse* BranchNode = CompilerContext.SpawnIntermediateNode<UK2Node_IfThenElse>(this, SourceGraph);
+	{
+		BranchNode->AllocateDefaultPins();
+	}
+
+	// Casting subsytem.result into MVVMSubsystem
+	CallGetSubsystemNode->GetReturnValuePin()->PinType = CallGetViewNode->FindPinChecked(UEdGraphSchema_K2::PN_Self)->PinType;
+	// subsystem.result to GetViewFromUserWidget.target
+	check(Schema->TryCreateConnection(CallGetSubsystemNode->GetReturnValuePin(), CallGetViewNode->FindPinChecked(UEdGraphSchema_K2::PN_Self)));
+	// self to GetViewFromUserWidget.UserWidget
+	check(Schema->TryCreateConnection(SelfNode->FindPinChecked(UEdGraphSchema_K2::PN_Self), CallGetViewNode->FindPinChecked(FName("UserWidget"))));
+	// GetViewFromUserWidget.Result to AreSourcesValidForEvent.target
+	check(Schema->TryCreateConnection(CallGetViewNode->FindPinChecked(UEdGraphSchema_K2::PN_ReturnValue), CallAreSourcesValidForEventNode->FindPinChecked(UEdGraphSchema_K2::PN_Self)));
+	// AreSourcesValidForEvent.result to branch.condition
+	check(Schema->TryCreateConnection(CallAreSourcesValidForEventNode->FindPinChecked(UEdGraphSchema_K2::PN_ReturnValue), BranchNode->GetConditionPin()));
+	// move this.exec to branch.exec
+	CompilerContext.MovePinLinksToIntermediate(*GetExecPin(), *BranchNode->GetExecPin());
+	// move this.then to branch.then
+	CompilerContext.MovePinLinksToIntermediate(*GetThenPin(), *BranchNode->GetThenPin());
+
+	BreakAllNodeLinks();
+}
+
+UEdGraphPin* UMVVMK2Node_AreSourcesValidForEvent::GetThenPin() const
+{
+	return FindPinChecked(UEdGraphSchema_K2::PN_Then);
+}
+
+UEdGraphPin* UMVVMK2Node_AreSourcesValidForEvent::GetElsePin() const
+{
+	return FindPinChecked(UEdGraphSchema_K2::PN_Else);
+}
+
+#undef LOCTEXT_NAMESPACE
diff --git a/Engine/Plugins/Runtime/ModelViewViewModel/Source/ModelViewViewModelBlueprint/Private/Node/MVVMK2Node_AreSourcesValidForEvent.h b/Engine/Plugins/Runtime/ModelViewViewModel/Source/ModelViewViewModelBlueprint/Private/Node/MVVMK2Node_AreSourcesValidForEvent.h
new file mode 100644
--- /dev/null
+++ b/Engine/Plugins/Runtime/ModelViewViewModel/Source/ModelViewViewModelBlueprint/Private/Node/MVVMK2Node_AreSourcesValidForEvent.h
@@ -0,0 +1,33 @@
+// Copyright Epic Games, Inc. All Rights Reserved.
+
+#pragma once
+
+#include "CoreMinimal.h"
+#include "K2Node.h"
+#include "View/MVVMViewTypes.h"
+
+#include "MVVMK2Node_AreSourcesValidForEvent.generated.h"
+
+class UEdGraphPin;
+class FKismetCompilerContext;
+
+UCLASS(MinimalAPI, Hidden, HideDropDown)
+class UMVVMK2Node_AreSourcesValidForEvent : public UK2Node
+{
+	GENERATED_BODY()
+
+public:
+	//~ Begin K2Node Interface.
+	virtual void ExpandNode(FKismetCompilerContext& CompilerContext, UEdGraph* SourceGraph) override;
+	virtual void AllocateDefaultPins() override;
+	//~ End K2Node Interface.
+
+	/** Get the then output pin */
+	UEdGraphPin* GetThenPin() const;
+	/** Get the else output pin */
+	UEdGraphPin* GetElsePin() const;
+
+	UPROPERTY()
+	FMVVMViewClass_EventKey EventKey;
+};
+
diff --git a/Engine/Plugins/Runtime/ModelViewViewModel/Source/ModelViewViewModelBlueprint/Public/Bindings/MVVMConversionFunctionHelper.h b/Engine/Plugins/Runtime/ModelViewViewModel/Source/ModelViewViewModelBlueprint/Public/Bindings/MVVMConversionFunctionHelper.h
--- a/Engine/Plugins/Runtime/ModelViewViewModel/Source/ModelViewViewModelBlueprint/Public/Bindings/MVVMConversionFunctionHelper.h
+++ b/Engine/Plugins/Runtime/ModelViewViewModel/Source/ModelViewViewModelBlueprint/Public/Bindings/MVVMConversionFunctionHelper.h
@@ -66,6 +66,9 @@ namespace UE::MVVM::ConversionFunctionHelper
 	/** */
 	MODELVIEWVIEWMODELBLUEPRINT_API FCreateGraphResult CreateGraph(UBlueprint* WidgetBlueprint, FName GraphName, const UFunction* Signature, const TSubclassOf<UK2Node> Node, bool bIsConst, bool bTransient, TFunctionRef<void(UK2Node*)> InitNodeCallback);
 
+	/** Insert a branch node to the existing graph to test before executing the rest of the . */
+	MODELVIEWVIEWMODELBLUEPRINT_API UK2Node* InsertEarlyExitBranchNode(UEdGraph* Graph, TSubclassOf<UK2Node> BranchNode);
+
 	/** Find the main conversion function node from the given graph. */
 	MODELVIEWVIEWMODELBLUEPRINT_API UK2Node* GetWrapperNode(const UEdGraph* Graph);
 
diff --git a/Engine/Plugins/Runtime/ModelViewViewModel/Source/ModelViewViewModelBlueprint/Public/MVVMBlueprintViewEvent.h b/Engine/Plugins/Runtime/ModelViewViewModel/Source/ModelViewViewModelBlueprint/Public/MVVMBlueprintViewEvent.h
--- a/Engine/Plugins/Runtime/ModelViewViewModel/Source/ModelViewViewModelBlueprint/Public/MVVMBlueprintViewEvent.h
+++ b/Engine/Plugins/Runtime/ModelViewViewModel/Source/ModelViewViewModelBlueprint/Public/MVVMBlueprintViewEvent.h
@@ -4,12 +4,14 @@
 
 #include "MVVMPropertyPath.h"
 #include "MVVMBlueprintPin.h"
+#include "View/MVVMViewTypes.h"
 
 #include "MVVMBlueprintViewEvent.generated.h"
 
 struct FEdGraphEditAction;
 class UEdGraph;
 class UK2Node;
+class UMVVMK2Node_AreSourcesValidForEvent;
 class UWidgetBlueprint;
 
 /**
@@ -91,6 +93,8 @@ public:
 	void UpdatePinValues();
 	/** Keep the orphaned pins. Add the missing pins. */
 	bool HasOrphanedPin() const;
+	/** Event sources are tested at runtime to check if they are valid. */
+	void UpdateEventKey(FMVVMViewClass_EventKey EventKey);
 
 	UEdGraphPin* GetOrCreateGraphPin(const FMVVMBlueprintPinId& Pin);
 
@@ -129,9 +133,10 @@ private:
 	void HandleGraphChanged(const FEdGraphEditAction& Action);
 	void HandleUserDefinedPinRenamed(UK2Node* InNode, FName OldPinName, FName NewPinName);
 	UWidgetBlueprint* GetWidgetBlueprintInternal() const;
-	void SetCachedWrapperGraphInternal(UEdGraph* Graph, UK2Node* Node);
+	void SetCachedWrapperGraphInternal(UEdGraph* Graph, UK2Node* Node, UMVVMK2Node_AreSourcesValidForEvent* SourceNode);
 	UEdGraph* CreateWrapperGraphInternal();
 	void LoadPinValuesInternal();
+	void UpdateEventKeyInternal();
 
 private:
 	UPROPERTY(VisibleAnywhere, Category = "Viewmodel")
@@ -150,6 +155,9 @@ private:
 	UPROPERTY(VisibleAnywhere, Category = "Viewmodel")
 	FName GraphName;
 
+	UPROPERTY(VisibleAnywhere, Category = "Viewmodel")
+	FMVVMViewClass_EventKey EventKey;
+
 	mutable TArray<FMessage> Messages;
 	bool bLoadingPins = false;
 
@@ -159,6 +167,9 @@ private:
 	UPROPERTY(Transient, DuplicateTransient)
 	mutable TObjectPtr<UK2Node> CachedWrapperNode;
 
+	UPROPERTY(Transient, DuplicateTransient)
+	mutable TObjectPtr<UMVVMK2Node_AreSourcesValidForEvent> CachedSourceValidNode;
+
 	FDelegateHandle OnGraphChangedHandle;
 	FDelegateHandle OnUserDefinedPinRenamedHandle;
 };
