MVVM: Set the viewclass's type as explicitly constructed just to prevent implicit errors.

diff --git a/Engine/Plugins/Runtime/ModelViewViewModel/Source/ModelViewViewModel/Public/View/MVVMViewTypes.h b/Engine/Plugins/Runtime/ModelViewViewModel/Source/ModelViewViewModel/Public/View/MVVMViewTypes.h
--- a/Engine/Plugins/Runtime/ModelViewViewModel/Source/ModelViewViewModel/Public/View/MVVMViewTypes.h
+++ b/Engine/Plugins/Runtime/ModelViewViewModel/Source/ModelViewViewModel/Public/View/MVVMViewTypes.h
@@ -24,7 +24,7 @@ struct FMVVMViewClass_FieldId
 
 public:
 	FMVVMViewClass_FieldId() = default;
-	FMVVMViewClass_FieldId(UE::FieldNotification::FFieldId InFieldId)
+	explicit FMVVMViewClass_FieldId(UE::FieldNotification::FFieldId InFieldId)
 	{
 		FieldName = InFieldId.GetName();
 		FieldIndex = InFieldId.GetIndex();
@@ -107,7 +107,7 @@ struct FMVVMViewClass_BindingKey
 
 public:
 	FMVVMViewClass_BindingKey() = default;
-	FMVVMViewClass_BindingKey(int32 InIndex)
+	explicit FMVVMViewClass_BindingKey(int32 InIndex)
 		: Index(InIndex)
 	{}
 	bool IsValid() const
@@ -142,7 +142,7 @@ struct FMVVMViewClass_EvaluateBindingKey
 
 public:
 	FMVVMViewClass_EvaluateBindingKey() = default;
-	FMVVMViewClass_EvaluateBindingKey(int32 InIndex)
+	explicit FMVVMViewClass_EvaluateBindingKey(int32 InIndex)
 		: Index(InIndex)
 	{}
 	bool IsValid() const
@@ -177,7 +177,7 @@ struct FMVVMViewClass_EventKey
 
 public:
 	FMVVMViewClass_EventKey() = default;
-	FMVVMViewClass_EventKey(int32 InIndex)
+	explicit FMVVMViewClass_EventKey(int32 InIndex)
 		: Index(InIndex)
 	{}
 	bool IsValid() const
