MVVM: Fix GetReturn value that returns an invalid property in the MVVM context. There is only one none const out param or a property mark as the return type. That is to prevent unintended mistakes with argument ordering.

diff --git a/Engine/Plugins/Runtime/ModelViewViewModel/Source/ModelViewViewModel/Private/Bindings/MVVMBindingHelper.cpp b/Engine/Plugins/Runtime/ModelViewViewModel/Source/ModelViewViewModel/Private/Bindings/MVVMBindingHelper.cpp
--- a/Engine/Plugins/Runtime/ModelViewViewModel/Source/ModelViewViewModel/Private/Bindings/MVVMBindingHelper.cpp
+++ b/Engine/Plugins/Runtime/ModelViewViewModel/Source/ModelViewViewModel/Private/Bindings/MVVMBindingHelper.cpp
@@ -413,6 +413,7 @@ namespace UE::MVVM::BindingHelper
 			return TryGetPropertyTypeForDestinationBinding(InField.GetFunction());
 		}
 	}
+
 	TValueOrError<const FProperty*, FText> TryGetReturnTypeForConversionFunction(const UFunction* InFunction)
 	{
 		FText CommonResult = Private::TryGetPropertyTypeCommon(InFunction);
@@ -572,14 +573,24 @@ namespace UE::MVVM::BindingHelper
 		check(InFunction);
 		FProperty* Result = InFunction->GetReturnProperty();
 
-		if (Result == nullptr && InFunction->HasAllFunctionFlags(FUNC_HasOutParms))
+		if (InFunction->HasAllFunctionFlags(FUNC_HasOutParms))
 		{
 			for (TFieldIterator<FProperty> It(InFunction); It && (It->PropertyFlags & CPF_Parm); ++It)
 			{
-				if (It->HasAllPropertyFlags(CPF_OutParm) && !It->HasAnyPropertyFlags(CPF_Deprecated | CPF_EditorOnly | CPF_ConstParm | CPF_ReferenceParm))
+				if (It->HasAllPropertyFlags(CPF_OutParm) && !It->HasAnyPropertyFlags(CPF_ConstParm | CPF_ReferenceParm))
 				{
-					Result = *It;
-					break;
+					if (Result == nullptr)
+					{
+						if (!It->HasAnyPropertyFlags(CPF_Deprecated | CPF_EditorOnly))
+						{
+							Result = *It;
+						}
+					}
+					else if (Result != *It)
+					{
+						Result = nullptr;
+						break;
+					}
 				}
 			}
 		}
@@ -633,13 +644,36 @@ namespace UE::MVVM::BindingHelper
 		return Arguments;
 	}
 
+	namespace Private
+	{
+		const FProperty* GetRuntimeReturnProperty(const UFunction* InFunction)
+		{
+			check(InFunction);
+			FProperty* Result = InFunction->GetReturnProperty();
+
+			if (Result == nullptr && InFunction->HasAllFunctionFlags(FUNC_HasOutParms))
+			{
+				for (TFieldIterator<FProperty> It(InFunction); It && (It->PropertyFlags & CPF_Parm); ++It)
+				{
+					if (It->HasAllPropertyFlags(CPF_OutParm) && !It->HasAnyPropertyFlags(CPF_Deprecated | CPF_EditorOnly | CPF_ConstParm | CPF_ReferenceParm))
+					{
+						Result = *It;
+						break;
+					}
+				}
+			}
+
+			return Result;
+		}
+	}
+
 	void ExecuteBinding_NoCheck(const FFieldContext& Source, const FFieldContext& Destination)
 	{
 		check(!Source.GetObjectVariant().IsNull() && !Destination.GetObjectVariant().IsNull());
 		check(!Source.GetFieldVariant().IsEmpty() && !Destination.GetFieldVariant().IsEmpty());
 
 		const bool bIsSourceBindingIsProperty = Source.GetFieldVariant().IsProperty();
-		const FProperty* GetterType = bIsSourceBindingIsProperty ? Source.GetFieldVariant().GetProperty() : GetReturnProperty(Source.GetFieldVariant().GetFunction());
+		const FProperty* GetterType = bIsSourceBindingIsProperty ? Source.GetFieldVariant().GetProperty() : Private::GetRuntimeReturnProperty(Source.GetFieldVariant().GetFunction());
 		check(GetterType);
 
 		const bool bIsDestinationBindingIsProperty = Destination.GetFieldVariant().IsProperty();
@@ -696,7 +730,7 @@ namespace UE::MVVM::BindingHelper
 		ConversionFunction.GetObject()->ProcessEvent(ConversionFunction.GetFunction(), ConversionFunctionDataPtr);
 
 		{
-			const FProperty* ReturnConversionProperty = GetReturnProperty(ConversionFunction.GetFunction());
+			const FProperty* ReturnConversionProperty = Private::GetRuntimeReturnProperty(ConversionFunction.GetFunction());
 			check(ReturnConversionProperty);
 
 			const bool bIsDestinationBindingIsProperty = Destination.GetFieldVariant().IsProperty();
@@ -770,7 +804,7 @@ namespace UE::MVVM::BindingHelper
 			check(ArgumentConversionProperty);
 
 			const bool bIsSourceBindingIsProperty = Source.GetFieldVariant().IsProperty();
-			const FProperty* GetterType = bIsSourceBindingIsProperty ? Source.GetFieldVariant().GetProperty() : GetReturnProperty(Source.GetFieldVariant().GetFunction());
+			const FProperty* GetterType = bIsSourceBindingIsProperty ? Source.GetFieldVariant().GetProperty() : Private::GetRuntimeReturnProperty(Source.GetFieldVariant().GetFunction());
 			check(GetterType);
 			check(ArePropertiesCompatible(GetterType, ArgumentConversionProperty));
 
@@ -819,7 +853,7 @@ namespace UE::MVVM::BindingHelper
 		ConversionFunction.GetObject()->ProcessEvent(ConversionFunction.GetFunction(), ConversionFunctionDataPtr);
 
 		{
-			const FProperty* ReturnConversionProperty = GetReturnProperty(ConversionFunction.GetFunction());
+			const FProperty* ReturnConversionProperty = Private::GetRuntimeReturnProperty(ConversionFunction.GetFunction());
 			check(ReturnConversionProperty);
 
 			const bool bIsDestinationBindingIsProperty = Destination.GetFieldVariant().IsProperty();
diff --git a/Engine/Plugins/Runtime/ModelViewViewModel/Source/ModelViewViewModel/Public/Bindings/MVVMBindingHelper.h b/Engine/Plugins/Runtime/ModelViewViewModel/Source/ModelViewViewModel/Public/Bindings/MVVMBindingHelper.h
--- a/Engine/Plugins/Runtime/ModelViewViewModel/Source/ModelViewViewModel/Public/Bindings/MVVMBindingHelper.h
+++ b/Engine/Plugins/Runtime/ModelViewViewModel/Source/ModelViewViewModel/Public/Bindings/MVVMBindingHelper.h
@@ -113,6 +113,7 @@ namespace UE::MVVM::BindingHelper
 	 * int Foo(double) -> returns int
 	 * void Foo(int&, double) -> returns int
 	 * void Foo(const int&, double) -> returns null
+	 * int Foo(int&, double) -> returns null
 	 */
 	[[nodiscard]] MODELVIEWVIEWMODEL_API const FProperty* GetReturnProperty(const UFunction* InFunction);
 
diff --git a/Engine/Plugins/Runtime/ModelViewViewModel/Source/ModelViewViewModelEditor/Private/Tests/MVVMBindingHelperTest.cpp b/Engine/Plugins/Runtime/ModelViewViewModel/Source/ModelViewViewModelEditor/Private/Tests/MVVMBindingHelperTest.cpp
--- a/Engine/Plugins/Runtime/ModelViewViewModel/Source/ModelViewViewModelEditor/Private/Tests/MVVMBindingHelperTest.cpp
+++ b/Engine/Plugins/Runtime/ModelViewViewModel/Source/ModelViewViewModelEditor/Private/Tests/MVVMBindingHelperTest.cpp
@@ -208,8 +208,8 @@ bool FMVVMBindingHelperTest::RunTest(const FString& Parameters)
 	TestConversionFunction("FunctionConversionF", nullptr, FIntProperty::StaticClass());
 	TestConversionFunction("FunctionConversionG", FArrayProperty::StaticClass(), FSetProperty::StaticClass());
 	TestConversionFunction("FunctionConversionH", FArrayProperty::StaticClass(), FSetProperty::StaticClass());
-	TestConversionFunction("FunctionConversionI", FArrayProperty::StaticClass(), FSetProperty::StaticClass());
-	TestConversionFunction("FunctionConversionJ", FArrayProperty::StaticClass(), FMapProperty::StaticClass());
+	TestConversionFunction("FunctionConversionI", nullptr, FSetProperty::StaticClass());
+	TestConversionFunction("FunctionConversionJ", nullptr, FMapProperty::StaticClass());
 
 	return true;
 }
