MVVM: Evalute the bindings automatically when the source is set but the value didn't changed. That is a global option and an option on the viewmodel definition itself.
Change the category of MVVMDeveloperProjectSettings.
Fix bug in customization, the bCreateSetter flag was set but never propagated to the copied object.

diff --git a/Engine/Plugins/Runtime/ModelViewViewModel/Source/ModelViewViewModel/Private/View/MVVMView.cpp b/Engine/Plugins/Runtime/ModelViewViewModel/Source/ModelViewViewModel/Private/View/MVVMView.cpp
--- a/Engine/Plugins/Runtime/ModelViewViewModel/Source/ModelViewViewModel/Private/View/MVVMView.cpp
+++ b/Engine/Plugins/Runtime/ModelViewViewModel/Source/ModelViewViewModel/Private/View/MVVMView.cpp
@@ -24,6 +24,7 @@ DECLARE_CYCLE_STAT(TEXT("UninitializeSources"), STAT_UMG_Viewmodel_UninitializeS
 DECLARE_CYCLE_STAT(TEXT("InitializeBindings"), STAT_UMG_Viewmodel_InitializeBindings, STATGROUP_UMG_Viewmodel);
 DECLARE_CYCLE_STAT(TEXT("UninitializeBindings"), STAT_UMG_Viewmodel_UninitializeBindings, STATGROUP_UMG_Viewmodel);
 DECLARE_CYCLE_STAT(TEXT("SetSource"), STAT_UMG_Viewmodel_SetSource, STATGROUP_UMG_Viewmodel);
+DECLARE_CYCLE_STAT(TEXT("ExecuteViewModel"), STAT_UMG_Viewmodel_ExecuteViewModel, STATGROUP_UMG_Viewmodel);
 DECLARE_CYCLE_STAT(TEXT("ExecuteBinding ValueChanged"), STAT_UMG_Viewmodel_ExecuteBinding_ValueChanged, STATGROUP_UMG_Viewmodel);
 DECLARE_CYCLE_STAT(TEXT("ExecuteBinding Delayed"), STAT_UMG_Viewmodel_ExecuteBinding_Delayed, STATGROUP_UMG_Viewmodel);
 DECLARE_CYCLE_STAT(TEXT("ExecuteBinding Tick"), STAT_UMG_Viewmodel_ExecuteBinding_Tick, STATGROUP_UMG_Viewmodel);
@@ -47,6 +48,9 @@ UMVVMViewModelCollectionObject* GetGlobalCollection(UUserWidget* UserWidget)
 	}
 	return nullptr;
 }
+
+using FRecursiveDetctionElement = TTuple<const UObject*, FMVVMViewClass_BindingKey>;
+TArray<FRecursiveDetctionElement> RecursiveDetector;
 } // namespace
 
 void UMVVMView::ConstructView(const UMVVMViewClass* InClassExtension)
@@ -532,13 +536,6 @@ bool UMVVMView::EvaluateSource(FMVVMViewClass_SourceKey SourceKey)
 }
 
 
-namespace UE::MVVM::Private
-{
-	using FRecursiveDetctionElement = TTuple<const UObject*, FMVVMViewClass_BindingKey>;
-	TArray<FRecursiveDetctionElement> RecursiveDetector;
-}
-
-
 void UMVVMView::HandledLibraryBindingValueChanged(UObject* InSource, UE::FieldNotification::FFieldId InFieldId)
 {
 	SCOPE_CYCLE_COUNTER(STAT_UMG_Viewmodel_ExecuteBinding_ValueChanged);
@@ -585,42 +582,44 @@ void UMVVMView::HandledLibraryBindingValueChanged(UObject* InSource, UE::FieldNo
 		{
 			if (SourceBinding.GetFieldId().GetFieldName() == InFieldId.GetName())
 			{
-				const FMVVMViewClass_Binding& ClassBinding = ClassExtension->GetBinding(SourceBinding.GetBindingKey());
-				if (ensure(ClassBinding.IsOneWay()))
+				ExecuteBindingInternal(SourceBinding);
+			}
+		}
+	}
+}
+
+
+void UMVVMView::ExecuteBindingInternal(const FMVVMViewClass_SourceBinding& SourceBinding) const
+{
+	const FMVVMViewClass_Binding& ClassBinding = ClassExtension->GetBinding(SourceBinding.GetBindingKey());
+	if (ensure(ClassBinding.IsOneWay()))
+	{
+		const EMVVMExecutionMode ExecutionMode = ClassBinding.GetExecuteMode();
+		if (ExecutionMode == EMVVMExecutionMode::Immediate)
+		{
+			// Test for recursivity
+			const UMVVMView* Self = this;
+			if (UE::MVVM::Private::RecursiveDetector.FindByPredicate([Self, InCompiledBindingIndex = SourceBinding.GetBindingKey()](const UE::MVVM::Private::FRecursiveDetctionElement& Element)
 				{
-					const EMVVMExecutionMode ExecutionMode = ClassBinding.GetExecuteMode();
-					if (ExecutionMode == EMVVMExecutionMode::Immediate)
-					{
-						// Test for recursivity
-						const UMVVMView* Self = this;
-						if (UE::MVVM::Private::RecursiveDetector.FindByPredicate([Self, InCompiledBindingIndex = SourceBinding.GetBindingKey()](const UE::MVVM::Private::FRecursiveDetctionElement& Element)
-							{
-								return Element.Get<0>() == Self && Element.Get<1>() == InCompiledBindingIndex;
-							}) != nullptr)
-						{
-							ensureAlwaysMsgf(false, TEXT("Recursive binding detected"));
-							//Todo add more infos. Callstack maybe? Log the chain?
-							UE::MVVM::FMessageLog Log(Self->GetUserWidget());
-							Log.Warning(LOCTEXT("RecursionDetected", "A recursive binding was detected (ie. A->B->C->A->B->C) at runtime."));
-							return;
-						}
-
-						{
-							UE::MVVM::Private::RecursiveDetector.Emplace(this, SourceBinding.GetBindingKey());
-							ExecuteBindingImmediately(ClassBinding, SourceBinding.GetBindingKey());
-							UE::MVVM::Private::RecursiveDetector.Pop();
-						}
-					}
-					else if (ExecutionMode == EMVVMExecutionMode::Delayed)
-					{
-						GEngine->GetEngineSubsystem<UMVVMBindingSubsystem>()->AddDelayedBinding(this, SourceBinding.GetBindingKey());
-					}
-					else
-					{
-						ensureMsgf(false, TEXT("We should not have registered the binding."));
-					}
-				}
+					return Element.Get<0>() == Self && Element.Get<1>() == InCompiledBindingIndex;
+				}) != nullptr)
+			{
+				ensureAlwaysMsgf(false, TEXT("Recursive binding detected"));
+				//Todo add more infos. Callstack maybe? Log the chain?
+				UE::MVVM::FMessageLog Log(Self->GetUserWidget());
+				Log.Warning(LOCTEXT("RecursionDetected", "A recursive binding was detected (ie. A->B->C->A->B->C) at runtime."));
+				return;
 			}
+
+				{
+					UE::MVVM::Private::RecursiveDetector.Emplace(this, SourceBinding.GetBindingKey());
+					ExecuteBindingImmediately(ClassBinding, SourceBinding.GetBindingKey());
+					UE::MVVM::Private::RecursiveDetector.Pop();
+				}
+		}
+		else if (ExecutionMode == EMVVMExecutionMode::Delayed)
+		{
+			GEngine->GetEngineSubsystem<UMVVMBindingSubsystem>()->AddDelayedBinding(this, SourceBinding.GetBindingKey());
 		}
 	}
 }
@@ -808,9 +807,6 @@ bool UMVVMView::SetSourceInternal(FMVVMViewClass_SourceKey ClassSourceKey, TScri
 	}
 
 	UObject* PreviousValue = ViewSource.Source;
-	bool bPreviousSourceInitialized = ViewSource.bSourceInitialized;
-	bool bPreviousBindingsInitialized = ViewSource.bBindingsInitialized;
-
 	{
 		// Sanity check. Test if the bitfield matches the cached value.
 		const bool bIsValid = PreviousValue != nullptr;
@@ -820,6 +816,9 @@ bool UMVVMView::SetSourceInternal(FMVVMViewClass_SourceKey ClassSourceKey, TScri
 
 	if (PreviousValue != NewValue.GetObject())
 	{
+		bool bPreviousBindingsInitialized = ViewSource.bBindingsInitialized;
+		bool bPreviousSourceInitialized = ViewSource.bSourceInitialized;
+
 		if (ViewSource.bBindingsInitialized)
 		{
 			UninitializeSourceBindings(ClassSourceKey, ClassSource, ViewSource);
@@ -851,7 +850,7 @@ bool UMVVMView::SetSourceInternal(FMVVMViewClass_SourceKey ClassSourceKey, TScri
 				}
 			}
 
-			// initialize bindings and run time
+			// initialize bindings
 			InitializeSourceBindings(ViewSourceKey, true);
 		}
 
@@ -859,10 +858,74 @@ bool UMVVMView::SetSourceInternal(FMVVMViewClass_SourceKey ClassSourceKey, TScri
 		UE::MVVM::FDebugging::BroadcastViewSourceValueChanged(this, ClassSourceKey, ViewSourceKey);
 #endif
 	}
+	else if (ClassSource.AlwaysExecuteBindingsOnSetSource() && ViewSource.bBindingsInitialized)
+	{
+		ExecuteViewModelBindingsInternal(ClassSourceKey);
+	}
+
 	return true;
 }
 
 
+bool UMVVMView::ExecuteViewModelBindings(FName ViewModelName)
+{
+	if (ViewModelName.IsNone())
+	{
+		UE::MVVM::FMessageLog Log(GetUserWidget());
+		Log.Error(LOCTEXT("SetViewModelInvalidName", "The viewmodel name is empty."));
+		return false;
+	}
+
+	if (ClassExtension == nullptr)
+	{
+		UE::MVVM::FMessageLog Log(GetUserWidget());
+		Log.Error(LOCTEXT("SetViewModelInvalidClass", "The view is not constructed."));
+		return false;
+	}
+
+	const int32 ClassSourceIndex = ClassExtension->GetSources().IndexOfByPredicate([ViewModelName](const FMVVMViewClass_Source& Other)
+		{
+			return Other.GetName() == ViewModelName;
+		});
+
+	if (ClassSourceIndex == INDEX_NONE)
+	{
+		UE::MVVM::FMessageLog Log(GetUserWidget());
+		Log.Error(LOCTEXT("SetViewModelViewModelNameNotFound", "The viewmodel name could not be found."));
+		return false;
+	}
+
+	ExecuteViewModelBindingsInternal(FMVVMViewClass_SourceKey(ClassSourceIndex));
+	return true;
+}
+
+
+void UMVVMView::ExecuteViewModelBindingsInternal(FMVVMViewClass_SourceKey ClassSourceKey)
+{
+	SCOPE_CYCLE_COUNTER(STAT_UMG_Viewmodel_ExecuteViewModel);
+
+	const FMVVMViewClass_Source& ClassSource = ClassExtension->GetSource(ClassSourceKey);
+
+	// Run all evaluates
+	if (ClassSource.HasEvaluateBindings())
+	{
+		for (const FMVVMViewClass_EvaluateSource& ClassEvaluate : ClassExtension->GetEvaluateSources())
+		{
+			if (ClassEvaluate.GetParentSource() == ClassSourceKey)
+			{
+				EvaluateSource(ClassEvaluate.GetSource());
+			}
+		}
+	}
+
+	// Run the bindings
+	for (const FMVVMViewClass_SourceBinding& SourceBinding : ClassSource.GetBindings())
+	{
+		ExecuteBindingInternal(SourceBinding);
+	}
+}
+
+
 void UMVVMView::HandleViewModelCollectionChanged()
 {
 	if (ensure(ClassExtension))
@@ -885,6 +948,7 @@ void UMVVMView::HandleViewModelCollectionChanged()
 	}
 }
 
+
 void UMVVMView::InitializeEvents()
 {
 	if (bEventsInitialized)
diff --git a/Engine/Plugins/Runtime/ModelViewViewModel/Source/ModelViewViewModel/Private/View/MVVMViewClass.cpp b/Engine/Plugins/Runtime/ModelViewViewModel/Source/ModelViewViewModel/Private/View/MVVMViewClass.cpp
--- a/Engine/Plugins/Runtime/ModelViewViewModel/Source/ModelViewViewModel/Private/View/MVVMViewClass.cpp
+++ b/Engine/Plugins/Runtime/ModelViewViewModel/Source/ModelViewViewModel/Private/View/MVVMViewClass.cpp
@@ -302,6 +302,21 @@ FString FMVVMViewClass_Source::ToString(const UMVVMViewClass* ViewClass, FToStri
 			AddPipe();
 			StringBuilder << TEXT("IsViewModel");
 		}
+		if ((Flags & (uint16)EFlags::IsViewModelInstanceExposed) != 0)
+		{
+			AddPipe();
+			StringBuilder << TEXT("EditDefaultOnly");
+		}
+		if (RequireGlobalViewModelCollectionUpdate())
+		{
+			AddPipe();
+			StringBuilder << TEXT("GlobalCollectionAutoUpdate");
+		}
+		if (AlwaysExecuteBindingsOnSetSource())
+		{
+			AddPipe();
+			StringBuilder << TEXT("AlwaysExecuteBindingsOnSetSource");
+		}
 	}
 
 	if (Args.bAddCreationMode)
diff --git a/Engine/Plugins/Runtime/ModelViewViewModel/Source/ModelViewViewModel/Public/View/MVVMView.h b/Engine/Plugins/Runtime/ModelViewViewModel/Source/ModelViewViewModel/Public/View/MVVMView.h
--- a/Engine/Plugins/Runtime/ModelViewViewModel/Source/ModelViewViewModel/Public/View/MVVMView.h
+++ b/Engine/Plugins/Runtime/ModelViewViewModel/Source/ModelViewViewModel/Public/View/MVVMView.h
@@ -16,6 +16,8 @@ class UMVVMViewClass;
 struct FMVVMViewClass_Binding;
 struct FMVVMViewClass_Event;
 struct FMVVMViewClass_Source;
+struct FMVVMViewClass_SourceBinding;
+struct FMVVMViewClass_SourceKey;
 
 /**
  * Instance FMVVMViewClass_Source for the UUserWdiget
@@ -162,6 +164,12 @@ public:
 	UFUNCTION(BlueprintCallable, Category = "View")
 	bool SetViewModelByClass(TScriptInterface<INotifyFieldValueChanged> NewValue);
 
+	/**
+	 * Execute all the bindings that use the viewmodel.
+	 */
+	UFUNCTION(BlueprintCallable, Category = "View")
+	bool ExecuteViewModelBindings(FName ViewModelName);
+
 private:
 	/**
 	 * Set the viewmodel of the specified name.
@@ -183,7 +191,9 @@ private:
 	void InitializeSourceBindingsCommon();
 	void UninitializeSourceBindings(FMVVMViewClass_SourceKey SourceKey, const FMVVMViewClass_Source& ClassSource, FMVVMView_Source& ViewSource);
 	void HandledLibraryBindingValueChanged(UObject* InSource, UE::FieldNotification::FFieldId InFieldId);
+	void ExecuteBindingInternal(const FMVVMViewClass_SourceBinding& SourceBinding) const;
 	void ExecuteBindingImmediately(const FMVVMViewClass_Binding& ClassBinding, FMVVMViewClass_BindingKey KeyForLog) const;
+	void ExecuteViewModelBindingsInternal(FMVVMViewClass_SourceKey SourceKey);
 
 	//~ evaluate source
 	bool EvaluateSource(FMVVMViewClass_SourceKey SourceIndex);
diff --git a/Engine/Plugins/Runtime/ModelViewViewModel/Source/ModelViewViewModel/Public/View/MVVMViewClass.h b/Engine/Plugins/Runtime/ModelViewViewModel/Source/ModelViewViewModel/Public/View/MVVMViewClass.h
--- a/Engine/Plugins/Runtime/ModelViewViewModel/Source/ModelViewViewModel/Public/View/MVVMViewClass.h
+++ b/Engine/Plugins/Runtime/ModelViewViewModel/Source/ModelViewViewModel/Public/View/MVVMViewClass.h
@@ -298,6 +298,12 @@ public:
 	{
 		return (Flags & (uint16)EFlags::GlobalViewModelCollectionUpdate) != 0;
 	}
+
+	/** When a source is manually set, then always execute that source the bindings. */
+	bool AlwaysExecuteBindingsOnSetSource() const
+	{
+		return (Flags & (uint16)EFlags::AlwaysExecuteBindingsOnSetSource) != 0;;
+	}
 	
 	/**
 	 * The source GetOrCreateInstance can fail.
@@ -409,6 +415,7 @@ private:
 		IsViewModel = 1 << 9,
 		IsViewModelInstanceExposed = 1 << 10,
 		GlobalViewModelCollectionUpdate = 1 << 11,
+		AlwaysExecuteBindingsOnSetSource = 1 << 12,
 	};
 
 	UPROPERTY(VisibleAnywhere, Category = "View")
@@ -520,13 +527,13 @@ public:
 	{
 		return bInitializeBindingsOnConstruct;
 	}
-	
+
 	/** Should it automatically register the events when the view is constructed. */
 	[[nodiscard]] bool DoesInitializeEventsOnConstruct() const
 	{
 		return bInitializeEventsOnConstruct;
 	}
-	
+
 	/** Should the view listen to modification made in the global viewmodel collection. */
 	[[nodiscard]] bool DoesListenToViewModelCollectionChanged() const
 	{
diff --git a/Engine/Plugins/Runtime/ModelViewViewModel/Source/ModelViewViewModelBlueprint/Private/MVVMDeveloperProjectSettings.cpp b/Engine/Plugins/Runtime/ModelViewViewModel/Source/ModelViewViewModelBlueprint/Private/MVVMDeveloperProjectSettings.cpp
--- a/Engine/Plugins/Runtime/ModelViewViewModel/Source/ModelViewViewModelBlueprint/Private/MVVMDeveloperProjectSettings.cpp
+++ b/Engine/Plugins/Runtime/ModelViewViewModel/Source/ModelViewViewModelBlueprint/Private/MVVMDeveloperProjectSettings.cpp
@@ -29,8 +29,10 @@ UMVVMDeveloperProjectSettings::UMVVMDeveloperProjectSettings()
 	AllowedContextCreationType.Add(EMVVMBlueprintViewModelContextCreationType::PropertyPath);
 	AllowedContextCreationType.Add(EMVVMBlueprintViewModelContextCreationType::Resolver);
 
-	AllowedClassForConversionFunctions.Add(UBlueprintFunctionLibrary::StaticClass());
-	AllowedClassForConversionFunctions.Add(UK2Node_FormatText::StaticClass());
+	FTopLevelAssetPath BlueprintFunctionLibrary = FTopLevelAssetPath("/Script/Engine", "BlueprintFunctionLibrary");
+	FTopLevelAssetPath FormatText = FTopLevelAssetPath("/Script/BlueprintGraph", "K2Node_FormatText");
+	AllowedClassForConversionFunctions.Add(FSoftClassPath(BlueprintFunctionLibrary.ToString()));
+	AllowedClassForConversionFunctions.Add(FSoftClassPath(FormatText.ToString()));
 }
 
 FName UMVVMDeveloperProjectSettings::GetCategoryName() const
diff --git a/Engine/Plugins/Runtime/ModelViewViewModel/Source/ModelViewViewModelBlueprint/Private/MVVMViewBlueprintCompiler.cpp b/Engine/Plugins/Runtime/ModelViewViewModel/Source/ModelViewViewModelBlueprint/Private/MVVMViewBlueprintCompiler.cpp
--- a/Engine/Plugins/Runtime/ModelViewViewModel/Source/ModelViewViewModelBlueprint/Private/MVVMViewBlueprintCompiler.cpp
+++ b/Engine/Plugins/Runtime/ModelViewViewModel/Source/ModelViewViewModelBlueprint/Private/MVVMViewBlueprintCompiler.cpp
@@ -2175,6 +2175,7 @@ void FMVVMViewBlueprintCompiler::CompileSources(const FCompiledBindingLibraryCom
 		bool bIsUserWidgetProperty = !SourceCreatorContext.DynamicContext.IsValid();
 		bool bExposeInstanceInEditor = false;
 		bool bGlobalViewModelCollectionUpdate = false;
+		bool bAlwaysExecuteBindingsOnSetSource =  GetDefault<UMVVMDeveloperProjectSettings>()->bForceExecuteBindingsOnSetSource;
 
 		if (ViewModelContext.CreationType == EMVVMBlueprintViewModelContextCreationType::Manual)
 		{
@@ -2277,6 +2278,12 @@ void FMVVMViewBlueprintCompiler::CompileSources(const FCompiledBindingLibraryCom
 			}
 		}
 
+		if (ViewModelContext.bOverrideForceExecuteBindingsOnSetSource)
+		{
+			bAlwaysExecuteBindingsOnSetSource = ViewModelContext.bForceExecuteBindingsOnSetSource;
+		}
+		bAlwaysExecuteBindingsOnSetSource = bCanBeSet && bAlwaysExecuteBindingsOnSetSource;
+
 		CompiledSourceCreator.Flags = 0;
 		CompiledSourceCreator.Flags |= bCreateInstance ? (uint16)FMVVMViewClass_Source::EFlags::TypeCreateInstance : 0;
 		CompiledSourceCreator.Flags |= bIsUserWidgetProperty ? (uint16)FMVVMViewClass_Source::EFlags::IsUserWidgetProperty : 0;
@@ -2287,6 +2294,7 @@ void FMVVMViewBlueprintCompiler::CompileSources(const FCompiledBindingLibraryCom
 		CompiledSourceCreator.Flags |= (uint16)FMVVMViewClass_Source::EFlags::IsViewModel;
 		CompiledSourceCreator.Flags |= bExposeInstanceInEditor ? (uint16)FMVVMViewClass_Source::EFlags::IsViewModelInstanceExposed : 0;
 		CompiledSourceCreator.Flags |= bGlobalViewModelCollectionUpdate ? (uint16)FMVVMViewClass_Source::EFlags::GlobalViewModelCollectionUpdate : 0;
+		CompiledSourceCreator.Flags |= bAlwaysExecuteBindingsOnSetSource ? (uint16)FMVVMViewClass_Source::EFlags::AlwaysExecuteBindingsOnSetSource : 0;
 
 		{
 			FSortData SortData;
diff --git a/Engine/Plugins/Runtime/ModelViewViewModel/Source/ModelViewViewModelBlueprint/Public/MVVMBlueprintView.h b/Engine/Plugins/Runtime/ModelViewViewModel/Source/ModelViewViewModelBlueprint/Public/MVVMBlueprintView.h
--- a/Engine/Plugins/Runtime/ModelViewViewModel/Source/ModelViewViewModelBlueprint/Public/MVVMBlueprintView.h
+++ b/Engine/Plugins/Runtime/ModelViewViewModel/Source/ModelViewViewModelBlueprint/Public/MVVMBlueprintView.h
@@ -179,7 +179,7 @@ public:
 	TArray<TObjectPtr<UEdGraph>> TemporaryGraph;
 
 private:
-	UPROPERTY()
+	UPROPERTY(Instanced)
 	TObjectPtr<UMVVMBlueprintViewSettings> Settings;
 
 	UPROPERTY(EditAnywhere, Category = "Viewmodel")
diff --git a/Engine/Plugins/Runtime/ModelViewViewModel/Source/ModelViewViewModelBlueprint/Public/MVVMBlueprintViewModelContext.h b/Engine/Plugins/Runtime/ModelViewViewModel/Source/ModelViewViewModelBlueprint/Public/MVVMBlueprintViewModelContext.h
--- a/Engine/Plugins/Runtime/ModelViewViewModel/Source/ModelViewViewModelBlueprint/Public/MVVMBlueprintViewModelContext.h
+++ b/Engine/Plugins/Runtime/ModelViewViewModel/Source/ModelViewViewModelBlueprint/Public/MVVMBlueprintViewModelContext.h
@@ -164,6 +164,17 @@ public:
 	UPROPERTY(EditAnywhere, Category = "Viewmodel", AdvancedDisplay)
 	bool bGlobalViewModelCollectionUpdate = false;
 
+	UPROPERTY()
+	bool bOverrideForceExecuteBindingsOnSetSource = false;
+
+	/**
+	 * When a viewmodel is set manually and the viewmodel already initialized, then always execute the bindings associated with that viewmodel.
+	 * For performance and to keep the same pattern in all UMG, the bindings are usually skip if the new viewmodel value match the previous viewmodel value.
+	 * This behavior can be desired if the widget is inside a pool or a binding has a side effect with another widget.
+	 */
+	UPROPERTY(EditAnywhere, Category = "View", meta = (EditCondition = "bOverrideForceExecuteBindingsOnSetSource"))
+	bool bForceExecuteBindingsOnSetSource = false;
+
 	/** Can change the name in the editor. */
 	UPROPERTY()
 	bool bCanRename = true;
diff --git a/Engine/Plugins/Runtime/ModelViewViewModel/Source/ModelViewViewModelBlueprint/Public/MVVMDeveloperProjectSettings.h b/Engine/Plugins/Runtime/ModelViewViewModel/Source/ModelViewViewModelBlueprint/Public/MVVMDeveloperProjectSettings.h
--- a/Engine/Plugins/Runtime/ModelViewViewModel/Source/ModelViewViewModelBlueprint/Public/MVVMDeveloperProjectSettings.h
+++ b/Engine/Plugins/Runtime/ModelViewViewModel/Source/ModelViewViewModelBlueprint/Public/MVVMDeveloperProjectSettings.h
@@ -70,7 +70,6 @@ public:
 	virtual FName GetCategoryName() const override;
 	virtual FText GetSectionText() const override;
 
-
 	bool PropertyHasFiltering(const UStruct* ObjectStruct, const FProperty* Property) const;
 	bool IsPropertyAllowed(const UBlueprint* Context, const UStruct* ObjectStruct, const FProperty* Property) const;
 	bool IsFunctionAllowed(const UBlueprint* Context, const UClass* ObjectClass, const UFunction* Function) const;
@@ -97,79 +96,86 @@ public:
 
 private:
 	/** Permission list for filtering which properties are visible in UI. */
-	UPROPERTY(EditAnywhere, config, Category = "Viewmodel")
+	UPROPERTY(EditAnywhere, config, Category = "UX")
 	TMap<FSoftClassPath, FMVVMDeveloperProjectWidgetSettings> FieldSelectorPermissions;
 
 	/** Permission list for filtering which execution mode is allowed. */
-	UPROPERTY(EditAnywhere, config, Category = "Viewmodel")
+	UPROPERTY(EditAnywhere, config, Category = "Defaults")
 	TSet<EMVVMExecutionMode> AllowedExecutionMode;
 	
 	/** Permission list for filtering which context creation type is allowed. */
-	UPROPERTY(EditAnywhere, config, Category = "Viewmodel")
+	UPROPERTY(EditAnywhere, config, Category = "Defaults")
 	TSet<EMVVMBlueprintViewModelContextCreationType> AllowedContextCreationType;
 
 public:
 	/** Binding can be made from the DetailView Bind option. */
-	UPROPERTY(EditAnywhere, config, Category = "Viewmodel")
+	UPROPERTY(EditAnywhere, config, Category = "UX")
 	bool bAllowBindingFromDetailView = true;
 
 	/** When generating a source in the viewmodel editor, allow the compiler to generate a setter function. */
-	UPROPERTY(EditAnywhere, config, Category = "Viewmodel")
+	UPROPERTY(EditAnywhere, config, Category = "Features", meta=(DisplayName="Allow Generated Viewmodel Setter"))
 	bool bAllowGeneratedViewModelSetter = true;
 
 	/** When generating a binding with a long source path, allow the compiler to generate a new viewmodel source. */
-	UPROPERTY(EditAnywhere, config, Category = "Viewmodel")
+	UPROPERTY(EditAnywhere, config, Category = "Features")
 	bool bAllowLongSourcePath = true;
 	
 	/** For the binding list widget, allow the user to edit the binding in the detail view. */
-	UPROPERTY(EditAnywhere, config, Category = "Viewmodel")
+	UPROPERTY(EditAnywhere, config, Category = "UX")
 	bool bShowDetailViewOptionInBindingPanel = true;
 	
 	/** For the binding list widget and the viewmodel panel, allow the user to edit the view settings in the detail view. */
-	UPROPERTY(EditAnywhere, config, Category = "Viewmodel")
+	UPROPERTY(EditAnywhere, config, Category = "UX")
 	bool bShowViewSettings = true;
 
 	/** For the binding list widget, allow the user to generate a copy of the binding/event graph. */
-	UPROPERTY(EditAnywhere, config, Category = "Viewmodel")
+	UPROPERTY(EditAnywhere, config, Category = "UX")
 	bool bShowDeveloperGenerateGraphSettings = true;
 
-	/** When a conversion function requires a wrapper graph, add and save the generated graph to the blueprint. */
-	UPROPERTY(EditAnywhere, config, Category = "Viewmodel")
+	/**
+	 * When a conversion function requires a wrapper graph, add and save the generated graph to the blueprint.
+	 * It is strongly suggested to to use this feature. It can easly break deprecation and the graph will not auto update with new features.
+	 */
+	UPROPERTY(EditAnywhere, config, Category = "Features")
 	bool bAllowConversionFunctionGeneratedGraphInEditor = false;
 
 	/** When binding to a multicast delegate property, allow to create an event. */
-	UPROPERTY(EditAnywhere, config, Category = "Viewmodel")
+	UPROPERTY(EditAnywhere, config, Category = "Features")
 	bool bAllowBindingEvent = true;
-	
+
 	/** Allow to create an instanced viewmodel directly in the view editor. */
-	UPROPERTY(EditAnywhere, config, Category = "Viewmodel")
+	UPROPERTY(EditAnywhere, config, Category = "Features", meta=(DisplayName="Experimental - Can Create Viewmodel In View"))
 	bool bCanCreateViewModelInView = false;
 
-	/** 
+	/**
 	 * When a viewmodel is set to Create Instance, allow modifying the viewmodel instance in the editor on all instances of the owning widget.
 	 * The per-viewmodel setting "Expose Instance In Editor" overrides this.
 	 */
-	UPROPERTY(EditAnywhere, config, Category = "Viewmodel")
+	UPROPERTY(EditAnywhere, config, Category = "Features")
 	bool bExposeViewModelInstanceInEditor = false;
 
 	/** Permission list for filtering which execution mode is allowed. */
-	UPROPERTY(EditAnywhere, config, Category = "Viewmodel")
+	UPROPERTY(EditAnywhere, config, Category = "Defaults")
 	EMVVMDeveloperConversionFunctionFilterType ConversionFunctionFilter = EMVVMDeveloperConversionFunctionFilterType::BlueprintActionRegistry;
 
 	/** Classes to include in conversion function list. It includes the child class. */
-	UPROPERTY(EditAnywhere, config, Category = "Viewmodel", meta = (EditCondition = "ConversionFunctionFilter == EMVVMDeveloperConversionFunctionFilterType::AllowedList"))
+	UPROPERTY(EditAnywhere, config, Category = "Defaults", meta = (EditCondition = "ConversionFunctionFilter == EMVVMDeveloperConversionFunctionFilterType::AllowedList"))
 	TSet<FSoftClassPath> AllowedClassForConversionFunctions;
 
 	/** Classes excluded for conversion function list. */
-	UPROPERTY(EditAnywhere, config, Category = "Viewmodel", meta = (EditCondition = "ConversionFunctionFilter == EMVVMDeveloperConversionFunctionFilterType::AllowedList"))
+	UPROPERTY(EditAnywhere, config, Category = "Defaults", meta = (EditCondition = "ConversionFunctionFilter == EMVVMDeveloperConversionFunctionFilterType::AllowedList"))
 	TSet<FSoftClassPath> DeniedClassForConversionFunctions;
 
 	/** Modules excluded for conversion function list. ie. "/Script/MyModule" */
-	UPROPERTY(EditAnywhere, config, Category = "Viewmodel", meta = (EditCondition = "ConversionFunctionFilter == EMVVMDeveloperConversionFunctionFilterType::AllowedList"))
+	UPROPERTY(EditAnywhere, config, Category = "Defaults", meta = (EditCondition = "ConversionFunctionFilter == EMVVMDeveloperConversionFunctionFilterType::AllowedList"))
 	TSet<FName> DeniedModuleForConversionFunctions;
 
+	/** The default value of UMVVMBlueprintViewSettings::bForceExecuteBindingOnSetSource. */
+	UPROPERTY(EditAnywhere, config, Category = "Defaults")
+	bool bForceExecuteBindingsOnSetSource = false;
+
 	/** Settings for filtering the list of available properties and functions on binding creation. */
-	UPROPERTY(EditAnywhere, config, Category = "Viewmodel")
+	UPROPERTY(EditAnywhere, config, Category = "UX")
 	FMVVMViewBindingFilterSettings FilterSettings;
 };
 
diff --git a/Engine/Plugins/Runtime/ModelViewViewModel/Source/ModelViewViewModelEditor/Private/Customizations/MVVMBlueprintViewModelContextCustomization.cpp b/Engine/Plugins/Runtime/ModelViewViewModel/Source/ModelViewViewModelEditor/Private/Customizations/MVVMBlueprintViewModelContextCustomization.cpp
--- a/Engine/Plugins/Runtime/ModelViewViewModel/Source/ModelViewViewModelEditor/Private/Customizations/MVVMBlueprintViewModelContextCustomization.cpp
+++ b/Engine/Plugins/Runtime/ModelViewViewModel/Source/ModelViewViewModelEditor/Private/Customizations/MVVMBlueprintViewModelContextCustomization.cpp
@@ -194,10 +194,19 @@ void FBlueprintViewModelContextDetailCustomization::CustomizeChildren(TSharedRef
 	bool bCanEdit = ContextPtr->bCanEdit;
 	bool bCanRename = ContextPtr->bCanRename;
 
+	// Reset the value to what the user expect to see. It is not used in by the compiler.
+	if (!ContextPtr->bOverrideForceExecuteBindingsOnSetSource)
+	{
+		ContextPtr->bForceExecuteBindingsOnSetSource = GetDefault<UMVVMDeveloperProjectSettings>()->bForceExecuteBindingsOnSetSource;
+	}
+	
 	NotifyFieldValueClassHandle = PropertyHandle->GetChildHandle(TEXT("NotifyFieldValueClass"), false);
 	PropertyPathHandle = PropertyHandle->GetChildHandle(GET_MEMBER_NAME_CHECKED(FMVVMBlueprintViewModelContext, ViewModelPropertyPath), false);
 	CreationTypeHandle = PropertyHandle->GetChildHandle(GET_MEMBER_NAME_CHECKED(FMVVMBlueprintViewModelContext, CreationType), false);
 	ViewModelNameHandle = PropertyHandle->GetChildHandle(GET_MEMBER_NAME_CHECKED(FMVVMBlueprintViewModelContext, ViewModelName), false);
+	OptionalHandle = PropertyHandle->GetChildHandle(GET_MEMBER_NAME_CHECKED(FMVVMBlueprintViewModelContext, bOptional), false);
+	CreateSetterFunctionHandle = PropertyHandle->GetChildHandle(GET_MEMBER_NAME_CHECKED(FMVVMBlueprintViewModelContext, bCreateSetterFunction), false);
+	ForceExecuteBindingsOnSetSourceHandle = PropertyHandle->GetChildHandle(GET_MEMBER_NAME_CHECKED(FMVVMBlueprintViewModelContext, bForceExecuteBindingsOnSetSource), false);
 
 	if (ensure(NotifyFieldValueClassHandle))
 	{
@@ -219,6 +228,11 @@ void FBlueprintViewModelContextDetailCustomization::CustomizeChildren(TSharedRef
 		CreationTypeHandle->SetOnPropertyValueChanged(FSimpleDelegate::CreateSP(this, &FBlueprintViewModelContextDetailCustomization::HandleCreationTypeChanged));
 	}
 
+	if (ensure(CreateSetterFunctionHandle))
+	{
+		CreateSetterFunctionHandle->SetOnPropertyValueChanged(FSimpleDelegate::CreateSP(this, &FBlueprintViewModelContextDetailCustomization::HandleCreateSetterFunctionChanged));
+	}
+
 	if (ensure(ViewModelNameHandle))
 	{
 		IDetailPropertyRow& PropertyRow = ChildBuilder.AddProperty(ViewModelNameHandle.ToSharedRef())
@@ -374,7 +388,6 @@ void FBlueprintViewModelContextDetailCustomization::CustomizeChildren(TSharedRef
 
 		if (GetDefault<UMVVMDeveloperProjectSettings>()->bAllowGeneratedViewModelSetter)
 		{
-			TSharedPtr<IPropertyHandle> CreateSetterFunctionHandle = PropertyHandle->GetChildHandle(GET_MEMBER_NAME_CHECKED(FMVVMBlueprintViewModelContext, bCreateSetterFunction), false);
 			if (ensure(CreateSetterFunctionHandle))
 			{
 				ChildBuilder.AddProperty(CreateSetterFunctionHandle.ToSharedRef())
@@ -417,7 +430,6 @@ void FBlueprintViewModelContextDetailCustomization::CustomizeChildren(TSharedRef
 					}));
 		}
 
-		TSharedPtr<IPropertyHandle> OptionalHandle = PropertyHandle->GetChildHandle(GET_MEMBER_NAME_CHECKED(FMVVMBlueprintViewModelContext, bOptional), false);
 		if (ensure(OptionalHandle))
 		{
 			ChildBuilder.AddProperty(OptionalHandle.ToSharedRef())
@@ -429,6 +441,17 @@ void FBlueprintViewModelContextDetailCustomization::CustomizeChildren(TSharedRef
 						return bResult && bCanEdit;
 					}));
 		}
+
+		if (ensure(ForceExecuteBindingsOnSetSourceHandle))
+		{
+			ChildBuilder.AddProperty(ForceExecuteBindingsOnSetSourceHandle.ToSharedRef())
+				.IsEnabled(bCanEdit)
+				.Visibility(MakeAttributeLambda([ContextPtr]()
+					{
+						bool bResult = ContextPtr->bCreateSetterFunction;
+						return bResult ? EVisibility::Visible : EVisibility::Collapsed;
+					}));
+		}
 	}
 	else
 	{
@@ -487,8 +510,30 @@ void FBlueprintViewModelContextDetailCustomization::HandleCreationTypeChanged()
 		if (FMVVMBlueprintViewModelContext* ContextPtr = Private::GetViewModelContext(ContextHandle.ToSharedRef()))
 		{
 			const bool bIsManual = (EMVVMBlueprintViewModelContextCreationType)NewValue == EMVVMBlueprintViewModelContextCreationType::Manual;
-			ContextPtr->bOptional = bIsManual;
-			ContextPtr->bCreateSetterFunction = bIsManual;
+			if (ContextPtr->bOptional != bIsManual)
+			{
+				OptionalHandle->SetValue(bIsManual);
+			}
+			if (ContextPtr->bCreateSetterFunction != bIsManual)
+			{
+				CreateSetterFunctionHandle->SetValue(bIsManual);
+			}
+		}
+	}
+}
+
+void FBlueprintViewModelContextDetailCustomization::HandleCreateSetterFunctionChanged()
+{
+	bool bNewCreateSetterFunctionHandle = false;
+	if (CreateSetterFunctionHandle->GetValue(bNewCreateSetterFunctionHandle) == FPropertyAccess::Success)
+	{
+		if (FMVVMBlueprintViewModelContext* ContextPtr = Private::GetViewModelContext(ContextHandle.ToSharedRef()))
+		{
+			if (ContextPtr->bOverrideForceExecuteBindingsOnSetSource && !bNewCreateSetterFunctionHandle)
+			{
+				ContextPtr->bOverrideForceExecuteBindingsOnSetSource = false;
+				ForceExecuteBindingsOnSetSourceHandle->SetValue(GetDefault<UMVVMDeveloperProjectSettings>()->bForceExecuteBindingsOnSetSource);
+			}
 		}
 	}
 }
diff --git a/Engine/Plugins/Runtime/ModelViewViewModel/Source/ModelViewViewModelEditor/Private/Customizations/MVVMBlueprintViewModelContextCustomization.h b/Engine/Plugins/Runtime/ModelViewViewModel/Source/ModelViewViewModelEditor/Private/Customizations/MVVMBlueprintViewModelContextCustomization.h
--- a/Engine/Plugins/Runtime/ModelViewViewModel/Source/ModelViewViewModelEditor/Private/Customizations/MVVMBlueprintViewModelContextCustomization.h
+++ b/Engine/Plugins/Runtime/ModelViewViewModel/Source/ModelViewViewModelEditor/Private/Customizations/MVVMBlueprintViewModelContextCustomization.h
@@ -54,6 +54,7 @@ private:
 	FText GetClassName() const;
 	void HandleClassChanged();
 	void HandleCreationTypeChanged();
+	void HandleCreateSetterFunctionChanged();
 	TSharedRef<SWidget> HandleClassGetMenuContent();
 	void HandleClassCancelMenu();
 	void HandleClassCommitted(const UClass* SelectedClass);
@@ -69,6 +70,9 @@ private:
 	TSharedPtr<IPropertyHandle> PropertyPathHandle;
 	TSharedPtr<IPropertyHandle> CreationTypeHandle;
 	TSharedPtr<IPropertyHandle> ViewModelNameHandle;
+	TSharedPtr<IPropertyHandle> OptionalHandle;
+	TSharedPtr<IPropertyHandle> CreateSetterFunctionHandle;
+	TSharedPtr<IPropertyHandle> ForceExecuteBindingsOnSetSourceHandle;
 	TArray<EMVVMBlueprintViewModelContextCreationType> AllowedCreationTypes;
 };
 
