MVVM: In project settings allow the abstract class for the class picker.

diff --git a/Engine/Plugins/Runtime/ModelViewViewModel/Source/ModelViewViewModelBlueprint/Public/MVVMDeveloperProjectSettings.h b/Engine/Plugins/Runtime/ModelViewViewModel/Source/ModelViewViewModelBlueprint/Public/MVVMDeveloperProjectSettings.h
--- a/Engine/Plugins/Runtime/ModelViewViewModel/Source/ModelViewViewModelBlueprint/Public/MVVMDeveloperProjectSettings.h
+++ b/Engine/Plugins/Runtime/ModelViewViewModel/Source/ModelViewViewModelBlueprint/Public/MVVMDeveloperProjectSettings.h
@@ -96,7 +96,7 @@ public:
 
 private:
 	/** Permission list for filtering which properties are visible in UI. */
-	UPROPERTY(EditAnywhere, config, Category = "UX")
+	UPROPERTY(EditAnywhere, config, Category = "UX", meta=(AllowAbstract = true))
 	TMap<FSoftClassPath, FMVVMDeveloperProjectWidgetSettings> FieldSelectorPermissions;
 
 	/** Permission list for filtering which execution mode is allowed. */
@@ -159,11 +159,11 @@ public:
 	EMVVMDeveloperConversionFunctionFilterType ConversionFunctionFilter = EMVVMDeveloperConversionFunctionFilterType::BlueprintActionRegistry;
 
 	/** Classes to include in conversion function list. It includes the child class. */
-	UPROPERTY(EditAnywhere, config, Category = "Defaults", meta = (EditCondition = "ConversionFunctionFilter == EMVVMDeveloperConversionFunctionFilterType::AllowedList"))
+	UPROPERTY(EditAnywhere, config, Category = "Defaults", meta = (AllowAbstract = true, EditCondition = "ConversionFunctionFilter == EMVVMDeveloperConversionFunctionFilterType::AllowedList"))
 	TSet<FSoftClassPath> AllowedClassForConversionFunctions;
 
 	/** Classes excluded for conversion function list. */
-	UPROPERTY(EditAnywhere, config, Category = "Defaults", meta = (EditCondition = "ConversionFunctionFilter == EMVVMDeveloperConversionFunctionFilterType::AllowedList"))
+	UPROPERTY(EditAnywhere, config, Category = "Defaults", meta = (AllowAbstract = true, EditCondition = "ConversionFunctionFilter == EMVVMDeveloperConversionFunctionFilterType::AllowedList"))
 	TSet<FSoftClassPath> DeniedClassForConversionFunctions;
 
 	/** Modules excluded for conversion function list. ie. "/Script/MyModule" */
