MVVM: Parse the property path that start with "self" to check if it actually start with a widet or a viewmodel.

diff --git a/Engine/Plugins/Runtime/ModelViewViewModel/Source/ModelViewViewModelBlueprint/Private/MVVMPropertyPath.cpp b/Engine/Plugins/Runtime/ModelViewViewModel/Source/ModelViewViewModelBlueprint/Private/MVVMPropertyPath.cpp
--- a/Engine/Plugins/Runtime/ModelViewViewModel/Source/ModelViewViewModelBlueprint/Private/MVVMPropertyPath.cpp
+++ b/Engine/Plugins/Runtime/ModelViewViewModel/Source/ModelViewViewModelBlueprint/Private/MVVMPropertyPath.cpp
@@ -100,6 +100,12 @@ FMVVMBlueprintFieldPath::FMVVMBlueprintFieldPath(const UBlueprint* InContext, UE
 }
 
 
+FName FMVVMBlueprintFieldPath::GetRawFieldName() const
+{
+	return BindingReference.GetMemberName();
+}
+
+
 FName FMVVMBlueprintFieldPath::GetFieldName(const UClass* InContext) const
 {
 	UE::MVVM::FMVVMConstFieldVariant Result = GetField(InContext);
diff --git a/Engine/Plugins/Runtime/ModelViewViewModel/Source/ModelViewViewModelBlueprint/Private/MVVMViewBlueprintCompiler.cpp b/Engine/Plugins/Runtime/ModelViewViewModel/Source/ModelViewViewModelBlueprint/Private/MVVMViewBlueprintCompiler.cpp
--- a/Engine/Plugins/Runtime/ModelViewViewModel/Source/ModelViewViewModelBlueprint/Private/MVVMViewBlueprintCompiler.cpp
+++ b/Engine/Plugins/Runtime/ModelViewViewModel/Source/ModelViewViewModelBlueprint/Private/MVVMViewBlueprintCompiler.cpp
@@ -128,9 +128,9 @@ FText PropertyPathToText(const UClass* InSelfContext, const UMVVMBlueprintView*
 	return FText::FromString(PropertyPathToString(InSelfContext, BlueprintView, PropertyPath));
 }
 
-FText GetViewModelIdText(const FMVVMBlueprintPropertyPath& PropertyPath)
+FText GetViewModelIdText(FGuid Id)
 {
-	return FText::FromString(PropertyPath.GetViewModelId().ToString(EGuidFormats::DigitsWithHyphensInBraces));
+	return FText::FromString(Id.ToString(EGuidFormats::DigitsWithHyphensInBraces));
 }
 
 TValueOrError<FCompiledBindingLibraryCompiler::FFieldPathHandle, FText> AddObjectFieldPath(FCompiledBindingLibraryCompiler& BindingLibraryCompiler, const UWidgetBlueprintGeneratedClass* Class, FStringView ObjectPath, const UClass* ExpectedType)
@@ -814,14 +814,42 @@ void FMVVMViewBlueprintCompiler::CreateRequiredProperties(const FWidgetBlueprint
 
 	TSet<FName> WidgetSourcesCreated;
 	TSet<FName> WidgetUserPropertyCreated;
+	const FName DefaultWidgetCategory = Context.GetWidgetBlueprint()->GetFName();
 	bool bSelfBindingSourceCreated = NeededBindingSources.ContainsByPredicate([](const TSharedRef<FCompilerBindingSource>& Other)
 		{
 			return Other->Type == FCompilerBindingSource::EType::Self;
 		});
-	const FName DefaultWidgetCategory = Context.GetWidgetBlueprint()->GetFName();
+
 	auto GenerateCompilerContext = [Self = this, DefaultWidgetCategory, Class = Context.GetGeneratedClass(), &ViewModelGuids, &WidgetSourcesCreated, &WidgetUserPropertyCreated, &bSelfBindingSourceCreated](bool bInCreateSource, const FMVVMBlueprintPropertyPath& PropertyPath) -> TValueOrError<void, FText>
 	{
-		switch (PropertyPath.GetSource(Self->WidgetBlueprintCompilerContext.WidgetBlueprint()))
+		FName PropertyPathWidgetName = PropertyPath.GetWidgetName();
+		FGuid PropertyPathViewModelId = PropertyPath.GetViewModelId();
+
+		// If the path is "self.widget.property", "self.viewmodel.property", then remove the "self".
+		EMVVMBlueprintFieldPathSource FieldPathSource = PropertyPath.GetSource(Self->WidgetBlueprintCompilerContext.WidgetBlueprint());
+		if (FieldPathSource == EMVVMBlueprintFieldPathSource::SelfContext)
+		{
+			const TArrayView<const FMVVMBlueprintFieldPath> FieldPaths = PropertyPath.GetFieldPaths();
+			if (FieldPaths.Num() >= 2)
+			{
+				if (FieldPaths[0].GetBindingKind() == EBindingKind::Property && FieldPaths[0].IsFieldSelfContext())
+				{
+					FName RawFieldName = FieldPaths[0].GetRawFieldName();
+					if (Self->WidgetNameToWidgetPointerMap.Contains(RawFieldName))
+					{
+						FieldPathSource = EMVVMBlueprintFieldPathSource::Widget;
+						PropertyPathWidgetName = RawFieldName;
+					}
+					else if (const FMVVMBlueprintViewModelContext* SourceViewModelContext = Self->BlueprintView->FindViewModel(RawFieldName))
+					{
+						FieldPathSource = EMVVMBlueprintFieldPathSource::ViewModel;
+						PropertyPathViewModelId = SourceViewModelContext->GetViewModelId();
+					}
+				}
+			}
+		}
+
+		switch (FieldPathSource)
 		{
 		case EMVVMBlueprintFieldPathSource::SelfContext:
 		{
@@ -841,20 +869,20 @@ void FMVVMViewBlueprintCompiler::CreateRequiredProperties(const FWidgetBlueprint
 		case EMVVMBlueprintFieldPathSource::Widget:
 		{
 			// Only do this once
-			bool bNewCreateSource = !WidgetSourcesCreated.Contains(PropertyPath.GetWidgetName()) && bInCreateSource;
-			bool bNewAddWidgetProperty = !WidgetUserPropertyCreated.Contains(PropertyPath.GetWidgetName());
+			bool bNewCreateSource = !WidgetSourcesCreated.Contains(PropertyPathWidgetName) && bInCreateSource;
+			bool bNewAddWidgetProperty = !WidgetUserPropertyCreated.Contains(PropertyPathWidgetName);
 			if (bNewCreateSource || bNewAddWidgetProperty)
 			{
-				UWidget** WidgetPtr = Self->WidgetNameToWidgetPointerMap.Find(PropertyPath.GetWidgetName());
+				UWidget** WidgetPtr = Self->WidgetNameToWidgetPointerMap.Find(PropertyPathWidgetName);
 				if (WidgetPtr == nullptr || *WidgetPtr == nullptr)
 				{
-					return MakeError(FText::Format(LOCTEXT("InvalidWidgetFormat", "Could not find the targeted widget: {0}"), FText::FromName(PropertyPath.GetWidgetName())));
+					return MakeError(FText::Format(LOCTEXT("InvalidWidgetFormat", "Could not find the targeted widget: {0}"), FText::FromName(PropertyPathWidgetName)));
 				}
 				UWidget* Widget = *WidgetPtr;
 
 				if (bNewCreateSource)
 				{
-					FName PropertyName = PropertyPath.GetWidgetName();
+					FName PropertyName = PropertyPathWidgetName;
 					{
 						TSharedRef<FCompilerBindingSource>* FoundCompilerSource = Self->NeededBindingSources.FindByPredicate([PropertyName](const TSharedRef<FCompilerBindingSource>& Other)
 							{
@@ -863,7 +891,7 @@ void FMVVMViewBlueprintCompiler::CreateRequiredProperties(const FWidgetBlueprint
 						if (FoundCompilerSource != nullptr)
 						{
 							// It should be in the TSet<FName> WidgetSources
-							return MakeError(FText::Format(LOCTEXT("ExistingWidgetSourceFormat", "Internal error. A widget source already exist: {0}"), FText::FromName(PropertyPath.GetWidgetName())));
+							return MakeError(FText::Format(LOCTEXT("ExistingWidgetSourceFormat", "Internal error. A widget source already exist: {0}"), FText::FromName(PropertyPathWidgetName)));
 						}
 					}
 
@@ -881,7 +909,7 @@ void FMVVMViewBlueprintCompiler::CreateRequiredProperties(const FWidgetBlueprint
 				{
 					FCompilerUserWidgetProperty& SourceVariable = Self->NeededUserWidgetProperties.AddDefaulted_GetRef();
 					SourceVariable.AuthoritativeClass = Widget->GetClass();
-					SourceVariable.Name = PropertyPath.GetWidgetName();
+					SourceVariable.Name = PropertyPathWidgetName;
 					SourceVariable.DisplayName = FText::FromString(Widget->GetDisplayLabel());
 					SourceVariable.CategoryName = TEXT("Widget");
 					SourceVariable.CreationType = FCompilerUserWidgetProperty::ECreationType::CreateIfDoesntExist;
@@ -899,15 +927,15 @@ void FMVVMViewBlueprintCompiler::CreateRequiredProperties(const FWidgetBlueprint
 		}
 		case EMVVMBlueprintFieldPathSource::ViewModel:
 		{
-			const FMVVMBlueprintViewModelContext* SourceViewModelContext = Self->BlueprintView->FindViewModel(PropertyPath.GetViewModelId());
+			const FMVVMBlueprintViewModelContext* SourceViewModelContext = Self->BlueprintView->FindViewModel(PropertyPathViewModelId);
 			if (SourceViewModelContext == nullptr)
 			{
-				return MakeError(FText::Format(LOCTEXT("BindingViewModelNotFound", "Could not find viewmodel with GUID {0}."), GetViewModelIdText(PropertyPath)));
+				return MakeError(FText::Format(LOCTEXT("BindingViewModelNotFound", "Could not find viewmodel with GUID {0}."), GetViewModelIdText(PropertyPathViewModelId)));
 			}
 
 			if (!ViewModelGuids.Contains(SourceViewModelContext->GetViewModelId()))
 			{
-				return MakeError(FText::Format(LOCTEXT("BindingViewModelInvalid", "Viewmodel {0} {1} was invalid."), SourceViewModelContext->GetDisplayName(), GetViewModelIdText(PropertyPath)));
+				return MakeError(FText::Format(LOCTEXT("BindingViewModelInvalid", "Viewmodel {0} {1} was invalid."), SourceViewModelContext->GetDisplayName(), GetViewModelIdText(PropertyPathViewModelId)));
 			}
 			break;
 		}
@@ -1235,7 +1263,7 @@ void FMVVMViewBlueprintCompiler::CreateWriteFieldContexts(const FWidgetBlueprint
 				else
 				{
 					WriteFieldPath = MakeWriteFieldPath(
-						DestinationPropertyPath.GetSource(WidgetBlueprintCompilerContext.WidgetBlueprint())
+						FieldContextResult.GetValue().GeneratedFrom
 						, MoveTemp(FieldContextResult.GetValue().GeneratedFields)
 						, MoveTemp(FieldContextResult.GetValue().SkeletalGeneratedFields));
 					GeneratedWriteFieldPaths.Add(WriteFieldPath.ToSharedRef());
@@ -1280,9 +1308,9 @@ void FMVVMViewBlueprintCompiler::CreateWriteFieldContexts(const FWidgetBlueprint
 			else
 			{
 				WriteFieldPath = MakeWriteFieldPath(
-					EventPtr->GetDestinationPath().GetSource(WidgetBlueprintCompilerContext.WidgetBlueprint())
+					FieldContextResult.GetValue().GeneratedFrom
 					, MoveTemp(FieldContextResult.GetValue().GeneratedFields)
-					,MoveTemp(FieldContextResult.GetValue().SkeletalGeneratedFields));
+					, MoveTemp(FieldContextResult.GetValue().SkeletalGeneratedFields));
 				GeneratedWriteFieldPaths.Add(WriteFieldPath.ToSharedRef());
 			}
 		}
@@ -1708,7 +1736,7 @@ void FMVVMViewBlueprintCompiler::CreateReadFieldContexts(UWidgetBlueprintGenerat
 						// if there is a FieldId associated with the read property
 						if (CreateFieldResult.GetValue()->Source)
 						{
-							EMVVMBlueprintFieldPathSource PathSource = PropertyPath.GetSource(Self->WidgetBlueprintCompilerContext.WidgetBlueprint());
+							EMVVMBlueprintFieldPathSource PathSource = ReadFieldContext->GeneratedFrom;
 							bool bValidViewModel = CreateFieldResult.GetValue()->Source->Type == FCompilerBindingSource::EType::ViewModel && PathSource == EMVVMBlueprintFieldPathSource::ViewModel;
 							bool bValidWidget = CreateFieldResult.GetValue()->Source->Type == FCompilerBindingSource::EType::Widget && PathSource == EMVVMBlueprintFieldPathSource::Widget;
 							bool bDynamic = CreateFieldResult.GetValue()->Source->Type == FCompilerBindingSource::EType::DynamicViewmodel;
@@ -1748,6 +1776,7 @@ void FMVVMViewBlueprintCompiler::CreateReadFieldContexts(UWidgetBlueprintGenerat
 				Found->Source = MoveTemp(CreateSourceResult.GetValue().OptionalSource);
 				Found->GeneratedFields = MoveTemp(CreateSourceResult.GetValue().GeneratedFields);
 				Found->SkeletalGeneratedFields = MoveTemp(CreateSourceResult.GetValue().SkeletalGeneratedFields);
+				Found->GeneratedFrom = MoveTemp(CreateSourceResult.GetValue().GeneratedFrom);
 
 				if (CreateFieldId(BindingSourcePath, Found, FMVVMBlueprintPinId()).HasError())
 				{
@@ -1784,6 +1813,7 @@ void FMVVMViewBlueprintCompiler::CreateReadFieldContexts(UWidgetBlueprintGenerat
 							Found->Source = MoveTemp(CreateSourceResult.GetValue().OptionalSource);
 							Found->GeneratedFields = MoveTemp(CreateSourceResult.GetValue().GeneratedFields);
 							Found->SkeletalGeneratedFields = MoveTemp(CreateSourceResult.GetValue().SkeletalGeneratedFields);
+							Found->GeneratedFrom = MoveTemp(CreateSourceResult.GetValue().GeneratedFrom);
 
 							if (CreateFieldId(Pin.GetPath(), Found, Pin.GetId()).HasError())
 							{
@@ -3213,11 +3243,38 @@ TValueOrError<FMVVMViewBlueprintCompiler::FGetFieldsResult, FText> FMVVMViewBlue
 		return MakeValue(TSharedPtr<FCompilerBindingSource>());
 	};
 
-	switch (PropertyPath.GetSource(WidgetBlueprintCompilerContext.WidgetBlueprint()))
+	// If the path is "self.widget.property", "self.viewmodel.property", then remove the "self".
+	EMVVMBlueprintFieldPathSource FieldPathSource = PropertyPath.GetSource(WidgetBlueprintCompilerContext.WidgetBlueprint());
+	TArray<UE::MVVM::FMVVMConstFieldVariant> PropertyPathFields = PropertyPath.GetFields(Class);
+	FGuid ViewModelId = PropertyPath.GetViewModelId();
+	FName DestinationWidgetName = PropertyPath.GetWidgetName();
+
+	if (FieldPathSource == EMVVMBlueprintFieldPathSource::SelfContext)
+	{
+		if (PropertyPathFields.Num() >= 2 && PropertyPathFields[0].IsProperty())
+		{
+			FName FieldName = PropertyPathFields[0].GetName();
+			if (WidgetNameToWidgetPointerMap.Contains(FieldName))
+			{
+				FieldPathSource = EMVVMBlueprintFieldPathSource::Widget;
+				DestinationWidgetName = FieldName;
+				PropertyPathFields.RemoveAt(0);
+			}
+			else if (const FMVVMBlueprintViewModelContext* SourceViewModelContext = BlueprintView->FindViewModel(FieldName))
+			{
+				FieldPathSource = EMVVMBlueprintFieldPathSource::ViewModel;
+				ViewModelId = SourceViewModelContext->GetViewModelId();
+				PropertyPathFields.RemoveAt(0);
+			}
+		}
+	}
+
+	// Build the full path.
+	switch (FieldPathSource)
 	{
 	case EMVVMBlueprintFieldPathSource::ViewModel:
 	{
-		const FCompilerViewModelCreatorContext* FoundViewModelCreator = ViewModelCreatorContexts.FindByPredicate([ViewModelId = PropertyPath.GetViewModelId()](const FCompilerViewModelCreatorContext& Other){ return Other.ViewModelContext.GetViewModelId() == ViewModelId; });
+		const FCompilerViewModelCreatorContext* FoundViewModelCreator = ViewModelCreatorContexts.FindByPredicate([ViewModelId](const FCompilerViewModelCreatorContext& Other){ return Other.ViewModelContext.GetViewModelId() == ViewModelId; });
 		check(FoundViewModelCreator);
 		const FName PropertyName = FoundViewModelCreator->ViewModelContext.GetViewModelName();
 		const FCompilerUserWidgetProperty* FoundUserWidgetProperty = NeededUserWidgetProperties.FindByPredicate([PropertyName](const FCompilerUserWidgetProperty& Other){ return Other.Name == PropertyName; });
@@ -3235,7 +3292,7 @@ TValueOrError<FMVVMViewBlueprintCompiler::FGetFieldsResult, FText> FMVVMViewBlue
 			return MakeError(LOCTEXT("ViewModelShouldHaveSource", "Internal error. Viewmodel should have a source."));
 		}
 
-		Result.GeneratedFields = AppendBaseField(Class, PropertyName, PropertyPath.GetFields(Class));
+		Result.GeneratedFields = AppendBaseField(Class, PropertyName, PropertyPathFields);
 		Result.OptionalSource = FindSourceResult.StealValue();
 		break;
 	}
@@ -3251,13 +3308,12 @@ TValueOrError<FMVVMViewBlueprintCompiler::FGetFieldsResult, FText> FMVVMViewBlue
 			return MakeError(LOCTEXT("WidgetBlueprintShouldHaveSource", "Internal error. The blueprint should have a source."));
 		}
 
-		Result.GeneratedFields = AppendBaseField(Class, FName(), PropertyPath.GetFields(Class));
+		Result.GeneratedFields = AppendBaseField(Class, FName(), PropertyPathFields);
 		Result.OptionalSource = FindSourceResult.StealValue();
 		break;
 	}
 	case EMVVMBlueprintFieldPathSource::Widget:
 	{
-		FName DestinationWidgetName = PropertyPath.GetWidgetName();
 		check(WidgetNameToWidgetPointerMap.Contains(DestinationWidgetName));
 		checkf(!DestinationWidgetName.IsNone(), TEXT("The destination should have been checked and set bAreSourceContextsValid."));
 
@@ -3270,15 +3326,17 @@ TValueOrError<FMVVMViewBlueprintCompiler::FGetFieldsResult, FText> FMVVMViewBlue
 			return MakeError(FindSourceResult.StealError());
 		}
 
-		Result.GeneratedFields = AppendBaseField(Class, DestinationWidgetName, PropertyPath.GetFields(Class));
+		Result.GeneratedFields = AppendBaseField(Class, DestinationWidgetName, PropertyPathFields);
 		Result.OptionalSource = FindSourceResult.StealValue();
 		break;
 	}
 	default:
 		ensureAlwaysMsgf(false, TEXT("Not supported."));
-		Result.GeneratedFields = AppendBaseField(Class, FName(), PropertyPath.GetFields(Class));
+		Result.GeneratedFields = AppendBaseField(Class, FName(), PropertyPathFields);
 		break;
 	}
+
+	Result.GeneratedFrom = FieldPathSource;
 	return MakeValue(MoveTemp(Result));
 }
 
@@ -3309,6 +3367,7 @@ TValueOrError<FMVVMViewBlueprintCompiler::FCreateFieldsResult, FText> FMVVMViewB
 			return MakeError(GetFieldResult.StealError());
 		}
 
+		Result.GeneratedFrom = MoveTemp(GetFieldResult.GetValue().GeneratedFrom);
 		Result.OptionalSource = MoveTemp(GetFieldResult.GetValue().OptionalSource);
 		Result.GeneratedFields = MoveTemp(GetFieldResult.GetValue().GeneratedFields);
 	}
@@ -3351,7 +3410,7 @@ TValueOrError<TSharedPtr<FMVVMViewBlueprintCompiler::FCompilerNotifyFieldId>, FT
 		return MakeValue(TSharedPtr<FCompilerNotifyFieldId>());
 	}
 
-
+	// the path is UserWidget->Widget->Property. THe source is the UserWidget but it should not, UserWidget shuould be remove from the path.
 	FCompilerNotifyFieldId Result;
 	Result.NotificationId = BindingInfo.NotifyFieldId;
 	Result.Source = ReadFieldContext->Source;
diff --git a/Engine/Plugins/Runtime/ModelViewViewModel/Source/ModelViewViewModelBlueprint/Private/MVVMViewBlueprintCompiler.h b/Engine/Plugins/Runtime/ModelViewViewModel/Source/ModelViewViewModelBlueprint/Private/MVVMViewBlueprintCompiler.h
--- a/Engine/Plugins/Runtime/ModelViewViewModel/Source/ModelViewViewModelBlueprint/Private/MVVMViewBlueprintCompiler.h
+++ b/Engine/Plugins/Runtime/ModelViewViewModel/Source/ModelViewViewModelBlueprint/Private/MVVMViewBlueprintCompiler.h
@@ -239,6 +239,7 @@ private:
 		TSharedPtr<FCompilerBindingSource> Source;
 		TArray<UE::MVVM::FMVVMConstFieldVariant> GeneratedFields; // the string path converted into field
 		TArray<UE::MVVM::FMVVMConstFieldVariant> SkeletalGeneratedFields; // the field path converted with getter and setter
+		EMVVMBlueprintFieldPathSource GeneratedFrom = EMVVMBlueprintFieldPathSource::None;
 
 		TSharedPtr<FCompilerNotifyFieldId> NotificationField;
 		FCompiledBindingLibraryCompiler::FFieldPathHandle LibraryCompilerHandle;
@@ -366,11 +367,13 @@ private:
 	struct FGetFieldsResult
 	{
 		TSharedPtr<FCompilerBindingSource> OptionalSource;
+		EMVVMBlueprintFieldPathSource GeneratedFrom = EMVVMBlueprintFieldPathSource::None;
 		TArray<FMVVMConstFieldVariant> GeneratedFields;
 	};
 	struct FCreateFieldsResult
 	{
 		TSharedPtr<FCompilerBindingSource> OptionalSource;
+		EMVVMBlueprintFieldPathSource GeneratedFrom = EMVVMBlueprintFieldPathSource::None;
 		TArray<UE::MVVM::FMVVMConstFieldVariant> GeneratedFields;
 		TArray<UE::MVVM::FMVVMConstFieldVariant> SkeletalGeneratedFields;
 	};
diff --git a/Engine/Plugins/Runtime/ModelViewViewModel/Source/ModelViewViewModelBlueprint/Public/MVVMPropertyPath.h b/Engine/Plugins/Runtime/ModelViewViewModel/Source/ModelViewViewModelBlueprint/Public/MVVMPropertyPath.h
--- a/Engine/Plugins/Runtime/ModelViewViewModel/Source/ModelViewViewModelBlueprint/Public/MVVMPropertyPath.h
+++ b/Engine/Plugins/Runtime/ModelViewViewModel/Source/ModelViewViewModelBlueprint/Public/MVVMPropertyPath.h
@@ -33,6 +33,12 @@ public:
 
 	MODELVIEWVIEWMODELBLUEPRINT_API FMVVMBlueprintFieldPath(const UBlueprint* InContext, UE::MVVM::FMVVMConstFieldVariant InField);
 
+	/**
+	 * Get the binding name. It won't resolve the field.
+	 * That should only be used while compiling the skeletal class.
+	 */
+	MODELVIEWVIEWMODELBLUEPRINT_API FName GetRawFieldName() const;
+
 	/** Get the binding name, resolves reference deprecation / redirectors / etc before returning */
 	MODELVIEWVIEWMODELBLUEPRINT_API FName GetFieldName(const UClass* SelfContext) const;
 
@@ -43,6 +49,11 @@ public:
 	{
 		return BindingReference.IsLocalScope();
 	}
+	
+	bool IsFieldSelfContext() const
+	{
+		return BindingReference.IsSelfContext();
+	}
 
 	MODELVIEWVIEWMODELBLUEPRINT_API UClass* GetParentClass(const UClass* SelfContext) const;
 
@@ -126,7 +137,7 @@ public:
 		return Paths.Num() > 0;
 	}
 
-	TArrayView<FMVVMBlueprintFieldPath const> GetFieldPaths() const
+	const TArrayView<const FMVVMBlueprintFieldPath> GetFieldPaths() const
 	{
 		return Paths;
 	}
