MVVM: support view extensions for injecting MVVM-specific data to widgets.

diff --git a/Engine/Plugins/Runtime/ModelViewViewModel/Source/ModelViewViewModel/Private/View/MVVMView.cpp b/Engine/Plugins/Runtime/ModelViewViewModel/Source/ModelViewViewModel/Private/View/MVVMView.cpp
--- a/Engine/Plugins/Runtime/ModelViewViewModel/Source/ModelViewViewModel/Private/View/MVVMView.cpp
+++ b/Engine/Plugins/Runtime/ModelViewViewModel/Source/ModelViewViewModel/Private/View/MVVMView.cpp
@@ -10,6 +10,7 @@
 #include "Debugging/MVVMDebugging.h"
 #include "Engine/Engine.h"
 #include "Engine/GameInstance.h"
+#include "Extensions/MVVMViewClassExtension.h"
 #include "MVVMGameSubsystem.h"
 #include "MVVMMessageLog.h"
 #include "ModelViewViewModelModule.h"
@@ -53,13 +54,13 @@ using FRecursiveDetctionElement = TTuple<const UObject*, FMVVMViewClass_BindingK
 TArray<FRecursiveDetctionElement> RecursiveDetector;
 } // namespace
 
-void UMVVMView::ConstructView(const UMVVMViewClass* InClassExtension)
+void UMVVMView::ConstructView(const UMVVMViewClass* InGeneratedViewClass)
 {
-	ensure(ClassExtension == nullptr);
-	ClassExtension = InClassExtension;
+	ensure(GeneratedViewClass == nullptr);
+	GeneratedViewClass = InGeneratedViewClass;
 
 	check(Sources.Num() == 0);
-	int32 SourceNum = ClassExtension->GetSources().Num();
+	int32 SourceNum = GeneratedViewClass->GetSources().Num();
 	Sources.SetNum(SourceNum);
 	for (int32 Index = 0; Index < SourceNum; ++Index)
 	{
@@ -71,15 +72,15 @@ void UMVVMView::ConstructView(const UMVVMViewClass* InClassExtension)
 
 void UMVVMView::Construct()
 {
-	check(ClassExtension);
+	check(GeneratedViewClass);
 	check(bConstructed == false);
 
-	if (ClassExtension->DoesInitializeSourcesOnConstruct())
+	if (GeneratedViewClass->DoesInitializeSourcesOnConstruct())
 	{
 		InitializeSources();
 	}
 
-	if (ClassExtension->DoesInitializeEventsOnConstruct())
+	if (GeneratedViewClass->DoesInitializeEventsOnConstruct())
 	{
 		InitializeEvents();
 	}
@@ -108,7 +109,7 @@ void UMVVMView::Destruct()
 
 void UMVVMView::InitializeSources()
 {
-	if (bSourcesInitialized || ClassExtension == nullptr)
+	if (bSourcesInitialized || GeneratedViewClass == nullptr)
 	{
 		return;
 	}
@@ -118,11 +119,11 @@ void UMVVMView::InitializeSources()
 	check(bHasDefaultTickBinding == false && NumberOfSourceWithTickBinding == 0);
 
 	// Init Sources/ViewModel instances
-	const int32 NumberOfSources = ClassExtension->GetSources().Num();
+	const int32 NumberOfSources = GeneratedViewClass->GetSources().Num();
 	check(Sources.Num() == NumberOfSources);
 	check(NumberOfSources <= 64); // the max number of source the bitfield can hold.
 
-	if (ClassExtension->DoesListenToViewModelCollectionChanged())
+	if (GeneratedViewClass->DoesListenToViewModelCollectionChanged())
 	{
 		if (UMVVMViewModelCollectionObject* Collection = UE::MVVM::Private::GetGlobalCollection(GetUserWidget()))
 		{
@@ -135,18 +136,31 @@ void UMVVMView::InitializeSources()
 		InitializeSource(FMVVMView_SourceKey(Index));
 
 		// If DoesInitializeSourcesOnConstruct false but DoesInitializeBindingsOnConstruct, then we need to initialized the bindings when the source are initialized
-		if (ClassExtension->DoesInitializeBindingsOnConstruct())
+		if (GeneratedViewClass->DoesInitializeBindingsOnConstruct())
 		{
 			//note: Run the bindings. The binding from Source A can influence how the Source B is initialized/evaluated.
 			InitializeSourceBindings(FMVVMView_SourceKey(Index), false);
 		}
 	}
 
-	if (ClassExtension->DoesInitializeBindingsOnConstruct())
+	if (GeneratedViewClass->DoesInitializeBindingsOnConstruct())
 	{
 		InitializeSourceBindingsCommon();
 	}
 
+	for (UMVVMViewClassExtension* Extension : GeneratedViewClass->GetViewClassExtensions())
+	{
+		Extension->OnSourcesInitialized(GetUserWidget(), this);
+	}
+
+	if (GeneratedViewClass->DoesInitializeBindingsOnConstruct())
+	{
+		for (UMVVMViewClassExtension* Extension : GeneratedViewClass->GetViewClassExtensions())
+		{
+			Extension->OnBindingsInitialized(GetUserWidget(), this);
+		}
+	}
+
 	bSourcesInitialized = true;
 }
 
@@ -165,7 +179,12 @@ void UMVVMView::UninitializeSources()
 		UninitializeBindings();
 	}
 
-	if (ClassExtension->DoesListenToViewModelCollectionChanged())
+	for (UMVVMViewClassExtension* Extension : GeneratedViewClass->GetViewClassExtensions())
+	{
+		Extension->OnSourcesUninitialized(GetUserWidget(), this);
+	}
+
+	if (GeneratedViewClass->DoesListenToViewModelCollectionChanged())
 	{
 		if (UMVVMViewModelCollectionObject* Collection = UE::MVVM::Private::GetGlobalCollection(GetUserWidget()))
 		{
@@ -186,7 +205,7 @@ void UMVVMView::UninitializeSources()
 void UMVVMView::InitializeSource(FMVVMView_SourceKey SourceKey)
 {
 	FMVVMView_Source& ViewSource = Sources[SourceKey.GetIndex()];
-	const FMVVMViewClass_Source& ClassSource = ClassExtension->GetSource(ViewSource.ClassKey);
+	const FMVVMViewClass_Source& ClassSource = GeneratedViewClass->GetSource(ViewSource.ClassKey);
 
 	check(ViewSource.bSourceInitialized == false);
 
@@ -194,7 +213,7 @@ void UMVVMView::InitializeSource(FMVVMView_SourceKey SourceKey)
 	if (!ViewSource.bSetManually)
 	{
 		UUserWidget* UserWidget = GetUserWidget();
-		UObject* NewSource = ClassSource.GetOrCreateInstance(ClassExtension, this, UserWidget);
+		UObject* NewSource = ClassSource.GetOrCreateInstance(GeneratedViewClass, this, UserWidget);
 		InitializeSourceInternal(NewSource, ViewSource.ClassKey, ClassSource, ViewSource);
 
 #if UE_WITH_MVVM_DEBUGGING
@@ -238,7 +257,7 @@ void UMVVMView::InitializeSourceInternal(UObject* NewSource, FMVVMViewClass_Sour
 void UMVVMView::UninitializeSource(FMVVMView_SourceKey SourceKey)
 {
 	FMVVMView_Source& ViewSource = Sources[SourceKey.GetIndex()];
-	const FMVVMViewClass_Source& ClassSource = ClassExtension->GetSource(ViewSource.ClassKey);
+	const FMVVMViewClass_Source& ClassSource = GeneratedViewClass->GetSource(ViewSource.ClassKey);
 
 	if (ViewSource.bSourceInitialized)
 	{
@@ -292,6 +311,11 @@ void UMVVMView::InitializeBindings()
 	}
 
 	InitializeSourceBindingsCommon();
+
+	for (UMVVMViewClassExtension* Extension : GeneratedViewClass->GetViewClassExtensions())
+	{
+		Extension->OnBindingsInitialized(GetUserWidget(), this);
+	}
 }
 
 
@@ -325,7 +349,12 @@ void UMVVMView::UninitializeBindings()
 
 	bBindingsInitialized = false;
 
-	const TArrayView<const FMVVMViewClass_Source> ClassSources = ClassExtension->GetSources();
+	for (UMVVMViewClassExtension* Extension : GeneratedViewClass->GetViewClassExtensions())
+	{
+		Extension->OnBindingsUninitialized(GetUserWidget(), this);
+	}
+
+	const TArrayView<const FMVVMViewClass_Source> ClassSources = GeneratedViewClass->GetSources();
 	for (int32 SourceIndex = 0; SourceIndex < ClassSources.Num(); ++SourceIndex)
 	{
 		const FMVVMViewClass_Source& ClassSource = ClassSources[SourceIndex];
@@ -348,7 +377,7 @@ void UMVVMView::UninitializeBindings()
 void UMVVMView::InitializeSourceBindings(FMVVMView_SourceKey SourceKey, bool bRunAllBindings)
 {
 	FMVVMView_Source& ViewSource = Sources[SourceKey.GetIndex()];
-	const FMVVMViewClass_Source& ClassSource = ClassExtension->GetSource(ViewSource.ClassKey);
+	const FMVVMViewClass_Source& ClassSource = GeneratedViewClass->GetSource(ViewSource.ClassKey);
 
 	const bool bIsPointerValid = ViewSource.Source != nullptr;
 	const bool bIsBitfieldValid = (ViewSource.ClassKey.GetBit() & ValidSources) != 0;
@@ -378,7 +407,7 @@ void UMVVMView::InitializeSourceBindings(FMVVMView_SourceKey SourceKey, bool bRu
 		{
 			if (SourceBinding.ExecuteAtInitialization() || bRunAllBindings)
 			{
-				const FMVVMViewClass_Binding& ClassBinding = ClassExtension->GetBinding(SourceBinding.GetBindingKey());
+				const FMVVMViewClass_Binding& ClassBinding = GeneratedViewClass->GetBinding(SourceBinding.GetBindingKey());
 				ExecuteBindingImmediately(ClassBinding, SourceBinding.GetBindingKey());
 			}
 		}
@@ -441,21 +470,21 @@ void UMVVMView::UninitializeSourceBindings(FMVVMViewClass_SourceKey SourceKey, c
 
 void UMVVMView::ExecuteBindingImmediately(const FMVVMViewClass_Binding& ClassBinding, FMVVMViewClass_BindingKey KeyForLog) const
 {
-	check(ClassExtension);
+	check(GeneratedViewClass);
 	check(GetUserWidget())
 	if ((ClassBinding.GetSources() & ValidSources) == ClassBinding.GetSources())
 	{
 	
 		// All the source are valid. Run the binding.
 		FMVVMCompiledBindingLibrary::EConversionFunctionType FunctionType = ClassBinding.GetBinding().HasComplexConversionFunction() ? FMVVMCompiledBindingLibrary::EConversionFunctionType::Complex : FMVVMCompiledBindingLibrary::EConversionFunctionType::Simple;
-		TValueOrError<void, FMVVMCompiledBindingLibrary::EExecutionFailingReason> ExecutionResult = ClassExtension->GetBindingLibrary().Execute(GetUserWidget(), ClassBinding.GetBinding(), FunctionType);
+		TValueOrError<void, FMVVMCompiledBindingLibrary::EExecutionFailingReason> ExecutionResult = GeneratedViewClass->GetBindingLibrary().Execute(GetUserWidget(), ClassBinding.GetBinding(), FunctionType);
 
 #if UE_WITH_MVVM_DEBUGGING
 		if (ExecutionResult.HasError())
 		{
 			UE::MVVM::FMessageLog Log(GetUserWidget());
 			Log.Error(FText::Format(LOCTEXT("ExecuteBindingFailGeneric", "The binding '{0}' was not executed. {1}.")
-				, FText::FromString(ClassBinding.ToString(ClassExtension, FMVVMViewClass_Binding::FToStringArgs::Short()))
+				, FText::FromString(ClassBinding.ToString(GeneratedViewClass, FMVVMViewClass_Binding::FToStringArgs::Short()))
 				, FMVVMCompiledBindingLibrary::LexToText(ExecutionResult.GetError())
 			));
 
@@ -467,7 +496,7 @@ void UMVVMView::ExecuteBindingImmediately(const FMVVMViewClass_Binding& ClassBin
 			{
 				UE::MVVM::FMessageLog Log(GetUserWidget());
 				Log.Info(FText::Format(LOCTEXT("ExecuteBindingGeneric", "Execute binding '{0}'.")
-					, FText::FromString(ClassBinding.ToString(ClassExtension, FMVVMViewClass_Binding::FToStringArgs::All()))
+					, FText::FromString(ClassBinding.ToString(GeneratedViewClass, FMVVMViewClass_Binding::FToStringArgs::All()))
 				));
 			}
 			UE::MVVM::FDebugging::BroadcastLibraryBindingExecuted(this, KeyForLog);
@@ -493,12 +522,12 @@ void UMVVMView::ExecuteBindingImmediately(const FMVVMViewClass_Binding& ClassBin
 	else
 	{
 		const uint64 MissingSources = ClassBinding.GetSources() & (~ValidSources);
-		if ((MissingSources & ClassExtension->GetOptionalSources()) != MissingSources)
+		if ((MissingSources & GeneratedViewClass->GetOptionalSources()) != MissingSources)
 		{
 #if UE_WITH_MVVM_DEBUGGING
 			UE::MVVM::FMessageLog Log(GetUserWidget());
 			Log.Error(FText::Format(LOCTEXT("ExecuteBindingFailInvalidSource", "The binding '{0}' was not executed. There are invalid sources.")
-				, FText::FromString(ClassBinding.ToString(ClassExtension, FMVVMViewClass_Binding::FToStringArgs::Short()))
+				, FText::FromString(ClassBinding.ToString(GeneratedViewClass, FMVVMViewClass_Binding::FToStringArgs::Short()))
 			));
 			UE::MVVM::FDebugging::BroadcastLibraryBindingExecuted(this, KeyForLog, FMVVMCompiledBindingLibrary::EExecutionFailingReason::InvalidSource);
 #else
@@ -522,8 +551,8 @@ bool UMVVMView::EvaluateSource(FMVVMViewClass_SourceKey SourceKey)
 	}
 
 	UUserWidget* UserWidget = GetUserWidget();
-	const FMVVMViewClass_Source& ClassSource = ClassExtension->GetSource(SourceKey);
-	UObject* NewSource = ClassSource.GetOrCreateInstance(ClassExtension, this, UserWidget);
+	const FMVVMViewClass_Source& ClassSource = GeneratedViewClass->GetSource(SourceKey);
+	UObject* NewSource = ClassSource.GetOrCreateInstance(GeneratedViewClass, this, UserWidget);
 	ensureMsgf((NewSource == nullptr || NewSource->GetClass()->ImplementsInterface(UNotifyFieldValueChanged::StaticClass())), TEXT("The source has implement the interface. It should be check at compile time."));
 	bool bResult = SetSourceInternal(SourceKey, NewSource, true);
 	if (!bResult)
@@ -543,7 +572,7 @@ void UMVVMView::HandledLibraryBindingValueChanged(UObject* InSource, UE::FieldNo
 	check(InSource);
 	check(InFieldId.IsValid());
 
-	if (ensure(ClassExtension))
+	if (ensure(GeneratedViewClass))
 	{
 		int32 ViewSourceIndex = Sources.IndexOfByPredicate([InSource](const FMVVMView_Source& Other){ return Other.Source == InSource; });
 		if (ViewSourceIndex == INDEX_NONE)
@@ -563,12 +592,12 @@ void UMVVMView::HandledLibraryBindingValueChanged(UObject* InSource, UE::FieldNo
 		}
 
 		const FMVVMViewClass_SourceKey ClassSourceKey = ViewSource.ClassKey;
-		const FMVVMViewClass_Source& ClassSource = ClassExtension->GetSource(ClassSourceKey);
+		const FMVVMViewClass_Source& ClassSource = GeneratedViewClass->GetSource(ClassSourceKey);
 
 		// Run all evaluates
 		if (ClassSource.HasEvaluateBindings())
 		{
-			for (const FMVVMViewClass_EvaluateSource& ClassEvaluate : ClassExtension->GetEvaluateSources())
+			for (const FMVVMViewClass_EvaluateSource& ClassEvaluate : GeneratedViewClass->GetEvaluateSources())
 			{
 				if (ClassEvaluate.GetParentSource() == ClassSourceKey && ClassEvaluate.GetFieldId().GetFieldName() == InFieldId.GetName())
 				{
@@ -591,7 +620,7 @@ void UMVVMView::HandledLibraryBindingValueChanged(UObject* InSource, UE::FieldNo
 
 void UMVVMView::ExecuteBindingInternal(const FMVVMViewClass_SourceBinding& SourceBinding) const
 {
-	const FMVVMViewClass_Binding& ClassBinding = ClassExtension->GetBinding(SourceBinding.GetBindingKey());
+	const FMVVMViewClass_Binding& ClassBinding = GeneratedViewClass->GetBinding(SourceBinding.GetBindingKey());
 	if (ensure(ClassBinding.IsOneWay()))
 	{
 		const EMVVMExecutionMode ExecutionMode = ClassBinding.GetExecuteMode();
@@ -629,10 +658,10 @@ void UMVVMView::ExecuteDelayedBinding(const FMVVMViewClass_BindingKey& DelayedBi
 {
 	SCOPE_CYCLE_COUNTER(STAT_UMG_Viewmodel_ExecuteBinding_Delayed);
 
-	if (ensure(ClassExtension) && bBindingsInitialized)
+	if (ensure(GeneratedViewClass) && bBindingsInitialized)
 	{
 		ensure(bSourcesInitialized);
-		if (ensure(ClassExtension->GetBindings().IsValidIndex(DelayedBinding.GetIndex())))
+		if (ensure(GeneratedViewClass->GetBindings().IsValidIndex(DelayedBinding.GetIndex())))
 		{
 			// Test for recursivity
 			const UMVVMView* Self = this;
@@ -650,7 +679,7 @@ void UMVVMView::ExecuteDelayedBinding(const FMVVMViewClass_BindingKey& DelayedBi
 
 			{
 				UE::MVVM::Private::RecursiveDetector.Emplace(this, DelayedBinding);
-				const FMVVMViewClass_Binding& ClassBinding = ClassExtension->GetBinding(DelayedBinding);
+				const FMVVMViewClass_Binding& ClassBinding = GeneratedViewClass->GetBinding(DelayedBinding);
 				ExecuteBindingImmediately(ClassBinding, DelayedBinding);
 				UE::MVVM::Private::RecursiveDetector.Pop();
 			}
@@ -661,12 +690,12 @@ void UMVVMView::ExecuteDelayedBinding(const FMVVMViewClass_BindingKey& DelayedBi
 
 void UMVVMView::ExecuteTickBindings() const
 {
-	if (ensure(ClassExtension))
+	if (ensure(GeneratedViewClass))
 	{
 		ensure(NumberOfSourceWithTickBinding > 0 || bHasDefaultTickBinding);
 
 		bool bAtLeastOneBindingWasExecuted = false;
-		const TArrayView<const FMVVMViewClass_Binding> ClassBindings = ClassExtension->GetBindings();
+		const TArrayView<const FMVVMViewClass_Binding> ClassBindings = GeneratedViewClass->GetBindings();
 		for (int32 BindingIndex = 0; BindingIndex < ClassBindings.Num(); ++BindingIndex)
 		{
 			const FMVVMViewClass_Binding& ClassBinding = ClassBindings[BindingIndex];
@@ -688,7 +717,7 @@ void UMVVMView::ExecuteTickBindings() const
 
 TScriptInterface<INotifyFieldValueChanged> UMVVMView::GetViewModel(FName ViewModelName) const
 {
-	int32 FoundIndex = ClassExtension->GetSources().IndexOfByPredicate([ViewModelName](const FMVVMViewClass_Source& Other)
+	int32 FoundIndex = GeneratedViewClass->GetSources().IndexOfByPredicate([ViewModelName](const FMVVMViewClass_Source& Other)
 	{
 		return Other.GetName() == ViewModelName && Other.IsViewModel();
 	});
@@ -705,14 +734,14 @@ bool UMVVMView::SetViewModel(FName ViewModelName, TScriptInterface<INotifyFieldV
 		return false;
 	}
 
-	if (ClassExtension == nullptr)
+	if (GeneratedViewClass == nullptr)
 	{
 		UE::MVVM::FMessageLog Log(GetUserWidget());
 		Log.Error(LOCTEXT("SetViewModelInvalidClass", "The view is not constructed."));
 		return false;
 	}
 
-	const int32 ClassSourceIndex = ClassExtension->GetSources().IndexOfByPredicate([ViewModelName](const FMVVMViewClass_Source& Other)
+	const int32 ClassSourceIndex = GeneratedViewClass->GetSources().IndexOfByPredicate([ViewModelName](const FMVVMViewClass_Source& Other)
 		{
 			return Other.GetName() == ViewModelName;
 		});
@@ -736,16 +765,16 @@ bool UMVVMView::SetViewModelByClass(TScriptInterface<INotifyFieldValueChanged> N
 		return false;
 	}
 
-	if (ClassExtension == nullptr)
+	if (GeneratedViewClass == nullptr)
 	{
 		UE::MVVM::FMessageLog Log(GetUserWidget());
-		Log.Error(LOCTEXT("SetViewModelInvalidClassExtension", "The view is not initialized."));
+		Log.Error(LOCTEXT("SetViewModelInvalidGeneratedViewClass", "The view is not initialized."));
 		return false;
 	}
 
 	int32 FoundSourceIndex = INDEX_NONE;
 	int32 SourceIndex = 0;
-	for (const FMVVMViewClass_Source& ClassSource : ClassExtension->GetSources())
+	for (const FMVVMViewClass_Source& ClassSource : GeneratedViewClass->GetSources())
 	{
 		if (NewValue.GetObject()->GetClass()->IsChildOf(ClassSource.GetSourceClass()))
 		{
@@ -781,7 +810,7 @@ bool UMVVMView::SetSourceInternal(FMVVMViewClass_SourceKey ClassSourceKey, TScri
 	check(Sources.IsValidIndex(ClassSourceKey.GetIndex()));
 	FMVVMView_SourceKey ViewSourceKey = FMVVMView_SourceKey(ClassSourceKey.GetIndex());
 
-	const FMVVMViewClass_Source& ClassSource = ClassExtension->GetSource(ClassSourceKey);
+	const FMVVMViewClass_Source& ClassSource = GeneratedViewClass->GetSource(ClassSourceKey);
 	FMVVMView_Source& ViewSource = Sources[ViewSourceKey.GetIndex()];
 
 	if (bForDynamicSource && !ClassSource.CanBeEvaluated())
@@ -841,7 +870,7 @@ bool UMVVMView::SetSourceInternal(FMVVMViewClass_SourceKey ClassSourceKey, TScri
 			// If binding is A.B.Property, and A is set, we need to evaluate/update B
 			if (ClassSource.HasEvaluateBindings())
 			{
-				for (const FMVVMViewClass_EvaluateSource& ClassEvaluate : ClassExtension->GetEvaluateSources())
+				for (const FMVVMViewClass_EvaluateSource& ClassEvaluate : GeneratedViewClass->GetEvaluateSources())
 				{
 					if (ClassEvaluate.GetParentSource() == ClassSourceKey)
 					{
@@ -876,14 +905,14 @@ bool UMVVMView::ExecuteViewModelBindings(FName ViewModelName)
 		return false;
 	}
 
-	if (ClassExtension == nullptr)
+	if (GeneratedViewClass == nullptr)
 	{
 		UE::MVVM::FMessageLog Log(GetUserWidget());
 		Log.Error(LOCTEXT("SetViewModelInvalidClass", "The view is not constructed."));
 		return false;
 	}
 
-	const int32 ClassSourceIndex = ClassExtension->GetSources().IndexOfByPredicate([ViewModelName](const FMVVMViewClass_Source& Other)
+	const int32 ClassSourceIndex = GeneratedViewClass->GetSources().IndexOfByPredicate([ViewModelName](const FMVVMViewClass_Source& Other)
 		{
 			return Other.GetName() == ViewModelName;
 		});
@@ -904,12 +933,12 @@ void UMVVMView::ExecuteViewModelBindingsInternal(FMVVMViewClass_SourceKey ClassS
 {
 	SCOPE_CYCLE_COUNTER(STAT_UMG_Viewmodel_ExecuteViewModel);
 
-	const FMVVMViewClass_Source& ClassSource = ClassExtension->GetSource(ClassSourceKey);
+	const FMVVMViewClass_Source& ClassSource = GeneratedViewClass->GetSource(ClassSourceKey);
 
 	// Run all evaluates
 	if (ClassSource.HasEvaluateBindings())
 	{
-		for (const FMVVMViewClass_EvaluateSource& ClassEvaluate : ClassExtension->GetEvaluateSources())
+		for (const FMVVMViewClass_EvaluateSource& ClassEvaluate : GeneratedViewClass->GetEvaluateSources())
 		{
 			if (ClassEvaluate.GetParentSource() == ClassSourceKey)
 			{
@@ -928,14 +957,14 @@ void UMVVMView::ExecuteViewModelBindingsInternal(FMVVMViewClass_SourceKey ClassS
 
 void UMVVMView::HandleViewModelCollectionChanged()
 {
-	if (ensure(ClassExtension))
+	if (ensure(GeneratedViewClass))
 	{
-		const int32 NumberOfSources = ClassExtension->GetSources().Num();
+		const int32 NumberOfSources = GeneratedViewClass->GetSources().Num();
 		for (int32 Index = 0; Index < NumberOfSources; ++Index)
 		{
 			FMVVMView_SourceKey SourceKey(Index);
 			FMVVMView_Source& ViewSource = Sources[SourceKey.GetIndex()];
-			const FMVVMViewClass_Source& ClassSource = ClassExtension->GetSource(ViewSource.ClassKey);
+			const FMVVMViewClass_Source& ClassSource = GeneratedViewClass->GetSource(ViewSource.ClassKey);
 			if (ViewSource.bSourceInitialized && !ViewSource.bSetManually && ClassSource.RequireGlobalViewModelCollectionUpdate())
 			{
 				TScriptInterface<INotifyFieldValueChanged> NewValue = ClassSource.GetGlobalCollectionViewModel(GetUserWidget());
@@ -958,9 +987,9 @@ void UMVVMView::InitializeEvents()
 
 	UUserWidget* UserWidget = GetUserWidget();
 	check(UserWidget);
-	check(ClassExtension);
+	check(GeneratedViewClass);
 
-	const TArrayView<const FMVVMViewClass_Event>& ClassEvents = ClassExtension->GetEvents();
+	const TArrayView<const FMVVMViewClass_Event>& ClassEvents = GeneratedViewClass->GetEvents();
 
 	ensure(BoundEvents.IsEmpty());
 	BoundEvents.Reset(ClassEvents.Num());
@@ -970,6 +999,11 @@ void UMVVMView::InitializeEvents()
 		BindEvent(ClassEvent, FMVVMViewClass_EventKey(Index));
 	}
 
+	for (UMVVMViewClassExtension* Extension : GeneratedViewClass->GetViewClassExtensions())
+	{
+		Extension->OnEventsInitialized(GetUserWidget(), this);
+	}
+
 	bEventsInitialized = true;
 }
 
@@ -981,10 +1015,15 @@ void UMVVMView::UninitializeEvents()
 		return;
 	}
 
-	check(ClassExtension);
+	check(GeneratedViewClass);
 	
 	bEventsInitialized = false;
 
+	for (UMVVMViewClassExtension* Extension : GeneratedViewClass->GetViewClassExtensions())
+	{
+		Extension->OnEventsUninitialized(GetUserWidget(), this);
+	}
+
 	for (int32 Index = BoundEvents.Num() - 1; Index >= 0; --Index)
 	{
 		UnbindEvent(Index);
@@ -1003,7 +1042,7 @@ void UMVVMView::BindEvent(const FMVVMViewClass_Event& ClassEvent, FMVVMViewClass
 			Log.Error(FText::Format(LOCTEXT("EnableLibraryEventFailed_Format", "Widget '{0}' can't register event '{1}'. {2}")
 				, FText::FromString(Self->GetFullName())
 #if UE_WITH_MVVM_DEBUGGING
-				, FText::FromString(ClassEvent.ToString(Self->ClassExtension, FMVVMViewClass_Event::FToStringArgs()))
+				, FText::FromString(ClassEvent.ToString(Self->GeneratedViewClass, FMVVMViewClass_Event::FToStringArgs()))
 #else
 				, FText::AsNumber(KeyForLog.GetIndex())
 #endif
@@ -1018,11 +1057,11 @@ void UMVVMView::BindEvent(const FMVVMViewClass_Event& ClassEvent, FMVVMViewClass
 		return;
 	}
 
-	const FMVVMCompiledBindingLibrary& Library = ClassExtension->GetBindingLibrary();
+	const FMVVMCompiledBindingLibrary& Library = GeneratedViewClass->GetBindingLibrary();
 	TValueOrError<UE::MVVM::FFieldContext, void> FieldPathResult = Library.EvaluateFieldPath(UserWidget, ClassEvent.GetMulticastDelegatePath());
 	if (FieldPathResult.HasError())
 	{
-		bool bIsOptional = ClassEvent.GetSourceKey().IsValid() && (ClassExtension->GetOptionalSources() & ClassEvent.GetSourceKey().GetBit()) != 0;
+		bool bIsOptional = ClassEvent.GetSourceKey().IsValid() && (GeneratedViewClass->GetOptionalSources() & ClassEvent.GetSourceKey().GetBit()) != 0;
 		if (!bIsOptional)
 		{
 			// Problem. The path can be long Viewmodel.ObjectA.ObjectB.Multicast.
@@ -1035,7 +1074,7 @@ void UMVVMView::BindEvent(const FMVVMViewClass_Event& ClassEvent, FMVVMViewClass
 	UE::MVVM::FFieldContext& FieldContext = FieldPathResult.GetValue();
 	if (FieldContext.GetObjectVariant().IsNull() || !FieldContext.GetObjectVariant().IsUObject())
 	{
-		bool bIsOptional = ClassEvent.GetSourceKey().IsValid() && (ClassExtension->GetOptionalSources() & ClassEvent.GetSourceKey().GetBit()) != 0;
+		bool bIsOptional = ClassEvent.GetSourceKey().IsValid() && (GeneratedViewClass->GetOptionalSources() & ClassEvent.GetSourceKey().GetBit()) != 0;
 		if (!bIsOptional)
 		{
 			// Problem. The path can be long Viewmodel.ObjectA.ObjectB.Multicast.
@@ -1087,7 +1126,7 @@ void UMVVMView::UnbindEvent(int32 Index)
 	check(UserWidget);
 
 	FScriptDelegate Delegate;
-	Delegate.BindUFunction(UserWidget, ClassExtension->GetEvent(BoundEvent.EventKey).GetUserWidgetFunctionName());
+	Delegate.BindUFunction(UserWidget, GeneratedViewClass->GetEvent(BoundEvent.EventKey).GetUserWidgetFunctionName());
 	MulticastDelegateProp->RemoveDelegate(Delegate, EventObject);
 
 	BoundEvents.RemoveAtSwap(Index);
@@ -1109,7 +1148,7 @@ void UMVVMView::ReinitializeEvents(FMVVMViewClass_SourceKey SourceKey, UObject*
 		}
 	}
 
-	const TArrayView<const FMVVMViewClass_Event>& ClassEvents = ClassExtension->GetEvents();
+	const TArrayView<const FMVVMViewClass_Event>& ClassEvents = GeneratedViewClass->GetEvents();
 	for (int32 Index = 0; Index < ClassEvents.Num(); ++Index)
 	{
 		const FMVVMViewClass_Event& ClassEvent = ClassEvents[Index];
@@ -1124,9 +1163,9 @@ void UMVVMView::ReinitializeEvents(FMVVMViewClass_SourceKey SourceKey, UObject*
 bool UMVVMView::AreSourcesValidForEvent(int32 EventKeyIndex) const
 {
 	FMVVMViewClass_EventKey EventKey = FMVVMViewClass_EventKey(EventKeyIndex);
-	if (ClassExtension)
+	if (GeneratedViewClass)
 	{
-		const FMVVMViewClass_Event& ClassEvent = ClassExtension->GetEvent(EventKey);
+		const FMVVMViewClass_Event& ClassEvent = GeneratedViewClass->GetEvent(EventKey);
 		uint64 EventSources = ClassEvent.GetSources();
 		if ((EventSources & ValidSources) == EventSources)
 		{
@@ -1134,12 +1173,12 @@ bool UMVVMView::AreSourcesValidForEvent(int32 EventKeyIndex) const
 		}
 		
 		const uint64 MissingSources = EventSources & (~ValidSources);
-		if ((MissingSources & ClassExtension->GetOptionalSources()) != MissingSources)
+		if ((MissingSources & GeneratedViewClass->GetOptionalSources()) != MissingSources)
 		{
 #if UE_WITH_MVVM_DEBUGGING
 			UE::MVVM::FMessageLog Log(GetUserWidget());
 			Log.Error(FText::Format(LOCTEXT("ExecuteEventFailInvalidSource", "The event '{0}' was not executed. There are invalid sources.")
-				, FText::FromString(ClassEvent.ToString(ClassExtension, FMVVMViewClass_Event::FToStringArgs::Short()))
+				, FText::FromString(ClassEvent.ToString(GeneratedViewClass, FMVVMViewClass_Event::FToStringArgs::Short()))
 			));
 #else
 			UE::MVVM::FMessageLog Log(GetUserWidget());
diff --git a/Engine/Plugins/Runtime/ModelViewViewModel/Source/ModelViewViewModel/Public/Extensions/MVVMViewClassExtension.h b/Engine/Plugins/Runtime/ModelViewViewModel/Source/ModelViewViewModel/Public/Extensions/MVVMViewClassExtension.h
new file mode 100644
--- /dev/null
+++ b/Engine/Plugins/Runtime/ModelViewViewModel/Source/ModelViewViewModel/Public/Extensions/MVVMViewClassExtension.h
@@ -0,0 +1,30 @@
+// Copyright Epic Games, Inc. All Rights Reserved.
+
+#pragma once
+
+#include "CoreMinimal.h"
+
+#include "MVVMViewClassExtension.generated.h"
+
+class UMVVMView;
+class UMVVMBlueprintViewExtension;
+class UUserWidget;
+
+/**
+ * A runtime extension class to define MVVM-related properties and behaviour. This information comes from the
+ * corresponding UMVVMBlueprintViewExtension class. This class provides a hook into the MVVM runtime initializations.
+ */
+UCLASS(MinimalAPI)
+class UMVVMViewClassExtension : public UObject
+{
+	GENERATED_BODY()
+
+public:
+	//~ Functions to be overriden in a user-defined UMVVMViewMyWidgetExtension class
+	virtual void OnSourcesInitialized(UUserWidget* UserWidget, UMVVMView* View) {};
+	virtual void OnBindingsInitialized(UUserWidget* UserWidget, UMVVMView* View) {};
+	virtual void OnEventsInitialized(UUserWidget* UserWidget, UMVVMView* View) {};
+	virtual void OnEventsUninitialized(UUserWidget* UserWidget, UMVVMView* View) {};
+	virtual void OnBindingsUninitialized(UUserWidget* UserWidget, UMVVMView* View) {};
+	virtual void OnSourcesUninitialized(UUserWidget* UserWidget, UMVVMView* View) {};
+};
\ No newline at end of file
diff --git a/Engine/Plugins/Runtime/ModelViewViewModel/Source/ModelViewViewModel/Public/View/MVVMView.h b/Engine/Plugins/Runtime/ModelViewViewModel/Source/ModelViewViewModel/Public/View/MVVMView.h
--- a/Engine/Plugins/Runtime/ModelViewViewModel/Source/ModelViewViewModel/Public/View/MVVMView.h
+++ b/Engine/Plugins/Runtime/ModelViewViewModel/Source/ModelViewViewModel/Public/View/MVVMView.h
@@ -65,7 +65,7 @@ class MODELVIEWVIEWMODEL_API UMVVMView : public UUserWidgetExtension
 	GENERATED_BODY()
 
 public:
-	void ConstructView(const UMVVMViewClass* ClassExtension);
+	void ConstructView(const UMVVMViewClass* InGeneratedViewClass);
 
 	//~ Begin UUserWidgetExtension implementation
 	//virtual void Initialize() override;
@@ -125,7 +125,7 @@ public:
 	/** The shared information for each instance of the view. */
 	const UMVVMViewClass* GetViewClass() const
 	{
-		return ClassExtension;
+		return GeneratedViewClass;
 	}
 
 	/** The list of the sources needed by the view. */
@@ -206,7 +206,7 @@ private:
 
 private:
 	UPROPERTY(Transient)
-	TObjectPtr<const UMVVMViewClass> ClassExtension;
+	TObjectPtr<const UMVVMViewClass> GeneratedViewClass;
 
 	UPROPERTY(VisibleAnywhere, Transient, Category = "View")
 	TArray<FMVVMView_Source> Sources;
diff --git a/Engine/Plugins/Runtime/ModelViewViewModel/Source/ModelViewViewModel/Public/View/MVVMViewClass.h b/Engine/Plugins/Runtime/ModelViewViewModel/Source/ModelViewViewModel/Public/View/MVVMViewClass.h
--- a/Engine/Plugins/Runtime/ModelViewViewModel/Source/ModelViewViewModel/Public/View/MVVMViewClass.h
+++ b/Engine/Plugins/Runtime/ModelViewViewModel/Source/ModelViewViewModel/Public/View/MVVMViewClass.h
@@ -14,6 +14,7 @@
 
 class UMVVMView;
 class UMVVMViewClass;
+class UMVVMViewClassExtension;
 class UMVVMViewModelContextResolver;
 class UUserWidget;
 
@@ -604,6 +605,12 @@ public:
 		return Events[Key.GetIndex()];
 	}
 
+	/** The list of extensions for widgets. */
+	[[nodiscard]] const TArrayView<const TObjectPtr<UMVVMViewClassExtension>> GetViewClassExtensions() const
+	{
+		return ViewClassExtensions;
+	}
+
 #if UE_WITH_MVVM_DEBUGGING
 	struct FToStringArgs
 	{
@@ -633,6 +640,10 @@ private:
 	UPROPERTY(VisibleAnywhere, Category = "View")
 	TArray<FMVVMViewClass_Event> Events;
 
+	/** All MVVM extensions on widgets of the owning userwidget. */
+	UPROPERTY(Instanced)
+	TArray<TObjectPtr<UMVVMViewClassExtension>> ViewClassExtensions;
+
 	/** All the bindings shared between all the View instance. */
 	UPROPERTY()
 	FMVVMCompiledBindingLibrary BindingLibrary;
diff --git a/Engine/Plugins/Runtime/ModelViewViewModel/Source/ModelViewViewModelBlueprint/Private/MVVMViewBlueprintCompiler.cpp b/Engine/Plugins/Runtime/ModelViewViewModel/Source/ModelViewViewModelBlueprint/Private/MVVMViewBlueprintCompiler.cpp
--- a/Engine/Plugins/Runtime/ModelViewViewModel/Source/ModelViewViewModelBlueprint/Private/MVVMViewBlueprintCompiler.cpp
+++ b/Engine/Plugins/Runtime/ModelViewViewModel/Source/ModelViewViewModelBlueprint/Private/MVVMViewBlueprintCompiler.cpp
@@ -8,6 +8,7 @@
 #include "BlueprintEditorSettings.h"
 #include "Blueprint/WidgetTree.h"
 #include "Components/Widget.h"
+#include "Extensions/MVVMBlueprintViewExtension.h"
 #include "HAL/IConsoleManager.h"
 #include "Misc/NamePermissionList.h"
 #include "MVVMBlueprintView.h"
@@ -163,17 +164,17 @@ FMVVMViewBlueprintCompiler::FMVVMViewBlueprintCompiler(FWidgetBlueprintCompilerC
 }
 
 
-void FMVVMViewBlueprintCompiler::AddMessage(const FText& MessageText, EMessageType MessageType) const
+void FMVVMViewBlueprintCompiler::AddMessage(const FText& MessageText, Compiler::EMessageType MessageType) const
 {
 	switch (MessageType)
 	{
-	case EMessageType::Info:
+	case Compiler::EMessageType::Info:
 		WidgetBlueprintCompilerContext.MessageLog.Note(*MessageText.ToString());
 		break;
-	case EMessageType::Warning:
+	case Compiler::EMessageType::Warning:
 		WidgetBlueprintCompilerContext.MessageLog.Warning(*MessageText.ToString());
 		break;
-	case EMessageType::Error:
+	case Compiler::EMessageType::Error:
 		WidgetBlueprintCompilerContext.MessageLog.Error(*MessageText.ToString());
 		break;
 	default:
@@ -182,7 +183,7 @@ void FMVVMViewBlueprintCompiler::AddMessage(const FText& MessageText, EMessageTy
 }
 
 
-void FMVVMViewBlueprintCompiler::AddMessages(TArrayView<TWeakPtr<FCompilerBinding>> Bindings, TArrayView<TWeakPtr<FCompilerEvent>> Events, const FText& MessageText, EMessageType MessageType) const
+void FMVVMViewBlueprintCompiler::AddMessages(TArrayView<TWeakPtr<FCompilerBinding>> Bindings, TArrayView<TWeakPtr<FCompilerEvent>> Events, const FText& MessageText, Compiler::EMessageType MessageType) const
 {
 	for (TWeakPtr<FCompilerBinding>& Binding : Bindings)
 	{
@@ -200,7 +201,7 @@ void FMVVMViewBlueprintCompiler::AddMessages(TArrayView<TWeakPtr<FCompilerBindin
 }
 
 
-void FMVVMViewBlueprintCompiler::AddMessageForBinding(const TSharedPtr<FCompilerBinding>& Binding, const FText& MessageText, EMessageType MessageType, const FMVVMBlueprintPinId& PinId) const
+void FMVVMViewBlueprintCompiler::AddMessageForBinding(const TSharedPtr<FCompilerBinding>& Binding, const FText& MessageText, Compiler::EMessageType MessageType, const FMVVMBlueprintPinId& PinId) const
 {
 	const FMVVMBlueprintViewBinding* BindingPtr = Binding ? BlueprintView->GetBindingAt(Binding->Key.ViewBindingIndex) : nullptr;
 	if (BindingPtr)
@@ -226,7 +227,7 @@ FText GetJoinArgumentNames(const FMVVMBlueprintPinId& PinId)
 }
 
 
-void FMVVMViewBlueprintCompiler::AddMessageForBinding(const FMVVMBlueprintViewBinding& Binding, const FText& MessageText, EMessageType MessageType, const FMVVMBlueprintPinId& PinId) const
+void FMVVMViewBlueprintCompiler::AddMessageForBinding(const FMVVMBlueprintViewBinding& Binding, const FText& MessageText, Compiler::EMessageType MessageType, const FMVVMBlueprintPinId& PinId) const
 {
 	const FText BindingName = FText::FromString(Binding.GetDisplayNameString(WidgetBlueprintCompilerContext.WidgetBlueprint()));
 
@@ -247,7 +248,7 @@ void FMVVMViewBlueprintCompiler::AddMessageForBinding(const FMVVMBlueprintViewBi
 }
 
 
-void FMVVMViewBlueprintCompiler::AddMessageForEvent(const TSharedPtr<FCompilerEvent>& Event, const FText& MessageText, EMessageType MessageType, const FMVVMBlueprintPinId& PinId) const
+void FMVVMViewBlueprintCompiler::AddMessageForEvent(const TSharedPtr<FCompilerEvent>& Event, const FText& MessageText, Compiler::EMessageType MessageType, const FMVVMBlueprintPinId& PinId) const
 {
 	const UMVVMBlueprintViewEvent* EventPtr = Event ? Event->Event.Get() : nullptr;
 	if (EventPtr)
@@ -261,7 +262,7 @@ void FMVVMViewBlueprintCompiler::AddMessageForEvent(const TSharedPtr<FCompilerEv
 }
 
 
-void FMVVMViewBlueprintCompiler::AddMessageForEvent(const UMVVMBlueprintViewEvent* Event, const FText& MessageText, EMessageType MessageType, const FMVVMBlueprintPinId& PinId) const
+void FMVVMViewBlueprintCompiler::AddMessageForEvent(const UMVVMBlueprintViewEvent* Event, const FText& MessageText, Compiler::EMessageType MessageType, const FMVVMBlueprintPinId& PinId) const
 {
 	const FText EventName = Event->GetDisplayName(true);
 
@@ -282,14 +283,14 @@ void FMVVMViewBlueprintCompiler::AddMessageForEvent(const UMVVMBlueprintViewEven
 }
 
 
-void FMVVMViewBlueprintCompiler::AddMessageForViewModel(const FMVVMBlueprintViewModelContext& ViewModel, const FText& Message, EMessageType MessageType) const
+void FMVVMViewBlueprintCompiler::AddMessageForViewModel(const FMVVMBlueprintViewModelContext& ViewModel, const FText& Message, Compiler::EMessageType MessageType) const
 {
 	const FText FormattedError = FText::Format(LOCTEXT("ViewModelFormat", "Viewodel '{0}': {1}"), ViewModel.GetDisplayName(), Message);
 	AddMessage(FormattedError, MessageType);
 }
 
 
-void FMVVMViewBlueprintCompiler::AddMessageForViewModel(const FText& ViewModelDisplayName, const FText& Message, EMessageType MessageType) const
+void FMVVMViewBlueprintCompiler::AddMessageForViewModel(const FText& ViewModelDisplayName, const FText& Message, Compiler::EMessageType MessageType) const
 {
 	const FText FormattedError = FText::Format(LOCTEXT("ViewModelFormat", "Viewodel '{0}': {1}"), ViewModelDisplayName, Message);
 	AddMessage(FormattedError, MessageType);
@@ -656,7 +657,7 @@ void FMVVMViewBlueprintCompiler::CreateBindingList(const FWidgetBlueprintCompile
 		const bool bHasConversionFunction = Binding.Conversion.GetConversionFunction(true) || Binding.Conversion.GetConversionFunction(false);
 		if (bHasConversionFunction && Binding.BindingType == EMVVMBindingMode::TwoWay)
 		{
-			AddMessageForBinding(Binding, LOCTEXT("TwoWayBindingsWithConversion", "Two-way bindings are not allowed to use conversion functions."), EMessageType::Error, FMVVMBlueprintPinId());
+			AddMessageForBinding(Binding, LOCTEXT("TwoWayBindingsWithConversion", "Two-way bindings are not allowed to use conversion functions."), Compiler::EMessageType::Error, FMVVMBlueprintPinId());
 			bIsCreateVariableStepValid = false;
 			continue;
 		}
@@ -701,7 +702,7 @@ void FMVVMViewBlueprintCompiler::CreateEventList(const FWidgetBlueprintCompilerC
 
 		if (!EventPtr->GetEventPath().HasPaths())
 		{
-			AddMessageForEvent(EventPtr, LOCTEXT("EventInvalidEventPath", "The event path is invalid."), EMessageType::Error, FMVVMBlueprintPinId());
+			AddMessageForEvent(EventPtr, LOCTEXT("EventInvalidEventPath", "The event path is invalid."), Compiler::EMessageType::Error, FMVVMBlueprintPinId());
 			bIsCreateVariableStepValid = false;
 			continue;
 		}
@@ -732,14 +733,14 @@ void FMVVMViewBlueprintCompiler::CreateRequiredProperties(const FWidgetBlueprint
 	{
 		if (!ViewModelContext.GetViewModelId().IsValid())
 		{
-			AddMessageForViewModel(ViewModelContext, LOCTEXT("ViewmodelInvalidGuid", "GUID is invalid."), EMessageType::Error);
+			AddMessageForViewModel(ViewModelContext, LOCTEXT("ViewmodelInvalidGuid", "GUID is invalid."), Compiler::EMessageType::Error);
 			bIsCreateVariableStepValid = false;
 			continue;
 		}
 
 		if (ViewModelGuids.Contains(ViewModelContext.GetViewModelId()))
 		{
-			AddMessageForViewModel(ViewModelContext, LOCTEXT("ViewmodelAlreadyAdded", "Identical viewmodel has already been added."), EMessageType::Error);
+			AddMessageForViewModel(ViewModelContext, LOCTEXT("ViewmodelAlreadyAdded", "Identical viewmodel has already been added."), Compiler::EMessageType::Error);
 			bIsCreateVariableStepValid = false;
 			continue;
 		}
@@ -748,7 +749,7 @@ void FMVVMViewBlueprintCompiler::CreateRequiredProperties(const FWidgetBlueprint
 
 		if (ViewModelContext.GetViewModelClass() == nullptr || !ViewModelContext.IsValid())
 		{
-			AddMessageForViewModel(ViewModelContext, LOCTEXT("ViewmodelInvalidClass", "Invalid class."), EMessageType::Error);
+			AddMessageForViewModel(ViewModelContext, LOCTEXT("ViewmodelInvalidClass", "Invalid class."), Compiler::EMessageType::Error);
 			bIsCreateVariableStepValid = false;
 			continue;
 		}
@@ -956,7 +957,7 @@ void FMVVMViewBlueprintCompiler::CreateRequiredProperties(const FWidgetBlueprint
 				TValueOrError<void, FText> SourceContextResult = GenerateCompilerContext(bCreateSource, PropertyPath);
 				if (SourceContextResult.HasError())
 				{
-					Self->AddMessageForBinding(Binding, SourceContextResult.StealError(), EMessageType::Error, PinId);
+					Self->AddMessageForBinding(Binding, SourceContextResult.StealError(), Compiler::EMessageType::Error, PinId);
 					Self->bIsCreateVariableStepValid = false;
 				}
 			};
@@ -973,7 +974,7 @@ void FMVVMViewBlueprintCompiler::CreateRequiredProperties(const FWidgetBlueprint
 					{
 						AddMessageForBinding(Binding
 							, FText::Format(LOCTEXT("InvalidFunctionArgumentPathName", "The conversion function {0} has an invalid argument."), FText::FromString(Binding.GetDisplayNameString(WidgetBlueprintCompilerContext.WidgetBlueprint())))
-							, EMessageType::Error
+							, Compiler::EMessageType::Error
 							, FMVVMBlueprintPinId()
 						);
 						bIsCreateVariableStepValid = false;
@@ -1006,7 +1007,7 @@ void FMVVMViewBlueprintCompiler::CreateRequiredProperties(const FWidgetBlueprint
 			TValueOrError<void, FText> SourceContextResult = GenerateCompilerContext(bCreateSource, PropertyPath);
 			if (SourceContextResult.HasError())
 			{
-				Self->AddMessageForEvent(EventPtr, SourceContextResult.StealError(), EMessageType::Error, PinId);
+				Self->AddMessageForEvent(EventPtr, SourceContextResult.StealError(), Compiler::EMessageType::Error, PinId);
 				Self->bIsCreateVariableStepValid = false;
 			}
 		};
@@ -1021,7 +1022,7 @@ void FMVVMViewBlueprintCompiler::CreateRequiredProperties(const FWidgetBlueprint
 				{
 					AddMessageForEvent(EventPtr
 						, LOCTEXT("InvalidEventArgumentPathName", "The event has an invalid argument.")
-						, EMessageType::Error
+						, Compiler::EMessageType::Error
 						, FMVVMBlueprintPinId()
 					);
 					bIsCreateVariableStepValid = false;
@@ -1113,7 +1114,7 @@ void FMVVMViewBlueprintCompiler::CategorizeBindings(const FWidgetBlueprintCompil
 			UClass* NewClass = WidgetBlueprintCompilerContext.NewClass;
 			if (!ConversionFunction->IsValid(WidgetBlueprintCompilerContext.WidgetBlueprint()))
 			{
-				AddMessageForBinding(Binding, LOCTEXT("InvalidConversionFunction", "The conversion function is invalid."), EMessageType::Error, FMVVMBlueprintPinId());
+				AddMessageForBinding(Binding, LOCTEXT("InvalidConversionFunction", "The conversion function is invalid."), Compiler::EMessageType::Error, FMVVMBlueprintPinId());
 				bIsCreateFunctionsStepValid = false;
 				continue;
 			}
@@ -1123,7 +1124,7 @@ void FMVVMViewBlueprintCompiler::CategorizeBindings(const FWidgetBlueprintCompil
 			if (WrapperGraph == nullptr)
 			{
 				AddMessageForBinding(Binding, LOCTEXT("InvalidConversionFunctionGraph", "The conversion function graph could not be generated.")
-					, EMessageType::Error
+					, Compiler::EMessageType::Error
 					, FMVVMBlueprintPinId()
 				);
 				bIsCreateFunctionsStepValid = false;
@@ -1135,7 +1136,7 @@ void FMVVMViewBlueprintCompiler::CategorizeBindings(const FWidgetBlueprintCompil
 			if (ConversionFunction->HasOrphanedPin())
 			{
 				AddMessageForBinding(Binding, LOCTEXT("InvalidConversionFunctionGraphOrphaned", "The conversion function has an orphaned pin.")
-					, EMessageType::Warning
+					, Compiler::EMessageType::Warning
 					, FMVVMBlueprintPinId()
 				);
 				bIsCreateFunctionsStepValid = false;
@@ -1146,7 +1147,7 @@ void FMVVMViewBlueprintCompiler::CategorizeBindings(const FWidgetBlueprintCompil
 			{
 				AddMessageForBinding(Binding
 					, FText::Format(LOCTEXT("InvalidNumberOfFunctionPin", "The conversion function {0} has no source."), FText::FromString(Binding.GetDisplayNameString(WidgetBlueprintCompilerContext.WidgetBlueprint())))
-					, EMessageType::Error
+					, Compiler::EMessageType::Error
 					, FMVVMBlueprintPinId()
 				);
 				bIsCreateFunctionsStepValid = false;
@@ -1189,7 +1190,7 @@ void FMVVMViewBlueprintCompiler::CategorizeEvents(const FWidgetBlueprintCompiler
 		UEdGraph* WrapperGraph = EventPtr->GetOrCreateWrapperGraph();
 		if (WrapperGraph == nullptr)
 		{
-			AddMessageForEvent(EventPtr, LOCTEXT("InvalidEventGraph", "The event could not be generated."), EMessageType::Warning, FMVVMBlueprintPinId());
+			AddMessageForEvent(EventPtr, LOCTEXT("InvalidEventGraph", "The event could not be generated."), Compiler::EMessageType::Warning, FMVVMBlueprintPinId());
 			bIsCreateFunctionsStepValid = false;
 			continue;
 		}
@@ -1198,7 +1199,7 @@ void FMVVMViewBlueprintCompiler::CategorizeEvents(const FWidgetBlueprintCompiler
 
 		if (EventPtr->HasOrphanedPin())
 		{
-			AddMessageForEvent(EventPtr, LOCTEXT("InvalidEventGraphOrphaned", "The event has an orphaned pin."), EMessageType::Warning, FMVVMBlueprintPinId());
+			AddMessageForEvent(EventPtr, LOCTEXT("InvalidEventGraphOrphaned", "The event has an orphaned pin."), Compiler::EMessageType::Warning, FMVVMBlueprintPinId());
 			bIsCreateFunctionsStepValid = false;
 			continue;
 		}
@@ -1232,7 +1233,7 @@ void FMVVMViewBlueprintCompiler::CreateWriteFieldContexts(const FWidgetBlueprint
 			TValueOrError<FCreateFieldsResult, FText> FieldContextResult = CreateFieldContext(NewSkeletonClass, DestinationPropertyPath, false);
 			if (FieldContextResult.HasError())
 			{
-				AddMessageForBinding(Binding, FieldContextResult.StealError(), EMessageType::Error, FMVVMBlueprintPinId());
+				AddMessageForBinding(Binding, FieldContextResult.StealError(), Compiler::EMessageType::Error, FMVVMBlueprintPinId());
 				bIsCreateFunctionsStepValid = false;
 				continue;
 			}
@@ -1255,7 +1256,7 @@ void FMVVMViewBlueprintCompiler::CreateWriteFieldContexts(const FWidgetBlueprint
 					{
 						AddMessageForBinding(Binding
 							, FText::Format(LOCTEXT("PropertyPathAlreadyUsed", "The property path '{0}' is already used by another binding."), PropertyPathToText(NewSkeletonClass, BlueprintView.Get(), DestinationPropertyPath))
-							, EMessageType::Warning
+							, Compiler::EMessageType::Warning
 							, FMVVMBlueprintPinId()
 						);
 					}
@@ -1289,7 +1290,7 @@ void FMVVMViewBlueprintCompiler::CreateWriteFieldContexts(const FWidgetBlueprint
 		{
 			AddMessageForEvent(EventPtr
 				, FText::Format(Private::PropertyPathIsInvalidFormat, PropertyPathToText(NewSkeletonClass, BlueprintView.Get(), EventPtr->GetDestinationPath()))
-				, EMessageType::Error
+				, Compiler::EMessageType::Error
 				, FMVVMBlueprintPinId()
 			);
 			bIsCreateFunctionsStepValid = false;
@@ -1331,7 +1332,7 @@ void FMVVMViewBlueprintCompiler::CreateViewModelSetters(const FWidgetBlueprintCo
 			{
 				if (!UE::MVVM::FunctionGraphHelper::GenerateViewModelSetter(WidgetBlueprintCompilerContext, Setter.SetterGraph, Setter.PropertyName))
 				{
-					AddMessageForViewModel(Setter.DisplayName, LOCTEXT("SetterFunctionCouldNotBeGenerated", "The setter function could not be generated."), EMessageType::Warning);
+					AddMessageForViewModel(Setter.DisplayName, LOCTEXT("SetterFunctionCouldNotBeGenerated", "The setter function could not be generated."), Compiler::EMessageType::Warning);
 					bIsCreateFunctionsStepValid = false;
 					continue;
 				}
@@ -1351,7 +1352,7 @@ void FMVVMViewBlueprintCompiler::CreateIntermediateGraphFunctions(const FWidgetB
 		{
 			auto AddErrorMessage = [Self = this, &GeneratedDestination](const FText& ErrorMsg)
 			{
-				Self->AddMessages(GeneratedDestination->UsedByBindings, GeneratedDestination->UsedByEvents, ErrorMsg, EMessageType::Error);
+				Self->AddMessages(GeneratedDestination->UsedByBindings, GeneratedDestination->UsedByEvents, ErrorMsg, Compiler::EMessageType::Error);
 			};
 
 			const FProperty* SetterProperty = nullptr;
@@ -1454,6 +1455,7 @@ bool FMVVMViewBlueprintCompiler::PreCompile(UWidgetBlueprintGeneratedClass* Clas
 	PreCompileViewModelCreatorContexts(Class);
 	PreCompileBindings(Class);
 	PreCompileEvents(Class);
+	PreCompileViewExtensions(Class);
 	PreCompileSourceDependencies(Class);
 
 	return AreStepsValid();
@@ -1477,6 +1479,7 @@ bool FMVVMViewBlueprintCompiler::Compile(UWidgetBlueprintGeneratedClass* Class,
 	CompileBindings(CompileResult.GetValue(), Class, ViewExtension);
 	CompileEvaluateSources(CompileResult.GetValue(), Class, ViewExtension);
 	CompileEvents(CompileResult.GetValue(), Class, ViewExtension);
+	CompileViewExtensions(CompileResult.GetValue(), Class, ViewExtension);
 	SortSourceFields(CompileResult.GetValue(), Class, ViewExtension);
 
 	{
@@ -1704,7 +1707,7 @@ void FMVVMViewBlueprintCompiler::CreateReadFieldContexts(UWidgetBlueprintGenerat
 			{
 				Self->AddMessageForBinding(Binding
 					, FText::Format(Private::PropertyPathIsInvalidFormat, PropertyPathToText(Class, Self->BlueprintView.Get(), PropertyPath))
-					, EMessageType::Error
+					, Compiler::EMessageType::Error
 					, PinId
 				);
 				Self->bIsPreCompileStepValid = false;
@@ -1723,7 +1726,7 @@ void FMVVMViewBlueprintCompiler::CreateReadFieldContexts(UWidgetBlueprintGenerat
 				{
 					Self->AddMessageForBinding(Binding
 						, FText::Format(LOCTEXT("CreateNotifyFieldIdFailedInvalidSelfContext", "The property path '{0}' is invalid. {1}"), PropertyPathToText(Class, Self->BlueprintView.Get(), PropertyPath), CreateFieldResult.StealError())
-						, EMessageType::Error
+						, Compiler::EMessageType::Error
 						, PinId
 					);
 					return MakeError();
@@ -1745,7 +1748,7 @@ void FMVVMViewBlueprintCompiler::CreateReadFieldContexts(UWidgetBlueprintGenerat
 							{
 								Self->AddMessageForBinding(Binding
 									, FText::Format(LOCTEXT("CreateNotifyFieldIdFailedInvalidInvalidContext", "Internal error. The property path '{0}' is invalid. The context is invalid."), PropertyPathToText(Class, Self->BlueprintView.Get(), PropertyPath))
-									, EMessageType::Error
+									, Compiler::EMessageType::Error
 									, PinId
 								);
 								return MakeError();
@@ -1846,7 +1849,7 @@ void FMVVMViewBlueprintCompiler::CreateReadFieldContexts(UWidgetBlueprintGenerat
 				{
 					AddMessageForEvent(ValidEvent
 						, FText::Format(Private::PropertyPathIsInvalidFormat, PropertyPathToText(Class, BlueprintView.Get(), Pin.GetPath()))
-						, EMessageType::Error
+						, Compiler::EMessageType::Error
 						, Pin.GetId()
 					);
 					bIsPreCompileStepValid = false;
@@ -1912,7 +1915,7 @@ void FMVVMViewBlueprintCompiler::FixCompilerBindingSelfSource(UWidgetBlueprintGe
 	if (Counter > 1)
 	{
 		AddMessage(LOCTEXT("MoreThanSeldContext", "Internal error. There is more than self context.")
-			, EMessageType::Warning
+			, Compiler::EMessageType::Warning
 		);
 	}
 }
@@ -1951,7 +1954,7 @@ void FMVVMViewBlueprintCompiler::AddWarningForPropertyWithMVVMAndLegacyBinding(U
 					AddMessages(WriteFieldPath->UsedByBindings
 						, TArrayView<TWeakPtr<FCompilerEvent>>()
 						, LOCTEXT("BindingConflictWithLegacy", "The binding is set on a property with legacy binding.")
-						, EMessageType::Warning
+						, Compiler::EMessageType::Warning
 					);
 				}
 			}
@@ -1971,7 +1974,7 @@ void FMVVMViewBlueprintCompiler::PreCompileViewModelCreatorContexts(UWidgetBluep
 		{
 			AddMessageForViewModel(ViewModelContext
 				, FText::Format(LOCTEXT("ViewModelTypeDeprecated", "Viewmodel class '{0}' is deprecated and should not be used. Please update it in the View Models panel."), ViewModelContext.GetViewModelClass()->GetDisplayNameText())
-				, EMessageType::Warning
+				, Compiler::EMessageType::Warning
 			);
 		}
 
@@ -1979,7 +1982,7 @@ void FMVVMViewBlueprintCompiler::PreCompileViewModelCreatorContexts(UWidgetBluep
 		{
 			AddMessageForViewModel(ViewModelContext
 				, LOCTEXT("ViewmodelInvalidInterface", "The class doesn't implement the interface NotifyFieldValueChanged.")
-				, EMessageType::Error
+				, Compiler::EMessageType::Error
 			);
 			bIsPreCompileStepValid = false;
 			continue;
@@ -1992,7 +1995,7 @@ void FMVVMViewBlueprintCompiler::PreCompileViewModelCreatorContexts(UWidgetBluep
 			{
 				AddMessageForViewModel(ViewModelContext
 					, LOCTEXT("ViewModelContextCreationTypeInvalid", "It has an invalidate creation type. You can change it in the View Models panel.")
-					, EMessageType::Error
+					, Compiler::EMessageType::Error
 				);
 				bIsPreCompileStepValid = false;
 				continue;
@@ -2008,7 +2011,7 @@ void FMVVMViewBlueprintCompiler::PreCompileViewModelCreatorContexts(UWidgetBluep
 			{
 				AddMessageForViewModel(ViewModelContext
 					, FText::Format(LOCTEXT("ViewModelTypeAbstract", "Viewmodel class '{0}' is abstract and can't be created. You can change it in the View Models panel."), ViewModelContext.GetViewModelClass()->GetDisplayNameText())
-					, EMessageType::Error
+					, Compiler::EMessageType::Error
 				);
 				bIsPreCompileStepValid = false;
 				continue;
@@ -2035,7 +2038,7 @@ void FMVVMViewBlueprintCompiler::PreCompileViewModelCreatorContexts(UWidgetBluep
 				{
 					AddMessageForViewModel(ViewModelContext
 						, ReadFieldPathResult.GetError()
-						, EMessageType::Error
+						, Compiler::EMessageType::Error
 					);
 					bIsPreCompileStepValid = false;
 					continue;
@@ -2055,7 +2058,7 @@ void FMVVMViewBlueprintCompiler::PreCompileViewModelCreatorContexts(UWidgetBluep
 			{
 				AddMessageForViewModel(ViewModelContext
 					, LOCTEXT("ViewModelInvalidGetter", "Viewmodel has an invalid Getter. You can select a new one in the View Models panel.")
-					, EMessageType::Error
+					, Compiler::EMessageType::Error
 				);
 				bIsPreCompileStepValid = true;
 				continue;
@@ -2066,7 +2069,7 @@ void FMVVMViewBlueprintCompiler::PreCompileViewModelCreatorContexts(UWidgetBluep
 			{
 				AddMessageForViewModel(ViewModelContext
 					, ReadFieldPathResult.GetError()
-					, EMessageType::Error
+					, Compiler::EMessageType::Error
 				);
 				bIsPreCompileStepValid = false;
 				continue;
@@ -2080,7 +2083,7 @@ void FMVVMViewBlueprintCompiler::PreCompileViewModelCreatorContexts(UWidgetBluep
 			{
 				AddMessageForViewModel(ViewModelContext
 					, LOCTEXT("ViewmodelInvalidGlobalIdentifier", "Viewmodel doesn't have a valid Global identifier. You can specify a new one in the Viewmodels panel.")
-					, EMessageType::Error
+					, Compiler::EMessageType::Error
 				);
 				bIsPreCompileStepValid = false;
 				continue;
@@ -2092,7 +2095,7 @@ void FMVVMViewBlueprintCompiler::PreCompileViewModelCreatorContexts(UWidgetBluep
 			{
 				AddMessageForViewModel(ViewModelContext
 					, LOCTEXT("ViewmodelInvalidResolver", "Viewmodel doesn't have a valid Resolver. You can specify a new one in the Viewmodels panel.")
-					, EMessageType::Error
+					, Compiler::EMessageType::Error
 				);
 				bIsPreCompileStepValid = false;
 				continue;
@@ -2102,7 +2105,7 @@ void FMVVMViewBlueprintCompiler::PreCompileViewModelCreatorContexts(UWidgetBluep
 		{
 			AddMessageForViewModel(ViewModelContext
 				, LOCTEXT("ViewmodelInvalidCreationType", "Viewmodel doesn't have a valid creation type. You can select one in the Viewmodels panel.")
-				, EMessageType::Error
+				, Compiler::EMessageType::Error
 			);
 			bIsPreCompileStepValid = false;
 			continue;
@@ -2114,7 +2117,7 @@ void FMVVMViewBlueprintCompiler::PreCompileViewModelCreatorContexts(UWidgetBluep
 		if (!WidgetCreator.Source->AuthoritativeClass->ImplementsInterface(UNotifyFieldValueChanged::StaticClass()))
 		{
 			AddMessage(FText::Format(LOCTEXT("WidgetInvalidInterface", "The widget {0} class doesn't implement the interface NotifyFieldValueChanged."), FText::FromName(WidgetCreator.Source->Name))
-				, EMessageType::Error
+				, Compiler::EMessageType::Error
 			);
 			bIsPreCompileStepValid = false;
 			continue;
@@ -2129,7 +2132,7 @@ void FMVVMViewBlueprintCompiler::PreCompileViewModelCreatorContexts(UWidgetBluep
 					, FText::FromName(WidgetCreator.Source->Name)
 					, ReadFieldPathResult.GetError()
 				)
-					, EMessageType::Error
+					, Compiler::EMessageType::Error
 				);
 				bIsPreCompileStepValid = false;
 				continue;
@@ -2224,7 +2227,7 @@ void FMVVMViewBlueprintCompiler::CompileSources(const FCompiledBindingLibraryCom
 				{
 					AddMessageForViewModel(ViewModelContext
 						, LOCTEXT("ViewModelInvalidInstanceNotFound", "The path for the created viewmodel instance was not found.")
-						, EMessageType::Error
+						, Compiler::EMessageType::Error
 					);
 					bIsCompileStepValid = false;
 					continue;
@@ -2241,7 +2244,7 @@ void FMVVMViewBlueprintCompiler::CompileSources(const FCompiledBindingLibraryCom
 			{
 				AddMessageForViewModel(ViewModelContext
 					, LOCTEXT("ViewModelInvalidInitializationBindingNotGenerated", "The viewmodel initialization binding was not generated.")
-					, EMessageType::Error
+					, Compiler::EMessageType::Error
 				);
 				bIsCompileStepValid = false;
 				continue;
@@ -2261,7 +2264,7 @@ void FMVVMViewBlueprintCompiler::CompileSources(const FCompiledBindingLibraryCom
 			{
 				AddMessageForViewModel(ViewModelContext
 					, LOCTEXT("ViewmodelGlobalContextIdentifier", "The context for viewmodel could not be created. Change the identifier.")
-					, EMessageType::Warning
+					, Compiler::EMessageType::Warning
 				);
 			}
 
@@ -2276,7 +2279,7 @@ void FMVVMViewBlueprintCompiler::CompileSources(const FCompiledBindingLibraryCom
 			{
 				AddMessageForViewModel(ViewModelContext
 					, LOCTEXT("ViewmodelFailedResolverDuplicate", "Internal error. The resolver could not be dupliated.")
-					, EMessageType::Error
+					, Compiler::EMessageType::Error
 				);
 				bIsCompileStepValid = false;
 				continue;
@@ -2289,7 +2292,7 @@ void FMVVMViewBlueprintCompiler::CompileSources(const FCompiledBindingLibraryCom
 		{
 			AddMessageForViewModel(ViewModelContext
 				, LOCTEXT("ViewModelWithoutValidCreationType", "The viewmodel doesn't have a valid creation type.")
-				, EMessageType::Error
+				, Compiler::EMessageType::Error
 			);
 			bIsCompileStepValid = false;
 			continue;
@@ -2301,7 +2304,7 @@ void FMVVMViewBlueprintCompiler::CompileSources(const FCompiledBindingLibraryCom
 			{
 				AddMessageForViewModel(ViewModelContext
 					, LOCTEXT("ViewModelWithoutValidDyanmicParentSource", "The viewmodel doesn't have a valid parent source.")
-					, EMessageType::Error
+					, Compiler::EMessageType::Error
 				);
 				bIsCompileStepValid = false;
 				continue;
@@ -2353,7 +2356,7 @@ void FMVVMViewBlueprintCompiler::CompileSources(const FCompiledBindingLibraryCom
 			if (CompiledFieldPath == nullptr)
 			{
 				AddMessage(FText::Format(LOCTEXT("WidgetInvalidInitializationBindingNotGenerated", "The widget {0} initialization binding was not generated."), FText::FromName(WidgetCreator.Source->Name))
-					, EMessageType::Error
+					, Compiler::EMessageType::Error
 				);
 				bIsCompileStepValid = false;
 				continue;
@@ -2379,7 +2382,7 @@ void FMVVMViewBlueprintCompiler::CompileSources(const FCompiledBindingLibraryCom
 			if (FoundClassSource == nullptr)
 			{
 				AddMessage(FText::Format(LOCTEXT("CompileSources_MissingSources", "Internal error. The source {0} was not compiled."), FText::FromName(BindingSource->Name))
-					, EMessageType::Warning
+					, Compiler::EMessageType::Warning
 				);
 			}
 		}
@@ -2403,7 +2406,7 @@ void FMVVMViewBlueprintCompiler::CompileSources(const FCompiledBindingLibraryCom
 			{
 				AddMessage(
 					FText::Format(LOCTEXT("CompileSources_Dependencies", "The source {0} circularly depends on the source {1}"), FText::FromName(SortDataPair.Key), FText::FromName(OtherSource))
-					, EMessageType::Info
+					, Compiler::EMessageType::Info
 				);
 			}
 		}
@@ -2428,7 +2431,7 @@ void FMVVMViewBlueprintCompiler::CompileSources(const FCompiledBindingLibraryCom
 	{
 		// The view use a uint64 bitfield to filter which viewmodel/source is valid/initialized.
 		//You can use the MVVM.LogViewCompiledResult command to display the compile result and see the name of the sources.
-		AddMessage(LOCTEXT("TooManySources", "There is too many sources for the view. Try spliting your widget into other widgets."), EMessageType::Error);
+		AddMessage(LOCTEXT("TooManySources", "There is too many sources for the view. Try spliting your widget into other widgets."), Compiler::EMessageType::Error);
 	}
 }
 
@@ -2443,7 +2446,7 @@ void FMVVMViewBlueprintCompiler::PreCompileBindings(UWidgetBlueprintGeneratedCla
 			{
 				Self->AddMessageForBinding(Binding
 					, LOCTEXT("NotAllowedExecutionMode", "The binding has a restricted execution mode.")
-					, EMessageType::Error
+					, Compiler::EMessageType::Error
 					, FMVVMBlueprintPinId()
 				);
 				return false;
@@ -2474,7 +2477,7 @@ void FMVVMViewBlueprintCompiler::PreCompileBindings(UWidgetBlueprintGeneratedCla
 						{
 							Self->AddMessageForBinding(Binding
 								, FText::Format(LOCTEXT("CouldNotCreateFieldId", "Could not create Field. {0}"), FieldIdResult.GetError())
-								, EMessageType::Error
+								, Compiler::EMessageType::Error
 								, FMVVMBlueprintPinId()
 							);
 							return false;
@@ -2493,7 +2496,7 @@ void FMVVMViewBlueprintCompiler::PreCompileBindings(UWidgetBlueprintGeneratedCla
 		{
 			Self->AddMessageForBinding(Binding
 				, LOCTEXT("CouldNotCreateSourceFields", "There is no field to bind to. The binding must be a OneTime binding.")
-				, EMessageType::Error
+				, Compiler::EMessageType::Error
 				, FMVVMBlueprintPinId()
 			);
 			return false;
@@ -2520,7 +2523,7 @@ void FMVVMViewBlueprintCompiler::PreCompileBindings(UWidgetBlueprintGeneratedCla
 					{
 						Self->AddMessageForBinding(Binding
 							, FText::Format(Private::CouldNotCreateSourceFieldPathFormat, ::UE::MVVM::FieldPathHelper::ToText(ReadPath->GeneratedFields), FieldPathResult.GetError())
-							, EMessageType::Error
+							, Compiler::EMessageType::Error
 							, FMVVMBlueprintPinId()
 						);
 						return false;
@@ -2538,7 +2541,7 @@ void FMVVMViewBlueprintCompiler::PreCompileBindings(UWidgetBlueprintGeneratedCla
 		{
 			Self->AddMessageForBinding(Binding
 				, LOCTEXT("CouldNotCreateSourceReadIsPresent", "Internal error. There should be no property to read from.")
-				, EMessageType::Error
+				, Compiler::EMessageType::Error
 				, FMVVMBlueprintPinId()
 			);
 			return false;
@@ -2558,7 +2561,7 @@ void FMVVMViewBlueprintCompiler::PreCompileBindings(UWidgetBlueprintGeneratedCla
 		{
 			Self->AddMessageForBinding(Binding
 				, LOCTEXT("InvalidWritePath", "Internal Error. The write path is invalid.")
-				, EMessageType::Error
+				, Compiler::EMessageType::Error
 				, FMVVMBlueprintPinId()
 			);
 			return false;
@@ -2569,7 +2572,7 @@ void FMVVMViewBlueprintCompiler::PreCompileBindings(UWidgetBlueprintGeneratedCla
 		{
 			Self->AddMessageForBinding(Binding
 				, FText::Format(Private::CouldNotCreateDestinationFieldPathFormat, ::UE::MVVM::FieldPathHelper::ToText(ValidBinding.WritePath->GeneratedFields), FieldPathResult.GetError())
-				, EMessageType::Error
+				, Compiler::EMessageType::Error
 				, FMVVMBlueprintPinId()
 			);
 			return false;
@@ -2593,7 +2596,7 @@ void FMVVMViewBlueprintCompiler::PreCompileBindings(UWidgetBlueprintGeneratedCla
 				{
 					Self->AddMessageForBinding(Binding
 						, FText::Format(LOCTEXT("ConversionFunctionNotFound", "The conversion function '{0}' could not be found."), FText::FromName(ViewConversionFunction->GetCompiledFunctionName(Class)))
-						, EMessageType::Error
+						, Compiler::EMessageType::Error
 						, FMVVMBlueprintPinId()
 					);
 					return false;
@@ -2607,7 +2610,7 @@ void FMVVMViewBlueprintCompiler::PreCompileBindings(UWidgetBlueprintGeneratedCla
 					{
 						Self->AddMessageForBinding(Binding
 							, FText::Format(LOCTEXT("ConversionFunctionNotAllow", "The conversion function {0} is not allowed."), FText::FromName(WrapperFunction ? WrapperFunction->GetFName() : FName()))
-							, EMessageType::Error
+							, Compiler::EMessageType::Error
 							, FMVVMBlueprintPinId()
 						);
 						return false;
@@ -2619,7 +2622,7 @@ void FMVVMViewBlueprintCompiler::PreCompileBindings(UWidgetBlueprintGeneratedCla
 					{
 						Self->AddMessageForBinding(Binding
 							, FText::Format(LOCTEXT("ConversionFunctionNotAllow", "The conversion function {0} is not allowed."), FText::FromName(FunctionOrWrapperFunction.GetNode().Get() ? FunctionOrWrapperFunction.GetNode()->GetFName() : FName()))
-							, EMessageType::Error
+							, Compiler::EMessageType::Error
 							, FMVVMBlueprintPinId()
 						);
 						return false;
@@ -2629,7 +2632,7 @@ void FMVVMViewBlueprintCompiler::PreCompileBindings(UWidgetBlueprintGeneratedCla
 				{
 					Self->AddMessageForBinding(Binding
 						, LOCTEXT("ConversionFunctionNodeNotAllow", "The conversion function node is not allowed.")
-						, EMessageType::Error
+						, Compiler::EMessageType::Error
 						, FMVVMBlueprintPinId()
 					);
 					return false;
@@ -2646,7 +2649,7 @@ void FMVVMViewBlueprintCompiler::PreCompileBindings(UWidgetBlueprintGeneratedCla
 					, FText::Format(LOCTEXT("CouldNotCreateConversionFunctionFieldPath", "Couldn't create the conversion function field path '{0}'. {1}")
 						, FText::FromString(ConversionFunction->GetPathName())
 						, FieldPathResult.GetError())
-					, EMessageType::Error
+					, Compiler::EMessageType::Error
 					, FMVVMBlueprintPinId()
 				);
 				return false;
@@ -2660,14 +2663,14 @@ void FMVVMViewBlueprintCompiler::PreCompileBindings(UWidgetBlueprintGeneratedCla
 			const bool bShouldHaveConversionFunction = ValidBinding.Type == FCompilerBinding::EType::ComplexConversionFunction || ValidBinding.Type == FCompilerBinding::EType::SimpleConversionFunction;
 			if ((ConversionFunction != nullptr) != bShouldHaveConversionFunction)
 			{
-				Self->AddMessageForBinding(Binding, LOCTEXT("ConversionFunctionShouldExist", "Internal error. The conversion function should exist."), EMessageType::Error, FMVVMBlueprintPinId());
+				Self->AddMessageForBinding(Binding, LOCTEXT("ConversionFunctionShouldExist", "Internal error. The conversion function should exist."), Compiler::EMessageType::Error, FMVVMBlueprintPinId());
 				return false;
 			}
 
 			const bool bShouldHaveComplexConversionFunction = ValidBinding.Type == FCompilerBinding::EType::ComplexConversionFunction;
 			if (bShouldHaveComplexConversionFunction != BindingHelper::IsValidForComplexRuntimeConversion(ConversionFunction))
 			{
-				Self->AddMessageForBinding(Binding, LOCTEXT("ConversionFunctionIsNotComplex", "Internal Error. The complex conversion function does not respect the prerequisite."), EMessageType::Error, FMVVMBlueprintPinId());
+				Self->AddMessageForBinding(Binding, LOCTEXT("ConversionFunctionIsNotComplex", "Internal Error. The complex conversion function does not respect the prerequisite."), Compiler::EMessageType::Error, FMVVMBlueprintPinId());
 				return false;
 			}
 		}
@@ -2682,7 +2685,7 @@ void FMVVMViewBlueprintCompiler::PreCompileBindings(UWidgetBlueprintGeneratedCla
 			&& ValidBinding->Type != FCompilerBinding::EType::ComplexConversionFunction
 			&& ValidBinding->Type != FCompilerBinding::EType::SimpleConversionFunction)
 		{
-			AddMessageForBinding(Binding, LOCTEXT("UnsupportedBindingType", "The binding is invalid."), EMessageType::Error, FMVVMBlueprintPinId());
+			AddMessageForBinding(Binding, LOCTEXT("UnsupportedBindingType", "The binding is invalid."), Compiler::EMessageType::Error, FMVVMBlueprintPinId());
 			bIsPreCompileStepValid = false;
 			continue;
 		}
@@ -2711,13 +2714,14 @@ void FMVVMViewBlueprintCompiler::PreCompileBindings(UWidgetBlueprintGeneratedCla
 		{
 			AddMessageForBinding(Binding
 				, FText::Format(LOCTEXT("CouldNotCreateBinding", "Could not create binding. {0}"), BindingResult.StealError())
-				, EMessageType::Error
+				, Compiler::EMessageType::Error
 				, FMVVMBlueprintPinId()
 			);
 			bIsPreCompileStepValid = false;
 			continue;
 		}
 		ValidBinding->BindingHandle = BindingResult.StealValue();
+		ValidBinding->CompilerBindingHandle = Compiler::FCompilerBindingHandle::MakeHandle();
 	}
 
 	// Add bindings for the dynamic viewmodels
@@ -2727,7 +2731,7 @@ void FMVVMViewBlueprintCompiler::PreCompileBindings(UWidgetBlueprintGeneratedCla
 		{
 			if (!ViewModelDynamic->NotificationId.IsValid())
 			{
-				AddMessage(FText::Format(LOCTEXT("InvalidNotificationFieldId", "{0} doesn't have a field Id."), FText::FromName(ViewModelDynamic->Source->Name)), EMessageType::Error);
+				AddMessage(FText::Format(LOCTEXT("InvalidNotificationFieldId", "{0} doesn't have a field Id."), FText::FromName(ViewModelDynamic->Source->Name)), Compiler::EMessageType::Error);
 				bIsPreCompileStepValid = false;
 				continue;
 			}
@@ -2736,7 +2740,7 @@ void FMVVMViewBlueprintCompiler::PreCompileBindings(UWidgetBlueprintGeneratedCla
 			TValueOrError<FCompiledBindingLibraryCompiler::FFieldIdHandle, FText> FieldIdResult = BindingLibraryCompiler.AddFieldId(SourceContextClass, ViewModelDynamic->NotificationId.GetFieldName());
 			if (FieldIdResult.HasError() || !FieldIdResult.GetValue().IsValid())
 			{
-				AddMessage(FText::Format(LOCTEXT("CouldNotCreateFieldIdForDynamic", "Could not create Field for {0}. {1}"), FText::FromName(ViewModelDynamic->Source->Name), FieldIdResult.GetError()), EMessageType::Error);
+				AddMessage(FText::Format(LOCTEXT("CouldNotCreateFieldIdForDynamic", "Could not create Field for {0}. {1}"), FText::FromName(ViewModelDynamic->Source->Name), FieldIdResult.GetError()), Compiler::EMessageType::Error);
 				bIsPreCompileStepValid = false;
 				continue;
 			}
@@ -2778,7 +2782,7 @@ void FMVVMViewBlueprintCompiler::CompileBindings(const FCompiledBindingLibraryCo
 		const FMVVMVCompiledBinding* CompiledBinding = CompileResult.Bindings.Find(ValidBinding->BindingHandle);
 		if (CompiledBinding == nullptr)
 		{
-			AddMessageForBinding(Binding, LOCTEXT("CompiledBindingNotGenerated", "Could not generate compiled binding."), EMessageType::Error, FMVVMBlueprintPinId());
+			AddMessageForBinding(Binding, LOCTEXT("CompiledBindingNotGenerated", "Could not generate compiled binding."), Compiler::EMessageType::Error, FMVVMBlueprintPinId());
 			bIsCompileStepValid = false;
 			continue;
 		}
@@ -2813,7 +2817,7 @@ void FMVVMViewBlueprintCompiler::CompileBindings(const FCompiledBindingLibraryCo
 		{
 			if (ReadPath->Source == nullptr)
 			{
-				AddMessageForBinding(Binding, LOCTEXT("InvalidSourceInternal", "Internal error. The binding has an invalid source."), EMessageType::Error, FMVVMBlueprintPinId());
+				AddMessageForBinding(Binding, LOCTEXT("InvalidSourceInternal", "Internal error. The binding has an invalid source."), Compiler::EMessageType::Error, FMVVMBlueprintPinId());
 				bIsCompileStepValid = false;
 				continue;
 			}
@@ -2824,7 +2828,7 @@ void FMVVMViewBlueprintCompiler::CompileBindings(const FCompiledBindingLibraryCo
 				});
 			if (!ViewExtension->Sources.IsValidIndex(ViewExtensionSourceCreatorsIndex))
 			{
-				AddMessageForBinding(Binding, LOCTEXT("CompiledSourceCreatorNotGenerated", "Internal error. The source creator was not generated."), EMessageType::Error, FMVVMBlueprintPinId());
+				AddMessageForBinding(Binding, LOCTEXT("CompiledSourceCreatorNotGenerated", "Internal error. The source creator was not generated."), Compiler::EMessageType::Error, FMVVMBlueprintPinId());
 				bIsCompileStepValid = false;
 				continue;
 			}
@@ -2838,7 +2842,7 @@ void FMVVMViewBlueprintCompiler::CompileBindings(const FCompiledBindingLibraryCo
 			const UE::FieldNotification::FFieldId* CompiledFieldId = bHasField ? CompileResult.FieldIds.Find(ReadPath->NotificationField->LibraryCompilerHandle) : nullptr;
 			if (CompiledFieldId == nullptr && bHasField)
 			{
-				AddMessageForBinding(Binding, LOCTEXT("CompiledFieldNotGenerated", "Internal error. The FieldId was not generated."), EMessageType::Error, FMVVMBlueprintPinId());
+				AddMessageForBinding(Binding, LOCTEXT("CompiledFieldNotGenerated", "Internal error. The FieldId was not generated."), Compiler::EMessageType::Error, FMVVMBlueprintPinId());
 				bIsCompileStepValid = false;
 				continue;
 			}
@@ -2931,7 +2935,7 @@ void FMVVMViewBlueprintCompiler::CompileEvaluateSources(const FCompiledBindingLi
 			});
 		if (!ViewExtension->Sources.IsValidIndex(ViewExtensionSourceCreatorsIndex))
 		{
-			AddMessage(LOCTEXT("CompiledSourceCreatorNotGenerated", "Internal error. The source creator was not generated."), EMessageType::Error);
+			AddMessage(LOCTEXT("CompiledSourceCreatorNotGenerated", "Internal error. The source creator was not generated."), Compiler::EMessageType::Error);
 			bIsCompileStepValid = false;
 			continue;
 		}
@@ -2941,7 +2945,7 @@ void FMVVMViewBlueprintCompiler::CompileEvaluateSources(const FCompiledBindingLi
 			});
 		if (!ViewExtension->Sources.IsValidIndex(ViewExtensionParentCreatorsIndex))
 		{
-			AddMessage(LOCTEXT("CompiledParentSourceCreatorNotGenerated", "Internal error. The parent source creator was not generated."), EMessageType::Error);
+			AddMessage(LOCTEXT("CompiledParentSourceCreatorNotGenerated", "Internal error. The parent source creator was not generated."), Compiler::EMessageType::Error);
 			bIsCompileStepValid = false;
 			continue;
 		}
@@ -2949,7 +2953,7 @@ void FMVVMViewBlueprintCompiler::CompileEvaluateSources(const FCompiledBindingLi
 		const UE::FieldNotification::FFieldId* CompiledFieldId = CompileResult.FieldIds.Find(ViewModelDynamic->NotificationIdLibraryCompilerHandle);
 		if (CompiledFieldId == nullptr)
 		{
-			AddMessage(LOCTEXT("CompiledFieldNotifyNotGenerated", "Internal error. The field notify was not generated."), EMessageType::Error);
+			AddMessage(LOCTEXT("CompiledFieldNotifyNotGenerated", "Internal error. The field notify was not generated."), Compiler::EMessageType::Error);
 			bIsCompileStepValid = false;
 			continue;
 		}
@@ -2996,7 +3000,7 @@ void FMVVMViewBlueprintCompiler::PreCompileEvents(UWidgetBlueprintGeneratedClass
 		{
 			AddMessageForEvent(EventPtr
 				, FText::Format(Private::PropertyPathIsInvalidFormat, PropertyPathToText(Class, BlueprintView.Get(), EventPtr->GetEventPath()))
-				, EMessageType::Error
+				, Compiler::EMessageType::Error
 				, FMVVMBlueprintPinId());
 			bIsPreCompileStepValid = false;
 			continue;
@@ -3008,7 +3012,7 @@ void FMVVMViewBlueprintCompiler::PreCompileEvents(UWidgetBlueprintGeneratedClass
 		{
 			AddMessageForEvent(EventPtr
 				, FText::Format(LOCTEXT("EventPathIsNotMulticastDelegate", "The event {0} is not a multicast delegate."), PropertyPathToText(Class, BlueprintView.Get(), EventPtr->GetEventPath()))
-				, EMessageType::Error
+				, Compiler::EMessageType::Error
 				, FMVVMBlueprintPinId());
 			bIsPreCompileStepValid = false;
 			continue;
@@ -3019,7 +3023,7 @@ void FMVVMViewBlueprintCompiler::PreCompileEvents(UWidgetBlueprintGeneratedClass
 		{
 			AddMessageForEvent(EventPtr
 				, FText::Format(LOCTEXT("EventPathFunctionSignatureError", "The event {0} doesn't match the function signature."), DelegateProperty->GetDisplayNameText())
-				, EMessageType::Error
+				, Compiler::EMessageType::Error
 				, FMVVMBlueprintPinId());
 			bIsPreCompileStepValid = false;
 			continue;
@@ -3031,7 +3035,7 @@ void FMVVMViewBlueprintCompiler::PreCompileEvents(UWidgetBlueprintGeneratedClass
 		{
 			AddMessageForEvent(EventPtr
 				, FText::Format(Private::CouldNotCreateSourceFieldPathFormat, PropertyPathToText(Class, BlueprintView.Get(), EventPtr->GetEventPath()), EventFieldPathResult.GetError())
-				, EMessageType::Error
+				, Compiler::EMessageType::Error
 				, FMVVMBlueprintPinId());
 			bIsPreCompileStepValid = false;
 			continue;
@@ -3085,14 +3089,14 @@ void FMVVMViewBlueprintCompiler::CompileEvents(const FCompiledBindingLibraryComp
 		const FMVVMVCompiledFieldPath* CompiledFieldPath = CompileResult.FieldPaths.Find(ValidEvent->DelegateFieldPathHandle);
 		if (CompiledFieldPath == nullptr)
 		{
-			AddMessageForEvent(ValidEvent->Event.Get(), LOCTEXT("CompiledEventFieldPathNotGenerated", "Could not generate the event path."), EMessageType::Error, FMVVMBlueprintPinId());
+			AddMessageForEvent(ValidEvent->Event.Get(), LOCTEXT("CompiledEventFieldPathNotGenerated", "Could not generate the event path."), Compiler::EMessageType::Error, FMVVMBlueprintPinId());
 			bIsCompileStepValid = false;
 			continue;
 		}
 
 		if (ValidEvent->GeneratedGraphName.IsNone() || Class->FindFunctionByName(ValidEvent->GeneratedGraphName) == nullptr)
 		{
-			AddMessageForEvent(ValidEvent->Event.Get(), LOCTEXT("CompiledEventFieldPathNotGenerated", "Could not generate the event path."), EMessageType::Error, FMVVMBlueprintPinId());
+			AddMessageForEvent(ValidEvent->Event.Get(), LOCTEXT("CompiledEventFieldPathNotGenerated", "Could not generate the event path."), Compiler::EMessageType::Error, FMVVMBlueprintPinId());
 			bIsCompileStepValid = false;
 			continue;
 		}
@@ -3125,7 +3129,7 @@ void FMVVMViewBlueprintCompiler::CompileEvents(const FCompiledBindingLibraryComp
 			{
 				if (ReadPath->Source == nullptr)
 				{
-					AddMessageForEvent(ValidEvent->Event.Get(), LOCTEXT("InvalidEventSourceInternal", "Internal error. The event has an invalid source."), EMessageType::Error, FMVVMBlueprintPinId());
+					AddMessageForEvent(ValidEvent->Event.Get(), LOCTEXT("InvalidEventSourceInternal", "Internal error. The event has an invalid source."), Compiler::EMessageType::Error, FMVVMBlueprintPinId());
 					bIsCompileStepValid = false;
 					continue;
 				}
@@ -3136,7 +3140,7 @@ void FMVVMViewBlueprintCompiler::CompileEvents(const FCompiledBindingLibraryComp
 					});
 				if (!ViewExtension->Sources.IsValidIndex(ViewExtensionSourceCreatorsIndex))
 				{
-					AddMessageForEvent(ValidEvent->Event.Get(), LOCTEXT("CompiledEventSourceCreatorNotGenerated", "Internal error. The source creator was not generated."), EMessageType::Error, FMVVMBlueprintPinId());
+					AddMessageForEvent(ValidEvent->Event.Get(), LOCTEXT("CompiledEventSourceCreatorNotGenerated", "Internal error. The source creator was not generated."), Compiler::EMessageType::Error, FMVVMBlueprintPinId());
 					bIsCompileStepValid = false;
 					continue;
 				}
@@ -3155,6 +3159,225 @@ void FMVVMViewBlueprintCompiler::CompileEvents(const FCompiledBindingLibraryComp
 	}
 }
 
+void FMVVMViewBlueprintCompiler::PreCompileViewExtensions(UWidgetBlueprintGeneratedClass* Class)
+{
+	struct FExposedCompiler : Compiler::IMVVMBlueprintViewPrecompile
+	{
+		FExposedCompiler(FMVVMViewBlueprintCompiler* InSelf, UWidgetBlueprintGeneratedClass* InClass)
+			: Self(InSelf),
+			Class(InClass)
+		{}
+
+		virtual const UMVVMBlueprintView* GetBlueprintView() const override
+		{
+			return Self->BlueprintView.Get();
+		}
+
+		virtual const TMap<FName, UWidget*>& GetWidgetNameToWidgetPointerMap() const override
+		{
+			return Self->WidgetNameToWidgetPointerMap;
+		}
+
+		virtual TArray<Compiler::FCompilerBindingHandle> GetAllBindings() const override
+		{
+			TArray<Compiler::FCompilerBindingHandle> BindingHandles;
+			for (const TSharedRef<FCompilerBinding>& Binding : Self->ValidBindings)
+			{
+				BindingHandles.Add(Binding->CompilerBindingHandle);
+			}
+
+			return BindingHandles;
+		}
+
+		virtual TArray<TArray<UE::MVVM::FMVVMConstFieldVariant>> GetBindingReadFields(Compiler::FCompilerBindingHandle BindingHandle) override
+		{
+			TArray<TArray<UE::MVVM::FMVVMConstFieldVariant>> BindingReadFields;
+			const TSharedRef<FCompilerBinding>* Binding = Self->ValidBindings.FindByPredicate([BindingHandle](const TSharedRef<FCompilerBinding>& OtherBinding) { return OtherBinding->CompilerBindingHandle == BindingHandle; });
+			if (Binding)
+			{
+				for (const TSharedPtr<FGeneratedReadFieldPathContext>& ReadPath : Binding->Get().ReadPaths)
+				{
+					BindingReadFields.Add(ReadPath->SkeletalGeneratedFields);
+				}
+			}
+
+			return BindingReadFields;
+		}
+
+		virtual TArray<UE::MVVM::FMVVMConstFieldVariant> GetBindingWriteFields(Compiler::FCompilerBindingHandle BindingHandle) override
+		{
+			const TSharedRef<FCompilerBinding>* Binding = Self->ValidBindings.FindByPredicate([BindingHandle](const TSharedRef<FCompilerBinding>& OtherBinding) { return OtherBinding->CompilerBindingHandle == BindingHandle; });
+			if (Binding)
+			{
+				if (Binding->Get().WritePath)
+				{
+					return Binding->Get().WritePath->SkeletalGeneratedFields;
+				}
+			}
+
+			return TArray<UE::MVVM::FMVVMConstFieldVariant>();
+		}
+
+		virtual const FProperty* GetBindingSourceProperty(Compiler::FCompilerBindingHandle BindingHandle) override
+		{
+			const TSharedRef<FCompilerBinding>* Binding = Self->ValidBindings.FindByPredicate([BindingHandle](const TSharedRef<FCompilerBinding>& OtherBinding) { return OtherBinding->CompilerBindingHandle == BindingHandle; });
+			if (Binding)
+			{
+				if (Binding->Get().Type == FCompilerBinding::EType::ComplexConversionFunction
+					|| Binding->Get().Type == FCompilerBinding::EType::SimpleConversionFunction)
+				{
+					if (const UFunction* CompiledFunction = Binding->Get().ConversionFunction->GetCompiledFunction(Class))
+					{
+						TValueOrError<const FProperty*, FText> ReturnProperty = UE::MVVM::BindingHelper::TryGetReturnTypeForConversionFunction(CompiledFunction);
+						if (ReturnProperty.HasValue())
+						{
+							return ReturnProperty.GetValue();
+						}
+					}
+				}
+				else if (Binding->Get().Type == FCompilerBinding::EType::Assignment)
+				{
+					if (Binding->Get().ReadPaths.Num() > 0)
+					{
+						if (const TSharedPtr<FGeneratedReadFieldPathContext>& ReadPath = Binding->Get().ReadPaths[0])
+						{
+							if (ReadPath->SkeletalGeneratedFields.Num() > 0)
+							{
+								UE::MVVM::FMVVMConstFieldVariant LastSourceField = ReadPath->SkeletalGeneratedFields.Last();
+								if (LastSourceField.IsProperty())
+								{
+									return LastSourceField.GetProperty();
+								}
+							}
+						}
+					}
+				}
+			}
+
+			return nullptr;
+		}
+
+		virtual TValueOrError<FCompiledBindingLibraryCompiler::FFieldPathHandle, FText> AddObjectFieldPath(const UE::MVVM::Compiler::IMVVMBlueprintViewPrecompile::FObjectFieldPathArgs& Args) override
+		{
+			return UE::MVVM::Private::AddObjectFieldPath(Self->BindingLibraryCompiler, Args.Class, Args.ObjectPath, Args.ExpectedType);
+		}
+
+		virtual TValueOrError<FCompiledBindingLibraryCompiler::FFieldPathHandle, FText> AddFieldPath(TArrayView<const FMVVMConstFieldVariant> InFieldPath, bool bInRead) override
+		{
+			return Self->BindingLibraryCompiler.AddFieldPath(InFieldPath, bInRead);
+		}
+
+		virtual void AddMessageForBinding(Compiler::FCompilerBindingHandle BindingHandle, const FText& MessageText, Compiler::EMessageType MessageType) const override
+		{
+			const TSharedRef<FCompilerBinding>* Binding = Self->ValidBindings.FindByPredicate([BindingHandle](const TSharedRef<FCompilerBinding>& OtherBinding) { return OtherBinding->CompilerBindingHandle == BindingHandle; });
+			if (Binding)
+			{
+				Self->AddMessageForBinding(*Binding, MessageText, MessageType, FMVVMBlueprintPinId());
+			}
+		}
+
+		virtual void AddMessage(const FText& MessageText, Compiler::EMessageType MessageType) override
+		{
+			Self->AddMessage(MessageText, MessageType);
+		}
+
+		virtual void MarkPrecompileStepInvalid() override
+		{
+			Self->bIsPreCompileStepValid = false;
+		}
+
+		virtual ~FExposedCompiler() {};
+
+	private:
+		FMVVMViewBlueprintCompiler* Self = nullptr;
+		UWidgetBlueprintGeneratedClass* Class = nullptr;
+	};
+
+	FExposedCompiler ExposedCompiler(this, Class);
+
+	for (FMVVMExtensionItem& Extension : BlueprintView->GetOuterUMVVMWidgetBlueprintExtension_View()->BlueprintExtensions)
+	{
+		if (ensure(Extension.ExtensionObj))
+		{
+			Extension.ExtensionObj->Precompile(&ExposedCompiler, Class);
+		}
+	}
+}
+
+void FMVVMViewBlueprintCompiler::CompileViewExtensions(const FCompiledBindingLibraryCompiler::FCompileResult& CompileResult, UWidgetBlueprintGeneratedClass* Class, UMVVMViewClass* ViewExtension)
+{
+	struct FExposedCompiler : Compiler::IMVVMBlueprintViewCompile
+	{
+		FExposedCompiler(FMVVMViewBlueprintCompiler* InSelf, const FCompiledBindingLibraryCompiler::FCompileResult& InCompileResult, UMVVMViewClass* InViewExtension)
+			: Self(InSelf)
+			, CompileResult(InCompileResult)
+			, ViewClass(InViewExtension)
+		{}
+
+		virtual const UMVVMBlueprintView* GetBlueprintView() const override
+		{
+			return Self->BlueprintView.Get();
+		}
+
+		virtual const TMap<FName, UWidget*>& GetWidgetNameToWidgetPointerMap() const override
+		{
+			return Self->WidgetNameToWidgetPointerMap;
+		}
+
+		virtual TValueOrError<FMVVMVCompiledFieldPath, void> GetFieldPath(FCompiledBindingLibraryCompiler::FFieldPathHandle FieldPath) override
+		{
+			if (FMVVMVCompiledFieldPath* CompiledPath = CompileResult.FieldPaths.Find(FieldPath))
+			{
+				return MakeValue(*CompiledPath);
+			}
+			else
+			{
+				return MakeError();
+			}
+		}
+
+		virtual void AddMessageForBinding(Compiler::FCompilerBindingHandle BindingHandle, const FText& MessageText, Compiler::EMessageType MessageType) const override
+		{
+			const TSharedRef<FCompilerBinding>* Binding = Self->ValidBindings.FindByPredicate([BindingHandle](const TSharedRef<FCompilerBinding>& OtherBinding) { return OtherBinding->CompilerBindingHandle == BindingHandle; });
+			if (Binding)
+			{
+				Self->AddMessageForBinding(*Binding, MessageText, MessageType, FMVVMBlueprintPinId());
+			}
+		}
+
+		virtual void AddMessage(const FText& MessageText, Compiler::EMessageType MessageType) override
+		{
+			Self->AddMessage(MessageText, MessageType);
+		}
+
+		virtual void MarkCompileStepInvalid() override
+		{
+			Self->bIsCompileStepValid = false;
+		}
+
+		virtual UMVVMViewClassExtension* CreateViewClassExtension(TSubclassOf<UMVVMViewClassExtension> ExtensionClass)
+		{
+			return Self->CreateViewClassExtension(ExtensionClass, ViewClass);
+		}
+
+		virtual ~FExposedCompiler() {};
+
+	private:
+		FMVVMViewBlueprintCompiler* Self = nullptr;
+		FCompiledBindingLibraryCompiler::FCompileResult CompileResult;
+		UMVVMViewClass* ViewClass = nullptr;
+	};
+
+	FExposedCompiler ExposedCompiler(this, CompileResult, ViewExtension);
+
+	for (FMVVMExtensionItem& Extension : BlueprintView->GetOuterUMVVMWidgetBlueprintExtension_View()->BlueprintExtensions)
+	{
+		if (ensure(Extension.ExtensionObj))
+		{
+			Extension.ExtensionObj->Compile(&ExposedCompiler, Class, ViewExtension);
+		}
+	}
+}
 
 void FMVVMViewBlueprintCompiler::PreCompileSourceDependencies(UWidgetBlueprintGeneratedClass* Class)
 {
@@ -3205,7 +3428,7 @@ void FMVVMViewBlueprintCompiler::SortSourceFields(const FCompiledBindingLibraryC
 		{
 			if (Field.GetFieldId().GetIndex() == INDEX_NONE)
 			{
-				AddMessage(LOCTEXT("SortSourceFieldsInvalidId", "Internal error. The id is invalid."), EMessageType::Error);
+				AddMessage(LOCTEXT("SortSourceFieldsInvalidId", "Internal error. The id is invalid."), Compiler::EMessageType::Error);
 				bIsCompileStepValid = false;
 				continue;
 			}
@@ -3645,6 +3868,17 @@ TValueOrError<TSharedPtr<FMVVMViewBlueprintCompiler::FCompilerNotifyFieldId>, FT
 	return MakeValue(MakeShared<FCompilerNotifyFieldId>(Result));
 }
 
+UMVVMViewClassExtension* FMVVMViewBlueprintCompiler::CreateViewClassExtension(TSubclassOf<UMVVMViewClassExtension> ExtensionClass, UMVVMViewClass* ViewClass)
+{
+	if (ensure(ExtensionClass.Get()))
+	{
+		UObject* ExtensionObj = NewObject<UObject>(ViewClass, ExtensionClass.Get(), NAME_None);
+		UMVVMViewClassExtension* Extension = CastChecked<UMVVMViewClassExtension>(ExtensionObj);
+		ViewClass->ViewClassExtensions.Add(Extension);
+		return Extension;
+	}
+	return nullptr;
+}
 
 bool FMVVMViewBlueprintCompiler::IsPropertyPathValid(const UBlueprint* Context, TArrayView<const FMVVMConstFieldVariant> PropertyPath)
 {
diff --git a/Engine/Plugins/Runtime/ModelViewViewModel/Source/ModelViewViewModelBlueprint/Private/MVVMViewBlueprintCompiler.h b/Engine/Plugins/Runtime/ModelViewViewModel/Source/ModelViewViewModelBlueprint/Private/MVVMViewBlueprintCompiler.h
--- a/Engine/Plugins/Runtime/ModelViewViewModel/Source/ModelViewViewModelBlueprint/Private/MVVMViewBlueprintCompiler.h
+++ b/Engine/Plugins/Runtime/ModelViewViewModel/Source/ModelViewViewModelBlueprint/Private/MVVMViewBlueprintCompiler.h
@@ -4,6 +4,7 @@
 
 #include "Bindings/MVVMCompiledBindingLibraryCompiler.h"
 #include "MVVMBlueprintView.h"
+#include "MVVMBlueprintViewCompilerInterface.h"
 #include "MVVMBlueprintViewEvent.h"
 #include "MVVMBlueprintViewModelContext.h"
 #include "Templates/ValueOrError.h"
@@ -19,6 +20,7 @@ class UMVVMBlueprintView;
 class UMVVMBlueprintViewConversionFunction;
 class UMVVMBlueprintViewEvent;
 class UMVVMViewClass;
+class UMVVMViewClassExtension;
 class UWidgetBlueprintGeneratedClass;
 namespace  UE::MVVM
 {
@@ -95,6 +97,7 @@ private:
 	void PreCompileViewModelCreatorContexts(UWidgetBlueprintGeneratedClass* Class);
 	void PreCompileBindings(UWidgetBlueprintGeneratedClass* Class);
 	void PreCompileEvents(UWidgetBlueprintGeneratedClass* Class);
+	void PreCompileViewExtensions(UWidgetBlueprintGeneratedClass* Class);
 	void PreCompileSourceDependencies(UWidgetBlueprintGeneratedClass* Class);
 
 	// Compile
@@ -102,6 +105,7 @@ private:
 	void CompileBindings(const FCompiledBindingLibraryCompiler::FCompileResult& CompileResult, UWidgetBlueprintGeneratedClass* Class, UMVVMViewClass* ViewExtension);
 	void CompileEvaluateSources(const FCompiledBindingLibraryCompiler::FCompileResult& CompileResult, UWidgetBlueprintGeneratedClass* Class, UMVVMViewClass* ViewExtension);
 	void CompileEvents(const FCompiledBindingLibraryCompiler::FCompileResult& CompileResult, UWidgetBlueprintGeneratedClass* Class, UMVVMViewClass* ViewExtension);
+	void CompileViewExtensions(const FCompiledBindingLibraryCompiler::FCompileResult& CompileResult, UWidgetBlueprintGeneratedClass* Class, UMVVMViewClass* ViewExtension);
 	void SortSourceFields(const FCompiledBindingLibraryCompiler::FCompileResult& CompileResult, UWidgetBlueprintGeneratedClass* Class, UMVVMViewClass* ViewExtension);
 
 private:
@@ -314,6 +318,7 @@ private:
 
 		FCompiledBindingLibraryCompiler::FBindingHandle BindingHandle;
 		FCompiledBindingLibraryCompiler::FFieldPathHandle ConversionFunctionHandle;
+		UE::MVVM::Compiler::FCompilerBindingHandle CompilerBindingHandle;
 	};
 	TArray<TSharedRef<FCompilerBinding>> ValidBindings;
 
@@ -351,18 +356,14 @@ private:
 	bool bIsCompileStepValid = true;
 
 private:
-	enum class EMessageType
-	{
-		Info = 0, Warning = 1, Error = 2
-	};
-	void AddMessage(const FText& MessageText, EMessageType MessageType) const;
-	void AddMessages(TArrayView<TWeakPtr<FCompilerBinding>> Bindings, TArrayView<TWeakPtr<FCompilerEvent>> Events, const FText& MessageText, EMessageType MessageType) const;
-	void AddMessageForBinding(const TSharedPtr<FCompilerBinding>& Binding, const FText& MessageText, EMessageType MessageType, const FMVVMBlueprintPinId& ArgumentName) const;
-	void AddMessageForBinding(const FMVVMBlueprintViewBinding& Binding, const FText& MessageText, EMessageType MessageType, const FMVVMBlueprintPinId& ArgumentName) const;
-	void AddMessageForEvent(const TSharedPtr<FCompilerEvent>& Event, const FText& MessageText, EMessageType MessageType, const FMVVMBlueprintPinId& ArgumentName) const;
-	void AddMessageForEvent(const UMVVMBlueprintViewEvent* Event, const FText& MessageText, EMessageType MessageType, const FMVVMBlueprintPinId& ArgumentName) const;
-	void AddMessageForViewModel(const FMVVMBlueprintViewModelContext& ViewModel, const FText& Message, EMessageType MessageType) const;
-	void AddMessageForViewModel(const FText& ViewModelDisplayName, const FText& Message, EMessageType MessageType) const;
+	void AddMessage(const FText& MessageText, Compiler::EMessageType MessageType) const;
+	void AddMessages(TArrayView<TWeakPtr<FCompilerBinding>> Bindings, TArrayView<TWeakPtr<FCompilerEvent>> Events, const FText& MessageText, Compiler::EMessageType MessageType) const;
+	void AddMessageForBinding(const TSharedPtr<FCompilerBinding>& Binding, const FText& MessageText, Compiler::EMessageType MessageType, const FMVVMBlueprintPinId& ArgumentName) const;
+	void AddMessageForBinding(const FMVVMBlueprintViewBinding& Binding, const FText& MessageText, Compiler::EMessageType MessageType, const FMVVMBlueprintPinId& ArgumentName) const;
+	void AddMessageForEvent(const TSharedPtr<FCompilerEvent>& Event, const FText& MessageText, Compiler::EMessageType MessageType, const FMVVMBlueprintPinId& ArgumentName) const;
+	void AddMessageForEvent(const UMVVMBlueprintViewEvent* Event, const FText& MessageText, Compiler::EMessageType MessageType, const FMVVMBlueprintPinId& ArgumentName) const;
+	void AddMessageForViewModel(const FMVVMBlueprintViewModelContext& ViewModel, const FText& Message, Compiler::EMessageType MessageType) const;
+	void AddMessageForViewModel(const FText& ViewModelDisplayName, const FText& Message, Compiler::EMessageType MessageType) const;
 
 	struct FGetFieldsResult
 	{
@@ -381,6 +382,7 @@ private:
 	TValueOrError<FCreateFieldsResult, FText> CreateFieldContext(const UWidgetBlueprintGeneratedClass* Class, const FMVVMBlueprintPropertyPath& PropertyPath, bool bForSourceReading) const;
 	TValueOrError<TSharedPtr<FCompilerNotifyFieldId>, FText> CreateNotifyFieldId(const UWidgetBlueprintGeneratedClass* Class, const TSharedPtr<FGeneratedReadFieldPathContext>& ReadFieldContext, const FMVVMBlueprintViewBinding& Binding);
 
+	UMVVMViewClassExtension* CreateViewClassExtension(TSubclassOf<UMVVMViewClassExtension> ExtensionClass, UMVVMViewClass* ViewClass);
 
 	static TArray<FMVVMConstFieldVariant> AppendBaseField(const UClass* Class, FName PropertyName, TArray<FMVVMConstFieldVariant> Properties);
 	static bool IsPropertyPathValid(const UBlueprint* Context, TArrayView<const FMVVMConstFieldVariant> PropertyPath);
diff --git a/Engine/Plugins/Runtime/ModelViewViewModel/Source/ModelViewViewModelBlueprint/Private/MVVMWidgetBlueprintExtension_View.cpp b/Engine/Plugins/Runtime/ModelViewViewModel/Source/ModelViewViewModelBlueprint/Private/MVVMWidgetBlueprintExtension_View.cpp
--- a/Engine/Plugins/Runtime/ModelViewViewModel/Source/ModelViewViewModelBlueprint/Private/MVVMWidgetBlueprintExtension_View.cpp
+++ b/Engine/Plugins/Runtime/ModelViewViewModel/Source/ModelViewViewModelBlueprint/Private/MVVMWidgetBlueprintExtension_View.cpp
@@ -2,6 +2,8 @@
 
 #include "MVVMWidgetBlueprintExtension_View.h"
 
+#include "Blueprint/WidgetTree.h"
+#include "Extensions/MVVMBlueprintViewExtension.h"
 #include "FindInBlueprintManager.h"
 #include "MVVMBlueprintInstancedViewModel.h"
 #include "MVVMBlueprintView.h"
@@ -56,6 +58,101 @@ void UMVVMWidgetBlueprintExtension_View::PostLoad()
 	Super::PostLoad();
 }
 
+UMVVMBlueprintViewExtension* UMVVMWidgetBlueprintExtension_View::CreateBlueprintWidgetExtension(TSubclassOf<UMVVMBlueprintViewExtension> ExtensionClass, FName WidgetName)
+{
+	if (ensure(ExtensionClass.Get()))
+	{
+		UObject* ExtensionObj = NewObject<UObject>(this, ExtensionClass.Get(), NAME_None, RF_Transactional);
+		UMVVMBlueprintViewExtension* NewExtension = CastChecked<UMVVMBlueprintViewExtension>(ExtensionObj);
+		NewExtension->Modify();
+
+		FMVVMExtensionItem ExtensionToAdd;
+		ExtensionToAdd.WidgetName = WidgetName;
+		ExtensionToAdd.ExtensionObj = NewExtension;
+		BlueprintExtensions.Add(ExtensionToAdd);
+
+		FBlueprintEditorUtils::MarkBlueprintAsStructurallyModified(GetWidgetBlueprint());
+		return NewExtension;
+	}
+	return nullptr;
+}
+
+void UMVVMWidgetBlueprintExtension_View::RemoveBlueprintWidgetExtension(UMVVMBlueprintViewExtension* ExtensionToRemove, FName WidgetName)
+{
+	FMVVMExtensionItem Extension;
+	Extension.WidgetName = WidgetName;
+	Extension.ExtensionObj = ExtensionToRemove;
+	BlueprintExtensions.Remove(Extension);
+	FBlueprintEditorUtils::MarkBlueprintAsStructurallyModified(GetWidgetBlueprint());
+}
+
+TArray<UMVVMBlueprintViewExtension*> UMVVMWidgetBlueprintExtension_View::GetBlueprintExtensionsForWidget(FName WidgetName) const
+{
+	TArray<UMVVMBlueprintViewExtension*> ThisWidgetExtensions;
+	for (const FMVVMExtensionItem& Extension : BlueprintExtensions)
+	{
+		if (Extension.WidgetName == WidgetName)
+		{
+			ThisWidgetExtensions.Add(Extension.ExtensionObj);
+		}
+	}
+	return ThisWidgetExtensions;
+}
+
+void UMVVMWidgetBlueprintExtension_View::VerifyWidgetExtensions()
+{
+	if (const UWidgetBlueprint* WidgetBlueprint = GetWidgetBlueprint())
+	{
+		if (const UWidgetTree* WidgetTree = WidgetBlueprint->WidgetTree)
+		{
+			TArray <FName, TInlineAllocator<4>> WidgetNamesToRemove;
+			for (const FMVVMExtensionItem& Extension : BlueprintExtensions)
+			{
+				if (!Extension.WidgetName.IsNone())
+				{
+					WidgetNamesToRemove.Add(Extension.WidgetName);
+				}
+			}
+			// Find widgets that are no longer in the tree and delete their extensions.
+			WidgetTree->ForEachWidget([&WidgetNamesToRemove, this](TObjectPtr<UWidget> Widget) {
+				if (Widget)
+				{
+					for (int32 Index = WidgetNamesToRemove.Num() - 1; Index >= 0; Index--)
+					{
+						const FName WidgetName = WidgetNamesToRemove[Index];
+						if (WidgetName == Widget->GetFName())
+						{
+							WidgetNamesToRemove.Remove(WidgetName);
+						}
+					}
+				}
+			});
+
+			for (int32 Index = BlueprintExtensions.Num() - 1; Index >= 0; Index--)
+			{
+				if (WidgetNamesToRemove.Contains(BlueprintExtensions[Index].WidgetName))
+				{
+					BlueprintExtensions.RemoveAt(Index);
+				}
+			}
+		}
+	}
+}
+
+void UMVVMWidgetBlueprintExtension_View::RenameWidgetExtensions(FName OldName, FName NewName)
+{
+	for (FMVVMExtensionItem& Extension : BlueprintExtensions)
+	{
+		if (Extension.WidgetName == OldName)
+		{
+			Extension.WidgetName = NewName;
+			if (Extension.ExtensionObj)
+			{
+				Extension.ExtensionObj->WidgetRenamed(OldName, NewName);
+			}
+		}
+	}
+}
 
 void UMVVMWidgetBlueprintExtension_View::HandlePreloadObjectsForCompilation(UBlueprint* OwningBlueprint)
 {
@@ -92,6 +189,8 @@ void UMVVMWidgetBlueprintExtension_View::HandlePreloadObjectsForCompilation(UBlu
 
 void UMVVMWidgetBlueprintExtension_View::HandleBeginCompilation(FWidgetBlueprintCompilerContext& InCreationContext)
 {
+	VerifyWidgetExtensions();
+
 	for (const FMVVMBlueprintViewModelContext& AvailableViewModel : BlueprintView->GetViewModels())
 	{
 		if (AvailableViewModel.InstancedViewModel)
diff --git a/Engine/Plugins/Runtime/ModelViewViewModel/Source/ModelViewViewModelBlueprint/Public/Extensions/MVVMBlueprintViewExtension.h b/Engine/Plugins/Runtime/ModelViewViewModel/Source/ModelViewViewModelBlueprint/Public/Extensions/MVVMBlueprintViewExtension.h
new file mode 100644
--- /dev/null
+++ b/Engine/Plugins/Runtime/ModelViewViewModel/Source/ModelViewViewModelBlueprint/Public/Extensions/MVVMBlueprintViewExtension.h
@@ -0,0 +1,34 @@
+// Copyright Epic Games, Inc. All Rights Reserved.
+
+#pragma once
+
+#include "CoreMinimal.h"
+
+#include "MVVMBlueprintViewExtension.generated.h"
+
+
+class UMVVMViewClass;
+class UWidgetBlueprintGeneratedClass;
+
+namespace UE::MVVM::Compiler
+{
+	class IMVVMBlueprintViewPrecompile;
+	class IMVVMBlueprintViewCompile;
+}
+
+/**
+ * An extension class to define MVVM-related properties and behaviour. When WBP compiled, this information is copied 
+ * into UMVVMViewClassExtension. This class provides a hook into the MVVM compiler and partially exposes it for injecting
+ * user-defined behaviour at compile-time.
+ */
+UCLASS(MinimalAPI)
+class UMVVMBlueprintViewExtension : public UObject
+{
+	GENERATED_BODY()
+
+public:
+	//~ Functions to be overriden in a user-defined UMVVMViewBlueprintMyWidgetExtension class
+	virtual void Precompile(UE::MVVM::Compiler::IMVVMBlueprintViewPrecompile* Compiler, UWidgetBlueprintGeneratedClass* Class) {}
+	virtual void Compile(UE::MVVM::Compiler::IMVVMBlueprintViewCompile* Compiler, UWidgetBlueprintGeneratedClass* Class, UMVVMViewClass* ViewExtension) {}
+	virtual void WidgetRenamed(FName OldName, FName NewName) {}
+};
\ No newline at end of file
diff --git a/Engine/Plugins/Runtime/ModelViewViewModel/Source/ModelViewViewModelBlueprint/Public/MVVMBlueprintViewCompilerInterface.h b/Engine/Plugins/Runtime/ModelViewViewModel/Source/ModelViewViewModelBlueprint/Public/MVVMBlueprintViewCompilerInterface.h
new file mode 100644
--- /dev/null
+++ b/Engine/Plugins/Runtime/ModelViewViewModel/Source/ModelViewViewModelBlueprint/Public/MVVMBlueprintViewCompilerInterface.h
@@ -0,0 +1,172 @@
+// Copyright Epic Games, Inc. All Rights Reserved.
+
+#pragma once
+#include "CoreMinimal.h"
+#include "Bindings/MVVMCompiledBindingLibraryCompiler.h"
+#include "Extensions/MVVMViewClassExtension.h"
+#include "View/MVVMViewClass.h"
+
+class UMVVMBlueprintView;
+class UWidget;
+class UWidgetBlueprintGeneratedClass;
+
+namespace UE::MVVM::Compiler
+{
+
+struct FCompilerBindingHandle
+{
+public:
+	explicit FCompilerBindingHandle()
+	: Id(0)
+	{
+	}
+
+	static FCompilerBindingHandle MakeHandle()
+	{
+		FCompilerBindingHandle Handle;
+		++IdGenerator;
+		Handle.Id = IdGenerator;
+		return Handle;
+	}
+
+	bool IsValid() const
+	{
+		return Id != 0;
+	}
+
+	bool operator==(const FCompilerBindingHandle& Other) const
+	{
+		return Id == Other.Id;
+	}
+
+	bool operator!=(const FCompilerBindingHandle& Other) const
+	{
+		return Id != Other.Id;
+	}
+
+	friend uint32 GetTypeHash(const FCompilerBindingHandle& Handle)
+	{
+		return ::GetTypeHash(Handle.Id);
+	}
+
+private:
+	static int32 IdGenerator;
+	int32 Id;
+};
+
+int32 FCompilerBindingHandle::IdGenerator = 0;
+
+enum class EMessageType
+{
+	Info = 0, 
+	Warning = 1, 
+	Error = 2
+};
+
+/**
+ * Exposed interface of MVVMViewBlueprintCompiler to be used in view extensions (MVVMViewBlueprintViewExtension) at precompile step
+ */
+class IMVVMBlueprintViewPrecompile
+{
+public:
+	struct FObjectFieldPathArgs
+	{
+		FObjectFieldPathArgs(UWidgetBlueprintGeneratedClass* InClass, const FString& InObjectPath, UClass* InExpectedType)
+			: Class(InClass)
+			, ObjectPath(InObjectPath)
+			, ExpectedType(InExpectedType)
+		{
+		}
+
+		UWidgetBlueprintGeneratedClass* Class;
+		FString ObjectPath;
+		UClass* ExpectedType;
+	};
+
+	virtual const UMVVMBlueprintView* GetBlueprintView() const = 0;
+
+	/* Returns a map of widget name to the widget pointer for all widgets in the blueprint. */
+	virtual const TMap<FName, UWidget*>& GetWidgetNameToWidgetPointerMap() const = 0;
+
+	/* Returns all the bindings via handles */
+	virtual TArray<Compiler::FCompilerBindingHandle> GetAllBindings() const = 0;
+
+	/* Returns the field path for all the read fields of a binding (passed by handle) */
+	virtual TArray<TArray<UE::MVVM::FMVVMConstFieldVariant>> GetBindingReadFields(Compiler::FCompilerBindingHandle BindingHandle) = 0;
+
+	/* Returns the field path for the write field of a binding (passed by handle) */
+	virtual TArray<UE::MVVM::FMVVMConstFieldVariant> GetBindingWriteFields(Compiler::FCompilerBindingHandle BindingHandle) = 0;
+
+	/* Returns the source of a binding (passed by handle) if it's valid */
+	virtual const FProperty* GetBindingSourceProperty(Compiler::FCompilerBindingHandle BindingHandle) = 0;
+
+	/**
+	 * Add a new object property to the compiler.
+	 * Store the returned the handle of this function to access the compiled path during Compile step by calling GetFieldPath.
+	 */
+	virtual TValueOrError<FCompiledBindingLibraryCompiler::FFieldPathHandle, FText> AddObjectFieldPath(const FObjectFieldPathArgs& Args) = 0;
+	
+	/**
+	 * Adds a new field to the compiler.
+	 * Store the returned the handle of this function to access the compiled path during Compile step by calling GetFieldPath.
+	 */
+	virtual TValueOrError<FCompiledBindingLibraryCompiler::FFieldPathHandle, FText> AddFieldPath(TArrayView<const FMVVMConstFieldVariant> InFieldPath, bool bInRead) = 0;
+
+	/*
+	 * Shows a message for the given binding.
+	 * If this is an error message, it will cause the blueprint compilation to fail. Next steps will still be executed.
+	 */
+	virtual void AddMessageForBinding(Compiler::FCompilerBindingHandle BindingHandle, const FText& MessageText, EMessageType MessageType) const = 0;
+
+	/*
+	 * Shows a message in the compile logs.
+	 * If this is an error message, it will cause the blueprint compilation to fail. Next steps will still be executed.
+	 */
+	virtual void AddMessage(const FText& MessageText, EMessageType MessageType) = 0;
+
+	/*
+	 * Lets the compiler know that the precompile step for this extension failed.
+	 * Calling this will cause the blueprint compilation to fail and next steps won't be executed.
+	 */
+	virtual void MarkPrecompileStepInvalid() = 0;
+};
+
+/**
+ * Exposed interface of MVVMViewBlueprintCompiler to be used in view extensions (MVVMViewBlueprintViewExtension) at compile step
+ */
+class IMVVMBlueprintViewCompile
+{
+public:
+	virtual const UMVVMBlueprintView* GetBlueprintView() const = 0;
+
+	/* Returns a map of widget name to the widget pointer for all widgets in the blueprint. */
+	virtual const TMap<FName, UWidget*>& GetWidgetNameToWidgetPointerMap() const = 0;
+
+	/**
+	 * Get the compiled field path of the passed handle. 
+	 * The field path should have been added in the precompile step using AddObjectFieldPath or AddFieldPath and the returned handle should be stored.
+	 */
+	virtual TValueOrError<FMVVMVCompiledFieldPath, void> GetFieldPath(FCompiledBindingLibraryCompiler::FFieldPathHandle FieldPath) = 0;
+
+	/*
+	 * Shows a message for the given binding.
+	 * If this is an error message, it will cause the blueprint compilation to fail. Next steps will still be executed.
+	 */
+	virtual void AddMessageForBinding(Compiler::FCompilerBindingHandle BindingHandle, const FText& MessageText, EMessageType MessageType) const = 0;
+
+	/*
+	 * Shows a message in the compile logs.
+	 * If this is an error message, it will cause the blueprint compilation to fail. Next steps will still be executed.
+	 */
+	virtual void AddMessage(const FText& MessageText, EMessageType MessageType) = 0;
+
+	/*
+	 * Lets the compiler know that the compile step for this extension failed.
+	 * Calling this will cause the blueprint compilation to fail and next steps won't be executed.
+	 */
+	virtual void MarkCompileStepInvalid() = 0;
+
+	/* Creates an extension object of the given class and adds it to the view class to be invoked at runtime. */
+	virtual UMVVMViewClassExtension* CreateViewClassExtension(TSubclassOf<UMVVMViewClassExtension> ExtensionClass) = 0;
+};
+}
\ No newline at end of file
diff --git a/Engine/Plugins/Runtime/ModelViewViewModel/Source/ModelViewViewModelBlueprint/Public/MVVMWidgetBlueprintExtension_View.h b/Engine/Plugins/Runtime/ModelViewViewModel/Source/ModelViewViewModelBlueprint/Public/MVVMWidgetBlueprintExtension_View.h
--- a/Engine/Plugins/Runtime/ModelViewViewModel/Source/ModelViewViewModelBlueprint/Public/MVVMWidgetBlueprintExtension_View.h
+++ b/Engine/Plugins/Runtime/ModelViewViewModel/Source/ModelViewViewModelBlueprint/Public/MVVMWidgetBlueprintExtension_View.h
@@ -8,6 +8,7 @@
 #include "MVVMWidgetBlueprintExtension_View.generated.h"
 
 class UMVVMBlueprintView;
+class UMVVMBlueprintViewExtension;
 class FWidgetBlueprintCompilerContext;
 class UWidgetBlueprintGeneratedClass;
 
@@ -16,6 +17,28 @@ namespace UE::MVVM::Private
 	struct FMVVMViewBlueprintCompiler;
 } //namespace
 
+USTRUCT()
+struct FMVVMExtensionItem
+{
+	GENERATED_BODY()
+
+	UPROPERTY()
+	FName WidgetName = NAME_None;
+
+	UPROPERTY()
+	FGuid ViewmodelId = FGuid();
+
+	UPROPERTY()
+	TObjectPtr<UMVVMBlueprintViewExtension> ExtensionObj;
+
+	bool operator==(const FMVVMExtensionItem& OtherExt) const
+	{
+		return OtherExt.WidgetName == WidgetName 
+			&&	OtherExt.ExtensionObj == ExtensionObj
+			&&	OtherExt.ViewmodelId == ViewmodelId;
+	}
+};
+
 /**
  *
  */
@@ -62,6 +85,12 @@ public:
 	virtual FSearchData HandleGatherSearchData(const UBlueprint* OwningBlueprint) const override;
 	//~ End UWidgetBlueprintExtension interface
 
+	void VerifyWidgetExtensions();
+	void RenameWidgetExtensions(FName OldName, FName NewName);
+	UMVVMBlueprintViewExtension* CreateBlueprintWidgetExtension(TSubclassOf<UMVVMBlueprintViewExtension> ExtensionClass, FName WidgetName);
+	void RemoveBlueprintWidgetExtension(UMVVMBlueprintViewExtension* ExtensionToRemove, FName WidgetName);
+	TArray<UMVVMBlueprintViewExtension*> GetBlueprintExtensionsForWidget(FName WidgetName) const;
+
 	void SetFilterSettings(FMVVMViewBindingFilterSettings InFilterSettings);
 	FMVVMViewBindingFilterSettings GetFilterSettings() const
 	{
@@ -80,6 +109,11 @@ private:
 
 	FSimpleMulticastDelegate BlueprintViewChangedDelegate;
 	TPimplPtr<UE::MVVM::Private::FMVVMViewBlueprintCompiler> CurrentCompilerContext;
+
+	UPROPERTY()
+	TArray<FMVVMExtensionItem> BlueprintExtensions;
+
+	friend UE::MVVM::Private::FMVVMViewBlueprintCompiler;
 };
 
 #if UE_ENABLE_INCLUDE_ORDER_DEPRECATED_IN_5_2
diff --git a/Engine/Plugins/Runtime/ModelViewViewModel/Source/ModelViewViewModelEditor/Private/ModelViewViewModelEditorModule.cpp b/Engine/Plugins/Runtime/ModelViewViewModel/Source/ModelViewViewModelEditor/Private/ModelViewViewModelEditorModule.cpp
--- a/Engine/Plugins/Runtime/ModelViewViewModel/Source/ModelViewViewModelEditor/Private/ModelViewViewModelEditorModule.cpp
+++ b/Engine/Plugins/Runtime/ModelViewViewModel/Source/ModelViewViewModelEditor/Private/ModelViewViewModelEditorModule.cpp
@@ -7,6 +7,7 @@
 #include "BlueprintModes/WidgetBlueprintApplicationModes.h"
 #include "Customizations/MVVMBlueprintViewModelContextCustomization.h"
 #include "Customizations/MVVMPropertyBindingExtension.h"
+#include "Extensions/MVVMBlueprintViewExtension.h"
 #include "Framework/Application/SlateApplication.h"
 #include "Framework/Docking/LayoutExtender.h"
 #include "Kismet2/BlueprintEditorUtils.h"
@@ -139,6 +140,8 @@ void FModelViewViewModelEditorModule::HandleRenameVariableReferences(UBlueprint*
 			{
 				BlueprintView->WidgetRenamed(OldVarName, NewVarName);
 			}
+
+			ViewExtension->RenameWidgetExtensions(OldVarName, NewVarName);
 		}
 	}
 }
