MVVM: Move static class member definition into translation unit.

diff --git a/Engine/Plugins/Runtime/ModelViewViewModel/Source/ModelViewViewModelBlueprint/Private/MVVMViewBlueprintCompiler.cpp b/Engine/Plugins/Runtime/ModelViewViewModel/Source/ModelViewViewModelBlueprint/Private/MVVMViewBlueprintCompiler.cpp
--- a/Engine/Plugins/Runtime/ModelViewViewModel/Source/ModelViewViewModelBlueprint/Private/MVVMViewBlueprintCompiler.cpp
+++ b/Engine/Plugins/Runtime/ModelViewViewModel/Source/ModelViewViewModelBlueprint/Private/MVVMViewBlueprintCompiler.cpp
@@ -28,6 +28,11 @@
 
 #define LOCTEXT_NAMESPACE "MVVMViewBlueprintCompiler"
 
+namespace UE::MVVM::Compiler
+{
+int32 FCompilerBindingHandle::IdGenerator = 0;
+}
+
 /**
 When compiling the skeletal class
 	CreateVariables()
diff --git a/Engine/Plugins/Runtime/ModelViewViewModel/Source/ModelViewViewModelBlueprint/Public/MVVMBlueprintViewCompilerInterface.h b/Engine/Plugins/Runtime/ModelViewViewModel/Source/ModelViewViewModelBlueprint/Public/MVVMBlueprintViewCompilerInterface.h
--- a/Engine/Plugins/Runtime/ModelViewViewModel/Source/ModelViewViewModelBlueprint/Public/MVVMBlueprintViewCompilerInterface.h
+++ b/Engine/Plugins/Runtime/ModelViewViewModel/Source/ModelViewViewModelBlueprint/Public/MVVMBlueprintViewCompilerInterface.h
@@ -54,8 +54,6 @@ private:
 	int32 Id;
 };
 
-int32 FCompilerBindingHandle::IdGenerator = 0;
-
 enum class EMessageType
 {
 	Info = 0, 
