MVVM: create a ListViewBase extension to support selecting an entry viewmodel and assigning them to entry widgets at runtime.

diff --git a/Engine/Plugins/Runtime/ModelViewViewModel/Source/ModelViewViewModel/Private/Extensions/MVVMViewListViewBaseExtension.cpp b/Engine/Plugins/Runtime/ModelViewViewModel/Source/ModelViewViewModel/Private/Extensions/MVVMViewListViewBaseExtension.cpp
new file mode 100644
--- /dev/null
+++ b/Engine/Plugins/Runtime/ModelViewViewModel/Source/ModelViewViewModel/Private/Extensions/MVVMViewListViewBaseExtension.cpp
@@ -0,0 +1,96 @@
+// Copyright Epic Games, Inc. All Rights Reserved.
+
+#include "Extensions/MVVMViewListViewBaseExtension.h"
+
+#include "Bindings/MVVMFieldPathHelper.h"
+#include "Components/ListViewBase.h"
+#include "MVVMMessageLog.h"
+#include "MVVMSubsystem.h"
+#include "Templates/ValueOrError.h"
+#include "Types/MVVMFieldContext.h"
+#include "View/MVVMView.h"
+#include "View/MVVMViewClass.h"
+
+#include UE_INLINE_GENERATED_CPP_BY_NAME(MVVMViewListViewBaseExtension)
+
+#define LOCTEXT_NAMESPACE "MVVMViewListViewBaseExtension"
+
+#if WITH_EDITOR
+void UMVVMViewListViewBaseExtension::Initialize(UMVVMViewListViewBaseExtension::FInitListViewBaseExtensionArgs InArgs)
+{
+	WidgetName = InArgs.WidgetName;
+	WidgetPath = InArgs.WidgetPath;
+	EntryViewModelName = InArgs.EntryViewModelName;
+}
+#endif
+
+void UMVVMViewListViewBaseExtension::OnSourcesInitialized(UUserWidget* UserWidget, UMVVMView* View)
+{
+	check(View->GetViewClass());
+	TValueOrError<UE::MVVM::FFieldContext, void> FieldPathResult = View->GetViewClass()->GetBindingLibrary().EvaluateFieldPath(UserWidget, WidgetPath);
+
+	if (FieldPathResult.HasValue())
+	{
+		TValueOrError<UObject*, void> ObjectResult = UE::MVVM::FieldPathHelper::EvaluateObjectProperty(FieldPathResult.GetValue());
+		if (ObjectResult.HasValue() && ObjectResult.GetValue() != nullptr)
+		{
+			if (UListViewBase* ListWidget = Cast<UListViewBase>(ObjectResult.GetValue()))
+			{
+				CachedListViewWidget = ListWidget;
+				ListWidget->OnEntryWidgetGenerated().AddUObject(this, &UMVVMViewListViewBaseExtension::HandleSetViewModelOnEntryWidget, UserWidget, ListWidget, EntryViewModelName);
+			}
+			else
+			{
+				UE::MVVM::FMessageLog Log(UserWidget);
+				Log.Error(FText::Format(LOCTEXT("BindToEntryGenerationFailWidgetNotList", "The object property '{0}' is not of type ListViewBase, but has an Viewmodel extension meant for ListViewBase widgets. The extension won't have any effects.")
+					, FText::FromName(ObjectResult.GetValue()->GetFName())
+				));
+			}
+		}
+		else
+		{
+			UE::MVVM::FMessageLog Log(UserWidget);
+			Log.Error(FText::Format(LOCTEXT("BindToEntryGenerationFailInvalidObjectPropertyWidget", "The property object for list-type widget '{0}' is not found, so viewmodels won't be bound to its entries.")
+				, FText::FromName(WidgetName)
+			));
+		}
+	}
+	else
+	{
+		UE::MVVM::FMessageLog Log(UserWidget);
+		Log.Error(FText::Format(LOCTEXT("BindToEntryGenerationFailInvalidFieldPathWidget", "The field path for list-type widget '{0}' is invalid, so viewmodels won't be bound to its entries.")
+			, FText::FromName(WidgetName)
+		));
+	}
+}
+
+void UMVVMViewListViewBaseExtension::OnSourcesUninitialized(UUserWidget* UserWidget, UMVVMView* View) 
+{
+	if (UListViewBase* ListWidget = CachedListViewWidget.Get())
+	{
+		ListWidget->OnEntryWidgetGenerated().RemoveAll(this);
+	}
+}
+
+void UMVVMViewListViewBaseExtension::HandleSetViewModelOnEntryWidget(UUserWidget& EntryWidget, UUserWidget* OwningUserWidget, UListViewBase* ListWidget, FName EntryViewmodelName)
+{
+	if (UMVVMView* View = UMVVMSubsystem::GetViewFromUserWidget(&EntryWidget))
+	{
+		if (UObject* EntryObj = ListWidget->GetListObjectFromEntry(EntryWidget))
+		{
+			if (EntryObj->Implements<UNotifyFieldValueChanged>())
+			{
+				View->SetViewModel(EntryViewmodelName, EntryObj);
+			}
+			else
+			{
+				UE::MVVM::FMessageLog Log(OwningUserWidget);
+				Log.Error(FText::Format(LOCTEXT("SetViewModelOnEntryWidgetFailNotViewModel", "Trying to set an object that is not a viewmodel on entries of list-type widget '{0}'. Check the implementation of widget '{0}' and make sure the function GetListObjectFromEntry is overriden to return a viewmodel object. If you do not wish to set viewmodels on the entries of this widget, please remove the corresonding Viewmodel extension from it.")
+					, FText::FromName(WidgetName)
+				));
+			}
+		}
+	}
+}
+
+#undef LOCTEXT_NAMESPACE
\ No newline at end of file
diff --git a/Engine/Plugins/Runtime/ModelViewViewModel/Source/ModelViewViewModel/Public/Extensions/MVVMViewListViewBaseExtension.h b/Engine/Plugins/Runtime/ModelViewViewModel/Source/ModelViewViewModel/Public/Extensions/MVVMViewListViewBaseExtension.h
new file mode 100644
--- /dev/null
+++ b/Engine/Plugins/Runtime/ModelViewViewModel/Source/ModelViewViewModel/Public/Extensions/MVVMViewListViewBaseExtension.h
@@ -0,0 +1,61 @@
+// Copyright Epic Games, Inc. All Rights Reserved.
+
+#pragma once
+
+#include "Bindings/MVVMCompiledBindingLibrary.h"
+#include "MVVMViewClassExtension.h"
+#include "MVVMViewListViewBaseExtension.generated.h"
+
+class UListViewBase;
+class UMVVMView;
+class UUserWidget;
+
+UCLASS()
+class MODELVIEWVIEWMODEL_API UMVVMViewListViewBaseExtension : public UMVVMViewClassExtension
+{
+	GENERATED_BODY()
+
+public:
+	//~ Begin UMVVMViewClassExtension overrides
+	virtual void OnSourcesInitialized(UUserWidget* UserWidget, UMVVMView* View) override;
+	virtual void OnSourcesUninitialized(UUserWidget* UserWidget, UMVVMView* View) override;
+	//~ End UMVVMViewClassExtension overrides
+
+#if WITH_EDITOR
+	struct FInitListViewBaseExtensionArgs
+	{
+		FInitListViewBaseExtensionArgs(FName InWidgetName, FName InEntryViewModelName, const FMVVMVCompiledFieldPath& InWidgetPath)
+			: WidgetName(InWidgetName),
+			EntryViewModelName(InEntryViewModelName),
+			WidgetPath(InWidgetPath)
+		{}
+
+		FName WidgetName;
+		FName EntryViewModelName;
+		FMVVMVCompiledFieldPath WidgetPath;
+	};
+
+	void Initialize(FInitListViewBaseExtensionArgs InArgs);
+#endif
+
+	FName GetEntryViewModelName() const
+	{ 
+		return EntryViewModelName; 
+	}
+
+private:
+	void HandleSetViewModelOnEntryWidget(UUserWidget& EntryWidget, UUserWidget* OwningUserWidget, UListViewBase* ListWidget, FName EntryViewmodelName);
+
+private:
+	UPROPERTY(VisibleAnywhere, Category = "MVVM Extension")
+	FName WidgetName;
+
+	UPROPERTY(VisibleAnywhere, Category = "MVVM Extension")
+	FName EntryViewModelName;
+
+	UPROPERTY()
+	FMVVMVCompiledFieldPath WidgetPath;
+
+	UPROPERTY(Transient)
+	TWeakObjectPtr<UListViewBase> CachedListViewWidget;
+};
\ No newline at end of file
diff --git a/Engine/Plugins/Runtime/ModelViewViewModel/Source/ModelViewViewModelBlueprint/Private/Extensions/MVVMViewBlueprintListViewBaseExtension.cpp b/Engine/Plugins/Runtime/ModelViewViewModel/Source/ModelViewViewModelBlueprint/Private/Extensions/MVVMViewBlueprintListViewBaseExtension.cpp
new file mode 100644
--- /dev/null
+++ b/Engine/Plugins/Runtime/ModelViewViewModel/Source/ModelViewViewModelBlueprint/Private/Extensions/MVVMViewBlueprintListViewBaseExtension.cpp
@@ -0,0 +1,219 @@
+// Copyright Epic Games, Inc. All Rights Reserved.
+
+#include "Extensions/MVVMViewBlueprintListViewBaseExtension.h"
+
+#include "Bindings/MVVMCompiledBindingLibraryCompiler.h"
+#include "Blueprint/WidgetBlueprintGeneratedClass.h"
+#include "Components/ListViewBase.h"
+#include "Extensions/MVVMViewListViewBaseExtension.h"
+#include "MVVMBlueprintView.h"
+#include "MVVMBlueprintViewCompilerInterface.h"
+#include "MVVMBlueprintViewModelContext.h"
+#include "MVVMPropertyPath.h"
+#include "MVVMWidgetBlueprintExtension_View.h"
+#include "Templates/ValueOrError.h"
+#include "View/MVVMViewClass.h"
+
+#include UE_INLINE_GENERATED_CPP_BY_NAME(MVVMViewBlueprintListViewBaseExtension)
+
+#define LOCTEXT_NAMESPACE "MVVMViewBlueprintListViewBaseExtension"
+
+void UMVVMViewBlueprintListViewBaseExtension::Precompile(UE::MVVM::Compiler::IMVVMBlueprintViewPrecompile* Compiler, UWidgetBlueprintGeneratedClass* Class)
+{
+	check (Compiler);
+	auto VerifyViewmodelTypeMatch = [this, Compiler](const TSubclassOf<UUserWidget> EntryWidgetClass, const FMVVMBlueprintPropertyPath& EntryViewModelPath, const FName& ContainerWidgetName)
+	{
+		bool bFoundSetterBinding = false;
+
+		// Find bindings that have write fields with metadata "ViewmodelBlueprintWidgetExtension", these are setter functions for assigning a viewmodel array.
+		// Use their corresponding read field to get the expected type of entry viewmodel, and compare it against the one that user selected in the details panel.
+		for (const UE::MVVM::Compiler::FCompilerBindingHandle& Binding : Compiler->GetAllBindings())
+		{
+			TArray<UE::MVVM::FMVVMConstFieldVariant> WriteFields = Compiler->GetBindingWriteFields(Binding);
+			for (int32 FieldIndex = 1; FieldIndex < WriteFields.Num(); FieldIndex++)
+			{
+				UE::MVVM::FMVVMConstFieldVariant& PathField = WriteFields[FieldIndex];
+				UE::MVVM::FMVVMConstFieldVariant& ParentPathField = WriteFields[FieldIndex - 1];
+				if (PathField.IsFunction() && ParentPathField.IsProperty() && ContainerWidgetName == ParentPathField.GetName())
+				{
+					const FString MetaData = PathField.GetFunction()->GetMetaData("ViewmodelBlueprintWidgetExtension");
+					if (MetaData.Equals("EntryViewmodel", ESearchCase::IgnoreCase))
+					{
+						if (const FProperty* SourceProperty = Compiler->GetBindingSourceProperty(Binding))
+						{
+							if (const FArrayProperty* ArrayProperty = CastField<FArrayProperty>(SourceProperty))
+							{
+								// This is the viewmodel that the user selected from the details panel
+								TArray<UE::MVVM::FMVVMConstFieldVariant> SelectedVMFields = EntryViewModelPath.GetFields(EntryWidgetClass);
+								if (SelectedVMFields.Num() > 0)
+								{
+									UE::MVVM::FMVVMConstFieldVariant SelectedVMPropertyField = SelectedVMFields.Last();
+
+									if (SelectedVMPropertyField.IsProperty())
+									{
+										bFoundSetterBinding = true;
+										const FProperty* SelectedVMProperty = SelectedVMPropertyField.GetProperty();
+										if (!ArrayProperty->Inner->SameType(SelectedVMProperty))
+										{
+											const FText UserSelectedViewmodelType = FText::FromString(SelectedVMProperty->GetCPPType());
+											const FText ArraySetterFunctionViewmodelType = FText::FromString(ArrayProperty->Inner->GetCPPType());
+											Compiler->AddMessageForBinding(Binding, FText::Format(LOCTEXT("EntryViewModelTypeMismatch", "The entry viewmodel type {0} does not match the array of viewmodels of type {1}."), UserSelectedViewmodelType, ArraySetterFunctionViewmodelType), UE::MVVM::Compiler::EMessageType::Error);
+											Compiler->MarkPrecompileStepInvalid();
+											break;
+										}
+									}
+								}
+							}
+						}
+
+						// No need to look into the next fields of this write path, we have found the field with ViewmodelBlueprintWidgetExtension metadata.
+						// We will go to the next binding
+						break;
+					}
+				}
+			}
+		}
+
+		if (!bFoundSetterBinding && EntryViewModelId.IsValid())
+		{
+			Compiler->AddMessage(FText::Format(LOCTEXT("PreCompileMVVMWidgetExtensionNoSetterBindingFound", "No binding found from an array of viewmodels to {0} -> SetListItems. Please add this binding or remove the Viewmodel extension on widget {0} from its details panel."), FText::FromName(WidgetName)), UE::MVVM::Compiler::EMessageType::Warning);
+		}
+	};
+	
+	if (UWidget* const* FoundList = Compiler->GetWidgetNameToWidgetPointerMap().Find(WidgetName))
+	{
+		if (UListViewBase* ListViewBase = Cast<UListViewBase>(*FoundList))
+		{
+			TSubclassOf<UUserWidget> EntryWidgetClass = ListViewBase->GetEntryWidgetClass();
+			if (EntryWidgetClass.Get())
+			{
+				UE::MVVM::Compiler::IMVVMBlueprintViewPrecompile::FObjectFieldPathArgs AddObjFieldPathArgs(Class, WidgetName.ToString(), UListViewBase::StaticClass());
+				TValueOrError<UE::MVVM::FCompiledBindingLibraryCompiler::FFieldPathHandle, FText> FieldPathResult = Compiler->AddObjectFieldPath(AddObjFieldPathArgs);
+				if (FieldPathResult.HasError() || !FieldPathResult.GetValue().IsValid())
+				{
+					Compiler->AddMessage(FText::Format(LOCTEXT("CouldNotCreateSourceFieldPath", "Couldn't create the source field path '{0}'. {1}"), FText::FromName(WidgetName), FieldPathResult.GetError()), UE::MVVM::Compiler::EMessageType::Error);
+					Compiler->MarkPrecompileStepInvalid();
+					return;
+				}
+
+				WidgetPathHandle = FieldPathResult.StealValue();
+
+				if (const UUserWidget* EntryUserWidget = Cast<UUserWidget>(EntryWidgetClass->ClassDefaultObject))
+				{
+					if (const UMVVMBlueprintView* EntryBPView = GetEntryWidgetBlueprintView(EntryUserWidget))
+					{
+						if (const FMVVMBlueprintViewModelContext* ViewModelContext = EntryBPView->FindViewModel(EntryViewModelId))
+						{
+							const FName EntryViewModelName = ViewModelContext->GetViewModelName();
+
+							if (const UWidgetBlueprint* EntryBlueprint = Cast<UWidgetBlueprint>(EntryUserWidget->GetClass()->ClassGeneratedBy))
+							{
+								// find the property path of the entry viewmodel
+								FMVVMBlueprintPropertyPath SelectedViewModel;
+								SelectedViewModel.SetViewModelId(EntryViewModelId);
+								FProperty* ViewModelProperty = FindFProperty<FProperty>(*EntryWidgetClass, EntryViewModelName);
+								SelectedViewModel.AppendPropertyPath(EntryBlueprint, UE::MVVM::FMVVMConstFieldVariant(ViewModelProperty));
+
+								VerifyViewmodelTypeMatch(EntryWidgetClass, SelectedViewModel, ListViewBase->GetFName());
+							}
+						}
+						else
+						{
+							// If the stored viewmodel ID invalid, it means the viewmodel was deleted so we clear the value.
+							// We need to manually clear this value here because the listview widget doesn't get notified when the entry widget viewmodels change.
+							EntryViewModelId = FGuid();
+							Compiler->AddMessage(FText::Format(LOCTEXT("PreCompileMVVMWidgetExtensionEntryVMDeleted", "No viewmodel selected for Entry widget {0}. Please select a viewmodel for it via the details panel or remove the Viewmodel extension on the containing widget {1}."), FText::FromName(EntryUserWidget->GetFName()), FText::FromName(WidgetName)), UE::MVVM::Compiler::EMessageType::Error);
+
+						}
+					}
+					else
+					{
+						// If no view is found in the entry widget, we clear the entry viewmodel value.
+						// We need to manually clear this value here because the listview widget doesn't get notified when the entry widget view is added/removed.
+						EntryViewModelId = FGuid();
+						Compiler->AddMessage(FText::Format(LOCTEXT("PreCompileMVVMWidgetExtensionEntryHasNoView", "Entry widget {0} doesn't have a View. Consider adding a binding to it or remove the MVVM extension on the containing widget {1}."), FText::FromName(EntryUserWidget->GetFName()), FText::FromName(WidgetName)), UE::MVVM::Compiler::EMessageType::Error);
+					}
+				}
+			}
+		}
+		else
+		{
+			Compiler->AddMessage(FText::Format(LOCTEXT("PreCompileMVVMWidgetExtensionWidgetNotListViewBase", "Widget {0} is not a UListViewBase but has a MVVMViewBlueprintListViewBaseExtension."), FText::FromName(WidgetName)), UE::MVVM::Compiler::EMessageType::Error);
+		}
+	}
+	else
+	{
+		Compiler->AddMessage(FText::Format(LOCTEXT("PreCompileMVVMWidgetExtensionInvalidWidgetName", "Widget with name {0} doesn't exist in the widget blueprint but a viewmodel widget extension exists that is attached to it."), FText::FromName(WidgetName)),UE::MVVM::Compiler::EMessageType::Error);
+	}
+}
+
+void UMVVMViewBlueprintListViewBaseExtension::Compile(UE::MVVM::Compiler::IMVVMBlueprintViewCompile* Compiler, UWidgetBlueprintGeneratedClass* Class, UMVVMViewClass* ViewExtension)
+{
+	check(Compiler);
+
+	if (UWidget* const* FoundList = Compiler->GetWidgetNameToWidgetPointerMap().Find(WidgetName))
+	{
+		if (UListViewBase* ListViewBase = Cast<UListViewBase>(*FoundList))
+		{
+			TSubclassOf<UUserWidget> EntryWidgetClass = ListViewBase->GetEntryWidgetClass();
+			if (EntryWidgetClass.Get())
+			{
+				check(WidgetPathHandle.IsValid());
+
+				// Verify the widget property
+				const TValueOrError<FMVVMVCompiledFieldPath, void> CompiledFieldPath = Compiler->GetFieldPath(WidgetPathHandle);
+				if (CompiledFieldPath.HasError())
+				{
+					Compiler->AddMessage(FText::Format(LOCTEXT("CompiledFieldPathForWidgetNotGenerated", "The field path for widget {0} was not generated."), FText::FromName(WidgetName)), UE::MVVM::Compiler::EMessageType::Error);
+					Compiler->MarkCompileStepInvalid();
+					return;
+				}
+
+				// Check that we have a valid entry view and valid entry viewmodel
+				// Otherwise, no runtime extension will be created.
+				if (const UUserWidget* EntryUserWidget = Cast<UUserWidget>(EntryWidgetClass->ClassDefaultObject))
+				{
+					if (const UMVVMBlueprintView* EntryBPView = GetEntryWidgetBlueprintView(EntryUserWidget))
+					{
+						if (const FMVVMBlueprintViewModelContext* ViewModelContext = EntryBPView->FindViewModel(EntryViewModelId))
+						{
+							// Create the corresponding runtime extension
+							UMVVMViewClassExtension* NewExtensionObj = Compiler->CreateViewClassExtension(UMVVMViewListViewBaseExtension::StaticClass());
+							UMVVMViewListViewBaseExtension* NewExtension = CastChecked<UMVVMViewListViewBaseExtension>(NewExtensionObj);
+
+							const FName EntryViewModelName = ViewModelContext->GetViewModelName();
+							NewExtension->Initialize(UMVVMViewListViewBaseExtension::FInitListViewBaseExtensionArgs(WidgetName, EntryViewModelName, CompiledFieldPath.GetValue()));
+						
+						}
+					}
+				}
+			}
+		}
+	}	
+}
+
+const UMVVMBlueprintView* UMVVMViewBlueprintListViewBaseExtension::GetEntryWidgetBlueprintView(const UUserWidget* EntryUserWidget) const
+{
+	if (const UWidgetBlueprint* EntryBlueprint = Cast<UWidgetBlueprint>(EntryUserWidget->GetClass()->ClassGeneratedBy))
+	{
+		if (const UMVVMWidgetBlueprintExtension_View* EntryWidgetExtension = UMVVMWidgetBlueprintExtension_View::GetExtension<UMVVMWidgetBlueprintExtension_View>(EntryBlueprint))
+		{
+			if (const UMVVMBlueprintView* EntryWidgetView = EntryWidgetExtension->GetBlueprintView())
+			{
+				return EntryWidgetView;
+			}
+		}
+	}
+	return nullptr;
+}
+
+void UMVVMViewBlueprintListViewBaseExtension::WidgetRenamed(FName OldName, FName NewName)
+{
+	if (WidgetName == OldName)
+	{
+		WidgetName = NewName;
+		Modify();
+	}
+}
+
+#undef LOCTEXT_NAMESPACE
\ No newline at end of file
diff --git a/Engine/Plugins/Runtime/ModelViewViewModel/Source/ModelViewViewModelBlueprint/Public/Extensions/MVVMViewBlueprintListViewBaseExtension.h b/Engine/Plugins/Runtime/ModelViewViewModel/Source/ModelViewViewModelBlueprint/Public/Extensions/MVVMViewBlueprintListViewBaseExtension.h
new file mode 100644
--- /dev/null
+++ b/Engine/Plugins/Runtime/ModelViewViewModel/Source/ModelViewViewModelBlueprint/Public/Extensions/MVVMViewBlueprintListViewBaseExtension.h
@@ -0,0 +1,55 @@
+// Copyright Epic Games, Inc. All Rights Reserved.
+
+#pragma once
+
+#include "Bindings/MVVMCompiledBindingLibraryCompiler.h"
+#include "MVVMBlueprintViewExtension.h"
+#include "MVVMViewBlueprintListViewBaseExtension.generated.h"
+
+class UWidgetBlueprintGeneratedClass;
+class UMVVMBlueprintView;
+class UMVVMViewClass;
+class UUserWidget;
+
+namespace UE::MVVM
+{
+	class FMVVMListViewBaseExtensionCustomizationExtender;
+}
+
+namespace UE::MVVM::Compiler
+{
+	class IMVVMBlueprintViewPrecompile;
+	class IMVVMBlueprintViewCompile;
+}
+
+UCLASS()
+class MODELVIEWVIEWMODELBLUEPRINT_API UMVVMViewBlueprintListViewBaseExtension : public UMVVMBlueprintViewExtension
+{
+	GENERATED_BODY()
+
+public:
+	//~ Begin UMVVMBlueprintViewExtension overrides
+	virtual void Precompile(UE::MVVM::Compiler::IMVVMBlueprintViewPrecompile* Compiler, UWidgetBlueprintGeneratedClass* Class) override;
+	virtual void Compile(UE::MVVM::Compiler::IMVVMBlueprintViewCompile* Compiler, UWidgetBlueprintGeneratedClass* Class, UMVVMViewClass* ViewExtension) override;
+	virtual void WidgetRenamed(FName OldName, FName NewName) override;
+	//~ End UMVVMBlueprintViewExtension overrides
+
+	FGuid GetEntryViewModelId() const
+	{
+		return EntryViewModelId;
+	}
+
+private:
+	const UMVVMBlueprintView* GetEntryWidgetBlueprintView(const UUserWidget* EntryUserWidget) const;
+
+private:
+	UPROPERTY()
+	FName WidgetName;
+
+	UPROPERTY()
+	FGuid EntryViewModelId;
+
+	UE::MVVM::FCompiledBindingLibraryCompiler::FFieldPathHandle WidgetPathHandle;
+
+	friend UE::MVVM::FMVVMListViewBaseExtensionCustomizationExtender;
+};
diff --git a/Engine/Plugins/Runtime/ModelViewViewModel/Source/ModelViewViewModelEditor/Private/Customizations/MVVMListViewBaseExtensionCustomizationExtender.cpp b/Engine/Plugins/Runtime/ModelViewViewModel/Source/ModelViewViewModelEditor/Private/Customizations/MVVMListViewBaseExtensionCustomizationExtender.cpp
new file mode 100644
--- /dev/null
+++ b/Engine/Plugins/Runtime/ModelViewViewModel/Source/ModelViewViewModelEditor/Private/Customizations/MVVMListViewBaseExtensionCustomizationExtender.cpp
@@ -0,0 +1,337 @@
+// Copyright Epic Games, Inc. All Rights Reserved.
+
+#include "MVVMListViewBaseExtensionCustomizationExtender.h"
+
+#include "AssetRegistry/AssetRegistryModule.h"
+#include "Components/ListViewBase.h"
+#include "DetailCategoryBuilder.h"
+#include "DetailLayoutBuilder.h"
+#include "Extensions/MVVMViewBlueprintListViewBaseExtension.h"
+#include "Kismet2/BlueprintEditorUtils.h"
+#include "MVVMWidgetBlueprintExtension_View.h"
+#include "MVVMBlueprintView.h"
+#include "MVVMBlueprintViewModel.h"
+#include "PropertyCustomizationHelpers.h"
+#include "PropertyHandle.h"
+#include "WidgetBlueprint.h"
+#include "WidgetBlueprintEditor.h"
+#include "Widgets/Input/SButton.h"
+#include "Widgets/Input/SComboButton.h"
+#include "Widgets/Text/STextBlock.h"
+
+#define LOCTEXT_NAMESPACE "MVVMListViewBaseExtensionCustomizationExtender"
+
+namespace UE::MVVM
+{
+
+TSharedPtr<FMVVMListViewBaseExtensionCustomizationExtender> FMVVMListViewBaseExtensionCustomizationExtender::MakeInstance()
+{
+	return MakeShared<FMVVMListViewBaseExtensionCustomizationExtender>();
+}
+
+void FMVVMListViewBaseExtensionCustomizationExtender::CustomizeDetails(IDetailLayoutBuilder& InDetailLayout, const TArrayView<UWidget*> InWidgets, const TSharedRef<FWidgetBlueprintEditor>& InWidgetBlueprintEditor)
+{
+	// multi-selection not supported for the data
+	if (InWidgets.Num() == 1)
+	{
+		if (UListViewBase* ListView = Cast<UListViewBase>(InWidgets[0]))
+		{
+			Widget = ListView;
+			WidgetBlueprintEditor = InWidgetBlueprintEditor;
+
+			// Only do a customization if we have a MVVM blueprint view class on this blueprint.
+			if (GetExtensionViewForSelectedWidgetBlueprint())
+			{
+				IDetailCategoryBuilder& MVVMCategory = InDetailLayout.EditCategory("ListEntries");
+
+				// Fetch the entry widget class handle.
+				TArray<TSharedRef<IPropertyHandle>> ListEntryProperties;
+				MVVMCategory.GetDefaultProperties(ListEntryProperties);
+				TSharedRef<IPropertyHandle>* EntryClassPtr = ListEntryProperties.FindByPredicate( [](const TSharedRef<IPropertyHandle> Property) 
+				{ 
+					return Property->GetPropertyDisplayName().EqualToCaseIgnored(FText::FromString(TEXT("Entry Widget Class"))); 
+				});
+
+				check(EntryClassPtr);
+				EntryClassHandle = *EntryClassPtr;
+				EntryClassHandle->SetOnPropertyValueChanged(FSimpleDelegate::CreateSP(this, &FMVVMListViewBaseExtensionCustomizationExtender::HandleEntryClassChanged, false));
+				bIsExtensionAdded = GetListBaseViewExtension() != nullptr;
+				HandleEntryClassChanged(true);
+
+				// Add a button that controls adding/removing the extension on the ListViewBase widget
+				MVVMCategory.AddCustomRow(FText::FromString(TEXT("Viewmodel")))
+				.NameContent()
+				[
+					SNew(STextBlock)
+					.Text(LOCTEXT("VMSupport", "Viewmodel Support"))
+					.Font(IDetailLayoutBuilder::GetDetailFont())
+				]
+				.ValueContent()
+				.HAlign(HAlign_Fill)
+				[
+					SNew(SHorizontalBox)
+					+ SHorizontalBox::Slot()
+					.AutoWidth()
+					[
+						SNew(SButton)
+						.OnClicked(this, &FMVVMListViewBaseExtensionCustomizationExtender::ModifyExtension)
+						[
+							SNew(SHorizontalBox)
+							+ SHorizontalBox::Slot()
+							.HAlign(HAlign_Center)
+							.VAlign(VAlign_Center)
+							.AutoWidth()
+							[
+								SNew(SImage)
+								.Image(this, &FMVVMListViewBaseExtensionCustomizationExtender::GetExtensionButtonIcon)
+							]
+							+ SHorizontalBox::Slot()
+							.Padding(FMargin(3, 0, 0, 0))
+							.VAlign(VAlign_Center)
+							.AutoWidth()
+							[
+								SNew(STextBlock)
+								.TextStyle(FAppStyle::Get(), "SmallButtonText")
+								.Text(this, &FMVVMListViewBaseExtensionCustomizationExtender::GetExtensionButtonText)
+							]
+						]
+					]
+				];
+
+				// Add a combobox that allows selecting from the viewmodels in the entry widgets
+				MVVMCategory.AddCustomRow(FText::FromString(TEXT("Viewmodel")))
+				.Visibility(TAttribute<EVisibility>::Create(TAttribute<EVisibility>::FGetter::CreateSP(this, &FMVVMListViewBaseExtensionCustomizationExtender::GetEntryViewModelVisibility)))
+				.NameContent()
+				[
+					SNew(STextBlock)
+					.Text(LOCTEXT("EntryVM", "Entry Viewmodel"))
+					.Font(IDetailLayoutBuilder::GetDetailFont())
+				]
+				.ValueContent()
+				.HAlign(HAlign_Fill)
+				[
+					SNew(SHorizontalBox)
+					+ SHorizontalBox::Slot()
+					.AutoWidth()
+					[
+						SNew(SComboButton)
+						.OnGetMenuContent(this, &FMVVMListViewBaseExtensionCustomizationExtender::OnGetViewModelsMenuContent)
+						.ButtonContent()
+						[
+							SNew(STextBlock)
+							.Text(this, &FMVVMListViewBaseExtensionCustomizationExtender::OnGetSelectedViewModel)
+						]
+					]
+					+ SHorizontalBox::Slot()
+					.AutoWidth()
+					[
+						PropertyCustomizationHelpers::MakeClearButton(
+						FSimpleDelegate::CreateSP(this, &FMVVMListViewBaseExtensionCustomizationExtender::ClearEntryViewModel))
+					]
+				];
+			}
+		}
+	}
+}
+
+FReply FMVVMListViewBaseExtensionCustomizationExtender::ModifyExtension()
+{
+	if (UMVVMViewBlueprintListViewBaseExtension* ListBaseViewExtension = GetListBaseViewExtension())
+	{
+		if (UListViewBase* WidgetPtr = Widget.Get())
+		{
+			GetExtensionViewForSelectedWidgetBlueprint()->RemoveBlueprintWidgetExtension(ListBaseViewExtension, WidgetPtr->GetFName());
+			bIsExtensionAdded = false;
+		}
+	}
+	else
+	{
+		CreateListBaseViewExtensionIfNotExisting();
+		bIsExtensionAdded = true;
+	}
+	return FReply::Handled();
+}
+
+void FMVVMListViewBaseExtensionCustomizationExtender::CreateListBaseViewExtensionIfNotExisting()
+{
+	if (UMVVMWidgetBlueprintExtension_View* Extension = GetExtensionViewForSelectedWidgetBlueprint())
+	{
+		if (UListViewBase* WidgetPtr = Widget.Get())
+		{
+			if (Extension->GetBlueprintExtensionsForWidget(WidgetPtr->GetFName()).IsEmpty())
+			{
+				UMVVMBlueprintViewExtension* NewExtension = Extension->CreateBlueprintWidgetExtension(UMVVMViewBlueprintListViewBaseExtension::StaticClass(), WidgetPtr->GetFName());
+				UMVVMViewBlueprintListViewBaseExtension* NewListViewExtension = CastChecked<UMVVMViewBlueprintListViewBaseExtension>(NewExtension);
+				NewListViewExtension->WidgetName = WidgetPtr->GetFName();
+			}
+		}
+	}
+}
+
+UMVVMViewBlueprintListViewBaseExtension* FMVVMListViewBaseExtensionCustomizationExtender::GetListBaseViewExtension() const
+{
+	if (UMVVMWidgetBlueprintExtension_View* ViewClass = GetExtensionViewForSelectedWidgetBlueprint())
+	{
+		if (UListViewBase* WidgetPtr = Widget.Get())
+		{
+			for (UMVVMBlueprintViewExtension* Extension : ViewClass->GetBlueprintExtensionsForWidget(WidgetPtr->GetFName()))
+			{
+				if (UMVVMViewBlueprintListViewBaseExtension* ListViewBaseExtension = Cast<UMVVMViewBlueprintListViewBaseExtension>(Extension))
+				{
+					return ListViewBaseExtension;
+				}
+
+			}
+		}
+	}
+
+	return nullptr;
+}
+
+UMVVMWidgetBlueprintExtension_View* FMVVMListViewBaseExtensionCustomizationExtender::GetExtensionViewForSelectedWidgetBlueprint() const
+{
+	if (const TSharedPtr<FWidgetBlueprintEditor> BPEditor = WidgetBlueprintEditor.Pin())
+	{
+		if (const UWidgetBlueprint* Blueprint = BPEditor->GetWidgetBlueprintObj())
+		{
+			return UMVVMWidgetBlueprintExtension_View::GetExtension<UMVVMWidgetBlueprintExtension_View>(Blueprint);
+		}
+	}
+
+	return nullptr;
+}
+
+void FMVVMListViewBaseExtensionCustomizationExtender::ClearEntryViewModel()
+{
+	SetEntryViewModel(FGuid());
+}
+
+void FMVVMListViewBaseExtensionCustomizationExtender::HandleEntryClassChanged(bool bIsInit)
+{
+	// Update the cached value of entry class
+	void* EntryClassPtr = nullptr;
+	TSubclassOf<UUserWidget>* EntryClassValue = nullptr;
+	bool bEntryClassChanged = false;
+
+	if (EntryClassHandle->IsValidHandle() && EntryClassHandle->GetValueData(EntryClassPtr) == FPropertyAccess::Success)
+	{
+		EntryClassValue = static_cast<TSubclassOf<UUserWidget>*>(EntryClassPtr);
+	}
+
+	if (!EntryClassValue || !EntryClass.Get() || EntryClass.Get() != *EntryClassValue)
+	{
+		bEntryClassChanged = true;
+	}
+	EntryClass = EntryClassValue ? *EntryClassValue : nullptr;
+
+	// Update other values that depend on the entry class (only if the cached value actually changed)
+	if (bEntryClassChanged && EntryClass)
+	{
+		if (UUserWidget* EntryCDO = Cast<UUserWidget>(EntryClass->ClassDefaultObject))
+		{
+			EntryWidgetBlueprint = Cast<UWidgetBlueprint>(EntryCDO->GetClass()->ClassGeneratedBy);
+		}
+
+		// Clear the saved entry viewmodel if we're not calling this from CustomizeDetails (not initializing the customizer)
+		if (!bIsInit)
+		{
+			SetEntryViewModel(FGuid(), false);
+		}
+	}
+}
+
+FText FMVVMListViewBaseExtensionCustomizationExtender::OnGetSelectedViewModel() const
+{
+	if (const UListViewBase* WidgetPtr = Widget.Get())
+	{
+		if (EntryClass)
+		{
+			if (UMVVMViewBlueprintListViewBaseExtension* ListBaseViewExtension = GetListBaseViewExtension())
+			{
+				if (const UUserWidget* EntryUserWidget = Cast<UUserWidget>(EntryClass->ClassDefaultObject))
+				{
+					if (const UWidgetBlueprint* EntryBlueprint = Cast<UWidgetBlueprint>(EntryUserWidget->GetClass()->ClassGeneratedBy))
+					{
+						if (const UMVVMWidgetBlueprintExtension_View* EntryWidgetExtension = UMVVMWidgetBlueprintExtension_View::GetExtension<UMVVMWidgetBlueprintExtension_View>(EntryBlueprint))
+						{
+							if (const UMVVMBlueprintView* EntryWidgetView = EntryWidgetExtension->GetBlueprintView())
+							{
+								if (const FMVVMBlueprintViewModelContext* ViewModelContext = EntryWidgetView->FindViewModel(ListBaseViewExtension->GetEntryViewModelId()))
+								{
+									return FText::FromName(ViewModelContext->GetViewModelName());
+								}
+							}
+						}
+					}
+				}
+			}
+		}
+	}
+
+	return LOCTEXT("NoViewmodel", "No Viewmodel");
+}
+
+FText FMVVMListViewBaseExtensionCustomizationExtender::GetExtensionButtonText() const
+{
+	return bIsExtensionAdded ? LOCTEXT("RemoveVMExt", "Remove Viewmodel Extension") : LOCTEXT("AddVMExt", "Add Viewmodel Extension");
+}
+
+const FSlateBrush* FMVVMListViewBaseExtensionCustomizationExtender::GetExtensionButtonIcon() const
+{
+	return bIsExtensionAdded ? FAppStyle::Get().GetBrush("Icons.X") : FAppStyle::Get().GetBrush("Icons.Plus");
+}
+
+TSharedRef<SWidget> FMVVMListViewBaseExtensionCustomizationExtender::OnGetViewModelsMenuContent()
+{
+	FMenuBuilder MenuBuilder(true, NULL);
+
+	if (EntryClass)
+	{
+		// Find all viewmodels in the entry widget
+		if (const UWidgetBlueprint* EntryWidgetBlueprintPtr = EntryWidgetBlueprint.Get())
+		{
+			if (const UMVVMWidgetBlueprintExtension_View* EntryWidgetExtension = UMVVMWidgetBlueprintExtension_View::GetExtension<UMVVMWidgetBlueprintExtension_View>(EntryWidgetBlueprintPtr))
+			{
+				if (const UMVVMBlueprintView* EntryWidgetView = EntryWidgetExtension->GetBlueprintView())
+				{
+					const TArrayView<const FMVVMBlueprintViewModelContext> ViewModels = EntryWidgetView->GetViewModels();
+					for (const FMVVMBlueprintViewModelContext& EntryViewModel : ViewModels)
+					{
+						// Create the menu action for this entry viewmodel
+						FUIAction ItemAction(FExecuteAction::CreateSP(this, &FMVVMListViewBaseExtensionCustomizationExtender::SetEntryViewModel, EntryViewModel.GetViewModelId(), true));
+						MenuBuilder.AddMenuEntry(FText::FromName(EntryViewModel.GetViewModelName()), TAttribute<FText>(), FSlateIcon(), ItemAction);
+					}
+				}
+			}
+		}
+	}
+	return MenuBuilder.MakeWidget();
+}
+
+void FMVVMListViewBaseExtensionCustomizationExtender::SetEntryViewModel(FGuid InEntryViewModelId, bool bMarkModified)
+{
+	if (UMVVMWidgetBlueprintExtension_View* Extension = GetExtensionViewForSelectedWidgetBlueprint())
+	{
+		if (const UListViewBase* WidgetPtr = Widget.Get())
+		{
+			if (UMVVMViewBlueprintListViewBaseExtension* ListBaseViewExtension = GetListBaseViewExtension())
+			{
+				ListBaseViewExtension->Modify();
+				ListBaseViewExtension->EntryViewModelId = InEntryViewModelId;
+				if (bMarkModified)
+				{
+					if (const TSharedPtr<FWidgetBlueprintEditor> BPEditor = WidgetBlueprintEditor.Pin())
+					{
+						if (UWidgetBlueprint* Blueprint = BPEditor->GetWidgetBlueprintObj())
+						{
+							FBlueprintEditorUtils::MarkBlueprintAsStructurallyModified(Blueprint);
+						}
+					}
+				}
+			}
+		}
+	}
+}
+
+}
+#undef LOCTEXT_NAMESPACE
\ No newline at end of file
diff --git a/Engine/Plugins/Runtime/ModelViewViewModel/Source/ModelViewViewModelEditor/Private/Customizations/MVVMListViewBaseExtensionCustomizationExtender.h b/Engine/Plugins/Runtime/ModelViewViewModel/Source/ModelViewViewModelEditor/Private/Customizations/MVVMListViewBaseExtensionCustomizationExtender.h
new file mode 100644
--- /dev/null
+++ b/Engine/Plugins/Runtime/ModelViewViewModel/Source/ModelViewViewModelEditor/Private/Customizations/MVVMListViewBaseExtensionCustomizationExtender.h
@@ -0,0 +1,79 @@
+// Copyright Epic Games, Inc. All Rights Reserved.
+
+#pragma once
+
+#include "Customizations/IBlueprintWidgetCustomizationExtender.h"
+#include "Blueprint/UserWidget.h"
+#include "Layout/Visibility.h"
+#include "MVVMPropertyPath.h"
+#include "UObject/WeakObjectPtr.h"
+
+class FReply;
+class IPropertyHandle;
+class SWidget;
+class UListViewBase;
+class UMVVMViewBlueprintListViewBaseExtension;
+class UMVVMWidgetBlueprintExtension_View;
+
+namespace  UE::MVVM
+{
+
+class FMVVMListViewBaseExtensionCustomizationExtender : public IBlueprintWidgetCustomizationExtender
+{
+public:
+	static TSharedPtr<FMVVMListViewBaseExtensionCustomizationExtender> MakeInstance();
+	virtual void CustomizeDetails(IDetailLayoutBuilder& InDetailLayout, const TArrayView<UWidget*> InWidgets, const TSharedRef<FWidgetBlueprintEditor>& InWidgetBlueprintEditor);
+
+private:
+	/** Set the entry viewmodel on the MVVMViewBlueprintListViewBaseExtension for this widget. */
+	void SetEntryViewModel(FGuid InEntryViewModelId, bool bMarkModified = true);
+
+	/** Get a list of all viewmodels in the entry class. */
+	TSharedRef<SWidget> OnGetViewModelsMenuContent();
+
+	/** Get the name of the currently-selected entry viewmodel from the extension. */
+	FText OnGetSelectedViewModel() const;
+
+	/** Clear the entry viewmodel on the MVVMViewBlueprintListViewBaseExtension for this widget. */
+	void ClearEntryViewModel();
+
+	/** Update the cached variables when the entry class property changes. */
+	void HandleEntryClassChanged(bool bIsInit);
+
+	/** Get the visibility of entry viewmodel row in the details panel. */
+	EVisibility GetEntryViewModelVisibility() const
+	{
+		return bIsExtensionAdded ? EVisibility::Visible : EVisibility::Collapsed;
+	}
+
+	/** Create a new MVVMViewBlueprintListViewBaseExtension for this widget in the blueprint view class. */
+	void CreateListBaseViewExtensionIfNotExisting();
+
+	/** Get the MVVMViewBlueprintListViewBaseExtension for this widget in the blueprint view class. */
+	UMVVMViewBlueprintListViewBaseExtension* GetListBaseViewExtension() const;
+
+	/** Get the MVVM blueprint view class of this widget blueprint. */
+	UMVVMWidgetBlueprintExtension_View* GetExtensionViewForSelectedWidgetBlueprint() const;
+
+	/** Add/Remove the MVVMViewBlueprintListViewBaseExtension for this widget on button click. */
+	FReply ModifyExtension();
+
+	/** Get + or X icon for the MVVM extension button. */
+	const FSlateBrush* GetExtensionButtonIcon() const;
+
+	/** Get display text for the MVVM extension button. */
+	FText GetExtensionButtonText() const;
+
+private:
+	/** The selected ListViewBase widget in the details panel. */
+	TWeakObjectPtr<UListViewBase> Widget;
+	TWeakPtr<FWidgetBlueprintEditor> WidgetBlueprintEditor;
+
+	TSubclassOf<UUserWidget> EntryClass;
+	TSharedPtr<IPropertyHandle> EntryClassHandle;
+	TWeakObjectPtr<UWidgetBlueprint> EntryWidgetBlueprint;
+
+	/** Keep track of whether we have a MVVMViewBlueprintListViewBaseExtension for this widget. */
+	bool bIsExtensionAdded = false;
+};
+}
\ No newline at end of file
diff --git a/Engine/Plugins/Runtime/ModelViewViewModel/Source/ModelViewViewModelEditor/Private/ModelViewViewModelEditorModule.cpp b/Engine/Plugins/Runtime/ModelViewViewModel/Source/ModelViewViewModelEditor/Private/ModelViewViewModelEditorModule.cpp
--- a/Engine/Plugins/Runtime/ModelViewViewModel/Source/ModelViewViewModelEditor/Private/ModelViewViewModelEditorModule.cpp
+++ b/Engine/Plugins/Runtime/ModelViewViewModel/Source/ModelViewViewModelEditor/Private/ModelViewViewModelEditorModule.cpp
@@ -6,6 +6,7 @@
 #include "BlueprintModes/WidgetBlueprintApplicationMode.h"
 #include "BlueprintModes/WidgetBlueprintApplicationModes.h"
 #include "Customizations/MVVMBlueprintViewModelContextCustomization.h"
+#include "Customizations/MVVMListViewBaseExtensionCustomizationExtender.h"
 #include "Customizations/MVVMPropertyBindingExtension.h"
 #include "Extensions/MVVMBlueprintViewExtension.h"
 #include "Framework/Application/SlateApplication.h"
@@ -45,6 +46,9 @@ void FModelViewViewModelEditorModule::StartupModule()
 	PropertyBindingExtension = MakeShared<UE::MVVM::FMVVMPropertyBindingExtension>();
 	UMGEditorModule.GetPropertyBindingExtensibilityManager()->AddExtension(PropertyBindingExtension.ToSharedRef());
 
+	ListViewBaseCustomizationExtender = UE::MVVM::FMVVMListViewBaseExtensionCustomizationExtender::MakeInstance();
+	UMGEditorModule.AddWidgetCustomizationExtender(ListViewBaseCustomizationExtender.ToSharedRef());
+
 	UMGEditorModule.RegisterInstancedCustomPropertyTypeLayout(
 		FMVVMBlueprintViewModelContext::StaticStruct()->GetStructPathName()
 		, IUMGEditorModule::FOnGetInstancePropertyTypeCustomizationInstance::CreateStatic(UE::MVVM::FBlueprintViewModelContextDetailCustomization::MakeInstance)
diff --git a/Engine/Plugins/Runtime/ModelViewViewModel/Source/ModelViewViewModelEditor/Private/ModelViewViewModelEditorModule.h b/Engine/Plugins/Runtime/ModelViewViewModel/Source/ModelViewViewModelEditor/Private/ModelViewViewModelEditorModule.h
--- a/Engine/Plugins/Runtime/ModelViewViewModel/Source/ModelViewViewModelEditor/Private/ModelViewViewModelEditorModule.h
+++ b/Engine/Plugins/Runtime/ModelViewViewModel/Source/ModelViewViewModelEditor/Private/ModelViewViewModelEditorModule.h
@@ -8,6 +8,7 @@
 namespace UE::MVVM
 {
 	class FAssetTypeActions_ViewModelBlueprint;
+	class FMVVMListViewBaseExtensionCustomizationExtender;
 	class FMVVMPropertyBindingExtension;
 }
 
@@ -44,4 +45,5 @@ private:
 private:
 	TSharedPtr<UE::MVVM::FMVVMPropertyBindingExtension> PropertyBindingExtension;
 	TSharedPtr<UE::MVVM::FAssetTypeActions_ViewModelBlueprint> ViewModelBlueprintActions;
+	TSharedPtr<UE::MVVM::FMVVMListViewBaseExtensionCustomizationExtender> ListViewBaseCustomizationExtender;
 };
diff --git a/Engine/Source/Runtime/UMG/Private/Components/ListView.cpp b/Engine/Source/Runtime/UMG/Private/Components/ListView.cpp
--- a/Engine/Source/Runtime/UMG/Private/Components/ListView.cpp
+++ b/Engine/Source/Runtime/UMG/Private/Components/ListView.cpp
@@ -311,6 +311,12 @@ UUserWidget& UListView::OnGenerateEntryWidgetInternal(UObject* Item, TSubclassOf
 	return GenerateTypedEntry(DesiredEntryClass, OwnerTable);
 }
 
+UObject* UListView::GetListObjectFromEntry(UUserWidget& EntryWidget)
+{
+	const TObjectPtrWrapTypeOf<UObject*>* Item = ITypedUMGListView::ItemFromEntryWidget(EntryWidget);
+	return Item ? *Item : nullptr;
+}
+
 void UListView::OnItemClickedInternal(UObject* ListItem)
 {
 	ITypedUMGListView::OnItemClickedInternal(ListItem);
diff --git a/Engine/Source/Runtime/UMG/Public/Components/ListView.h b/Engine/Source/Runtime/UMG/Public/Components/ListView.h
--- a/Engine/Source/Runtime/UMG/Public/Components/ListView.h
+++ b/Engine/Source/Runtime/UMG/Public/Components/ListView.h
@@ -123,6 +123,8 @@ public:
 	UFUNCTION(BlueprintCallable, Category = ListView)
 	UMG_API void NavigateToIndex(int32 Index);
 
+	UMG_API virtual UObject* GetListObjectFromEntry(UUserWidget& EntryWidget) override;
+
 protected:
 	UMG_API virtual void OnItemsChanged(const TArray<UObject*>& AddedItems, const TArray<UObject*>& RemovedItems);
 
@@ -275,7 +277,7 @@ private:
 	UMG_API void BP_CancelScrollIntoView();
 
 	/** Sets the array of objects to display rows for in the list */
-	UFUNCTION(BlueprintCallable, Category = ListView, meta = (AllowPrivateAccess = true, DisplayName = "Set List Items"))
+	UFUNCTION(BlueprintCallable, Category = ListView, meta = (AllowPrivateAccess = true, DisplayName = "Set List Items", ViewmodelBlueprintWidgetExtension = "EntryViewModel"))
 	UMG_API void BP_SetListItems(const TArray<UObject*>& InListItems);
 
 	/** Gets the first selected item, if any; recommended that you only use this for single selection lists. */
diff --git a/Engine/Source/Runtime/UMG/Public/Components/ListViewBase.h b/Engine/Source/Runtime/UMG/Public/Components/ListViewBase.h
--- a/Engine/Source/Runtime/UMG/Public/Components/ListViewBase.h
+++ b/Engine/Source/Runtime/UMG/Public/Components/ListViewBase.h
@@ -514,6 +514,9 @@ public:
 	UFUNCTION(BlueprintCallable, Category = ListViewBase)
 	UMG_API float GetScrollOffset() const;
 
+	/** Get the corresponding list object for this userwidget entry. Override this to call ITypedUMGListView::ItemFromEntryWidget in concrete widgets. */
+	UMG_API virtual UObject* GetListObjectFromEntry(UUserWidget& EntryWidget) { return nullptr; }
+
 	/**
 	 * Full regeneration of all entries in the list. Note that the entry UWidget instances will not be destroyed, but they will be released and re-generated.
 	 * In other words, entry widgets will not receive Destruct/Construct events. They will receive OnEntryReleased and IUserObjectListEntry implementations will receive OnListItemObjectSet.
