MVVM: Fix crashes when the application closes and the binding subsystem is no longer available.

diff --git a/Engine/Plugins/Runtime/ModelViewViewModel/Source/ModelViewViewModel/Private/View/MVVMView.cpp b/Engine/Plugins/Runtime/ModelViewViewModel/Source/ModelViewViewModel/Private/View/MVVMView.cpp
--- a/Engine/Plugins/Runtime/ModelViewViewModel/Source/ModelViewViewModel/Private/View/MVVMView.cpp
+++ b/Engine/Plugins/Runtime/ModelViewViewModel/Source/ModelViewViewModel/Private/View/MVVMView.cpp
@@ -116,7 +116,7 @@ void UMVVMView::InitializeSources()
 
 	SCOPE_CYCLE_COUNTER(STAT_UMG_Viewmodel_InitializeSources);
 
-	check(bHasDefaultTickBinding == false && NumberOfSourceWithTickBinding == 0);
+	ensure(bHasDefaultTickBinding == false && NumberOfSourceWithTickBinding == 0);
 
 	// Init Sources/ViewModel instances
 	const int32 NumberOfSources = GeneratedViewClass->GetSources().Num();
@@ -364,13 +364,16 @@ void UMVVMView::UninitializeBindings()
 
 	// Remove all delayed bindings from the 
 	UMVVMBindingSubsystem* BindingSubsystem = GEngine->GetEngineSubsystem<UMVVMBindingSubsystem>();
-	BindingSubsystem->RemoveDelayedBindings(this);
-
-	if (NumberOfSourceWithTickBinding == 0 && bHasDefaultTickBinding)
+	if (BindingSubsystem)
 	{
-		BindingSubsystem->RemoveViewWithTickBinding(this);
-		bHasDefaultTickBinding = false;
+		BindingSubsystem->RemoveDelayedBindings(this);
+
+		if (NumberOfSourceWithTickBinding == 0 && bHasDefaultTickBinding)
+		{
+			BindingSubsystem->RemoveViewWithTickBinding(this);
+		}
 	}
+	bHasDefaultTickBinding = false;
 }
 
 
@@ -459,7 +462,10 @@ void UMVVMView::UninitializeSourceBindings(FMVVMViewClass_SourceKey SourceKey, c
 			check(NumberOfSourceWithTickBinding >= 0);
 			if (NumberOfSourceWithTickBinding == 0 && !bHasDefaultTickBinding)
 			{
-				GEngine->GetEngineSubsystem<UMVVMBindingSubsystem>()->RemoveViewWithTickBinding(this);
+				if (UMVVMBindingSubsystem* BindingSubsystem = GEngine->GetEngineSubsystem<UMVVMBindingSubsystem>())
+				{
+					BindingSubsystem->RemoveViewWithTickBinding(this);
+				}
 			}
 		}
 	}
