MVVM: Fix function and property filtering on the generated bp class (userwidget).

diff --git a/Engine/Plugins/Runtime/ModelViewViewModel/Source/ModelViewViewModelBlueprint/Private/MVVMDeveloperProjectSettings.cpp b/Engine/Plugins/Runtime/ModelViewViewModel/Source/ModelViewViewModelBlueprint/Private/MVVMDeveloperProjectSettings.cpp
--- a/Engine/Plugins/Runtime/ModelViewViewModel/Source/ModelViewViewModelBlueprint/Private/MVVMDeveloperProjectSettings.cpp
+++ b/Engine/Plugins/Runtime/ModelViewViewModel/Source/ModelViewViewModelBlueprint/Private/MVVMDeveloperProjectSettings.cpp
@@ -84,11 +84,18 @@ bool UMVVMDeveloperProjectSettings::PropertyHasFiltering(const UStruct* ObjectSt
 
 namespace UE::MVVM::Private
 {
-bool ShouldDoFieldEditorPermission(const UBlueprint* GeneratingFor, const UClass* FieldOwner)
+//class ClassA { int A };
+//class ClassB { };
+//MyClassB.A; Maybe ClassB doesn't have the permission to use ClassA::A. Maybe MyClassB has the persmission but MyClassA doesn't have it.
+
+//GeneratingFor: the blueprint it's is executed from
+//AccessorOwner: the ClassB
+//FieldClassOwner: ClassA
+bool ShouldDoFieldEditorPermission(const UBlueprint* GeneratingFor, const UClass* AccessorOwner, const UClass* FieldClassOwner)
 {
-	if (GeneratingFor && FieldOwner)
+	if (GeneratingFor && FieldClassOwner)
 	{
-		const UClass* UpToDateClass = FBlueprintEditorUtils::GetMostUpToDateClass(FieldOwner);
+		const UClass* UpToDateClass = FBlueprintEditorUtils::GetMostUpToDateClass(FieldClassOwner);
 		return GeneratingFor->SkeletonGeneratedClass != UpToDateClass;
 	}
 	return true;
@@ -104,7 +111,7 @@ bool UMVVMDeveloperProjectSettings::IsPropertyAllowed(const UBlueprint* Generati
 	const UClass* AuthoritativeClass = Cast<const UClass>(ObjectStruct);
 	AuthoritativeClass = AuthoritativeClass ? AuthoritativeClass->GetAuthoritativeClass() : nullptr;
 
-	const bool bDoPropertyEditorPermission = UE::MVVM::Private::ShouldDoFieldEditorPermission(GeneratingFor, AuthoritativeClass);
+	const bool bDoPropertyEditorPermission = UE::MVVM::Private::ShouldDoFieldEditorPermission(GeneratingFor, AuthoritativeClass, Property->GetOwnerClass());
 	if (bDoPropertyEditorPermission)
 	{
 		if (!FPropertyEditorPermissionList::Get().DoesPropertyPassFilter(AuthoritativeClass, Property->GetFName()))
@@ -150,11 +157,11 @@ bool UMVVMDeveloperProjectSettings::IsFunctionAllowed(const UBlueprint* Generati
 		return false;
 	}
 
-	const bool bDoPropertyEditorPermission = UE::MVVM::Private::ShouldDoFieldEditorPermission(GeneratingFor, AuthoritativeClass);
-	if (bDoPropertyEditorPermission)
+	const FPathPermissionList& FunctionPermissions = GetMutableDefault<UBlueprintEditorSettings>()->GetFunctionPermissions();
+	if (FunctionPermissions.HasFiltering())
 	{
-		const FPathPermissionList& FunctionPermissions = GetMutableDefault<UBlueprintEditorSettings>()->GetFunctionPermissions();
-		if (FunctionPermissions.HasFiltering())
+		const bool bDoPropertyEditorPermission = UE::MVVM::Private::ShouldDoFieldEditorPermission(GeneratingFor, AuthoritativeClass, Function->GetOwnerClass());
+		if (bDoPropertyEditorPermission)
 		{
 			const UFunction* FunctionToTest = AuthoritativeClass->FindFunctionByName(Function->GetFName());
 			if (FunctionToTest == nullptr)
@@ -209,26 +216,30 @@ bool IsConversionFunctionAllowed(const TSet<FSoftClassPath>& AllowedClasses, con
 	{
 		TStringBuilder<512> FunctionClassPath;
 		CurrentClass->GetPathName(nullptr, FunctionClassPath);
-		TStringBuilder<512> AllowedClassPath;
-		for (const FSoftClassPath& SoftClass : AllowedClasses)
+		TStringBuilder<512> ToTestClassPath;
+
+
+		for (const FSoftClassPath& SoftClass : DeniedClasses)
 		{
-			SoftClass.ToString(AllowedClassPath);
-			if (AllowedClassPath.ToView() == FunctionClassPath.ToView())
+			SoftClass.ToString(ToTestClassPath);
+			if (ToTestClassPath.ToView() == FunctionClassPath.ToView())
 			{
-				return true;
+				return false;
 			}
-			AllowedClassPath.Reset();
+			ToTestClassPath.Reset();
 		}
-		for (const FSoftClassPath& SoftClass : DeniedClasses)
+
+		for (const FSoftClassPath& SoftClass : AllowedClasses)
 		{
-			SoftClass.ToString(AllowedClassPath);
-			if (AllowedClassPath.ToView() == FunctionClassPath.ToView())
+			SoftClass.ToString(ToTestClassPath);
+			if (ToTestClassPath.ToView() == FunctionClassPath.ToView())
 			{
-				return false;
+				return true;
 			}
-			AllowedClassPath.Reset();
+			ToTestClassPath.Reset();
 		}
 
+
 		CurrentClass = CurrentClass->GetSuperClass();
 	}
 	return false;
