MVVM: Widget property created by the view compilation are private. The user can create them public with the "As Variable" option.

diff --git a/Engine/Plugins/Runtime/ModelViewViewModel/Source/ModelViewViewModelBlueprint/Private/MVVMViewBlueprintCompiler.cpp b/Engine/Plugins/Runtime/ModelViewViewModel/Source/ModelViewViewModelBlueprint/Private/MVVMViewBlueprintCompiler.cpp
--- a/Engine/Plugins/Runtime/ModelViewViewModel/Source/ModelViewViewModelBlueprint/Private/MVVMViewBlueprintCompiler.cpp
+++ b/Engine/Plugins/Runtime/ModelViewViewModel/Source/ModelViewViewModelBlueprint/Private/MVVMViewBlueprintCompiler.cpp
@@ -921,7 +921,6 @@ void FMVVMViewBlueprintCompiler::CreateRequiredProperties(const FWidgetBlueprint
 					SourceVariable.CreationType = FCompilerUserWidgetProperty::ECreationType::CreateIfDoesntExist;
 					SourceVariable.bExposeOnSpawn = false;
 					SourceVariable.bPrivate = true;
-					SourceVariable.bPrivate = false;
 					SourceVariable.bReadOnly = true;
 					SourceVariable.bInstanced = false;
 					SourceVariable.bInstanceExposed = false;
