MVVM: Deprecated the ViewModelBlueprintGeneratedClass.

diff --git a/Engine/Plugins/Runtime/ModelViewViewModel/Source/ModelViewViewModel/Private/MVVMViewModelBase.cpp b/Engine/Plugins/Runtime/ModelViewViewModel/Source/ModelViewViewModel/Private/MVVMViewModelBase.cpp
--- a/Engine/Plugins/Runtime/ModelViewViewModel/Source/ModelViewViewModel/Private/MVVMViewModelBase.cpp
+++ b/Engine/Plugins/Runtime/ModelViewViewModel/Source/ModelViewViewModel/Private/MVVMViewModelBase.cpp
@@ -2,7 +2,7 @@
 
 #include "MVVMViewModelBase.h"
 #include "Blueprint/BlueprintExceptionInfo.h"
-#include "ViewModel/MVVMViewModelBlueprintGeneratedClass.h"
+#include "Engine/BlueprintGeneratedClass.h"
 
 #include UE_INLINE_GENERATED_CPP_BY_NAME(MVVMViewModelBase)
 
diff --git a/Engine/Plugins/Runtime/ModelViewViewModel/Source/ModelViewViewModel/Public/ViewModel/MVVMViewModelBlueprintGeneratedClass.h b/Engine/Plugins/Runtime/ModelViewViewModel/Source/ModelViewViewModel/Public/ViewModel/MVVMViewModelBlueprintGeneratedClass.h
--- a/Engine/Plugins/Runtime/ModelViewViewModel/Source/ModelViewViewModel/Public/ViewModel/MVVMViewModelBlueprintGeneratedClass.h
+++ b/Engine/Plugins/Runtime/ModelViewViewModel/Source/ModelViewViewModel/Public/ViewModel/MVVMViewModelBlueprintGeneratedClass.h
@@ -6,25 +6,9 @@
 
 #include "MVVMViewModelBlueprintGeneratedClass.generated.h"
 
-namespace UE::FieldNotification { struct FFieldId; }
-struct FFieldNotificationId;
-
-class UMVVMViewModelBase;
-
-namespace UE::MVVM
-{
-	class FViewModelBlueprintCompilerContext;
-}//namespace
-
-
-/** Will be deprecated in the next version. */
-UCLASS()
-class MODELVIEWVIEWMODEL_API UMVVMViewModelBlueprintGeneratedClass : public UBlueprintGeneratedClass
+/** */
+UCLASS(Deprecated, MinimalAPI)
+class UE_DEPRECATED(5.5, "Use the regular BlueprintGeneratedClass.") UDEPRECATED_MVVMViewModelBlueprintGeneratedClass : public UBlueprintGeneratedClass
 {
 	GENERATED_BODY()
 };
-
-#if UE_ENABLE_INCLUDE_ORDER_DEPRECATED_IN_5_2
-#include "CoreMinimal.h"
-#include "FieldNotification/FieldId.h"
-#endif
