MVVM: add an extension to panel widgets to define entry classes for them and allow binding the generated entries to an array of viewmodels

diff --git a/Engine/Plugins/Runtime/ModelViewViewModel/Source/ModelViewViewModel/Private/Extensions/MVVMViewPanelWidgetExtension.cpp b/Engine/Plugins/Runtime/ModelViewViewModel/Source/ModelViewViewModel/Private/Extensions/MVVMViewPanelWidgetExtension.cpp
new file mode 100644
--- /dev/null
+++ b/Engine/Plugins/Runtime/ModelViewViewModel/Source/ModelViewViewModel/Private/Extensions/MVVMViewPanelWidgetExtension.cpp
@@ -0,0 +1,210 @@
+// Copyright Epic Games, Inc. All Rights Reserved.
+
+#include "Extensions/MVVMViewPanelWidgetExtension.h"
+
+#include "Bindings/MVVMFieldPathHelper.h"
+#include "Components/PanelWidget.h"
+#include "MVVMMessageLog.h"
+#include "MVVMSubsystem.h"
+#include "Templates/ValueOrError.h"
+#include "Types/MVVMFieldContext.h"
+#include "View/MVVMView.h"
+#include "View/MVVMViewClass.h"
+
+#include UE_INLINE_GENERATED_CPP_BY_NAME(MVVMViewPanelWidgetExtension)
+
+#define LOCTEXT_NAMESPACE "MVVMViewPanelWidgetExtension"
+
+#if WITH_EDITOR
+void UMVVMViewPanelWidgetExtension::Initialize(UMVVMViewPanelWidgetExtension::FInitPanelWidgetExtensionArgs InArgs)
+{
+	WidgetName = InArgs.WidgetName;
+	WidgetPath = InArgs.WidgetPath;
+	EntryViewModelName = InArgs.EntryViewModelName;
+	EntryWidgetClass = InArgs.EntryWidgetClass;
+	SlotTemplate = InArgs.SlotTemplate;
+	PanelPropertyName = InArgs.PanelPropertyName;
+	EntryViewModelClass = InArgs.EntryViewModelClass;
+}
+#endif
+
+void UMVVMViewPanelWidgetExtension::OnViewConstructed(UUserWidget* UserWidget, UMVVMView* View)
+{
+	check(View->GetViewClass());
+	CachedOwningUserWidget = UserWidget;
+	// Set the extension object on the runtime user widget.
+	FObjectPropertyBase* FoundObjectProperty = FindFProperty<FObjectPropertyBase>(UserWidget->GetClass(), PanelPropertyName);
+	if (ensureAlwaysMsgf(FoundObjectProperty, TEXT("The compiler should have added the property")))
+	{
+		FoundObjectProperty->SetObjectPropertyValue_InContainer(UserWidget, this);
+	}
+
+	// Fetch and cache the panel widget
+	TValueOrError<UE::MVVM::FFieldContext, void> FieldPathResult = View->GetViewClass()->GetBindingLibrary().EvaluateFieldPath(UserWidget, WidgetPath);
+
+	if (FieldPathResult.HasValue())
+	{
+		TValueOrError<UObject*, void> ObjectResult = UE::MVVM::FieldPathHelper::EvaluateObjectProperty(FieldPathResult.GetValue());
+		if (ObjectResult.HasValue() && ObjectResult.GetValue() != nullptr)
+		{
+			if (UPanelWidget* PanelWidget = Cast<UPanelWidget>(ObjectResult.GetValue()))
+			{
+				CachedPanelWidget = PanelWidget;
+			}
+			else
+			{
+				UE::MVVM::FMessageLog Log(UserWidget);
+				Log.Error(FText::Format(LOCTEXT("BindToEntryGenerationFailWidgetNotPanel", "The object property '{0}' is not of type panel widget, but has an Viewmodel extension meant for panel widgets. The extension won't have any effects.")
+					, FText::FromName(ObjectResult.GetValue()->GetFName())
+				));
+			}
+		}
+		else
+		{
+			UE::MVVM::FMessageLog Log(UserWidget);
+			Log.Error(FText::Format(LOCTEXT("BindToEntryGenerationFailInvalidObjectPropertyWidget", "The property object for panel-type widget '{0}' is not found, so viewmodels won't be bound to its entries.")
+				, FText::FromName(WidgetName)
+			));
+		}
+	}
+	else
+	{
+		UE::MVVM::FMessageLog Log(UserWidget);
+		Log.Error(FText::Format(LOCTEXT("BindToEntryGenerationFailInvalidFieldPathWidget", "The field path for panel-type widget '{0}' is invalid, so viewmodels won't be bound to its entries.")
+			, FText::FromName(WidgetName)
+		));
+	}
+}
+
+void UMVVMViewPanelWidgetExtension::OnViewDestructed(UUserWidget* UserWidget, UMVVMView* View)
+{
+	FObjectPropertyBase* FoundObjectProperty = FindFProperty<FObjectPropertyBase>(UserWidget->GetClass(), PanelPropertyName);
+	if (ensureAlwaysMsgf(FoundObjectProperty, TEXT("The compiler should have added the property")))
+	{
+		FoundObjectProperty->SetObjectPropertyValue_InContainer(UserWidget, nullptr);
+	}
+}
+
+void UMVVMViewPanelWidgetExtension::BP_SetItems(const TArray<UObject*>& InItems)
+{
+	if (UPanelWidget* PanelWidgetPtr = CachedPanelWidget.Get())
+	{
+		// Store all the reusable slots in a temporary array so that we don't re-create them.
+		TArray<TTuple<UPanelSlot*, TScriptInterface<INotifyFieldValueChanged>>> PreviousSlots;
+		for (UPanelSlot* Slot : PanelWidgetPtr->GetSlots())
+		{
+			if (UUserWidget* Content = Cast<UUserWidget>(Slot->Content))
+			{
+				// The class of the content should strictly match the EntryWidgetClass.
+				if (Content->GetClass() == EntryWidgetClass.Get())
+				{
+					if (UMVVMView* View = UMVVMSubsystem::GetViewFromUserWidget(Content))
+					{
+						TScriptInterface<INotifyFieldValueChanged> Interface = View->GetViewModel(EntryViewModelName);
+						if (Interface.GetObject() && Interface.GetObject()->GetClass() == EntryViewModelClass)
+						{
+							PreviousSlots.Emplace(Slot, Interface);
+						}
+					}
+				}
+			}
+		}
+
+		UClass* SelectedVMClass = EntryViewModelClass.Get();
+		UUserWidget* OwningUserWidget = CachedOwningUserWidget.Get();
+
+		TArray<TTuple<UPanelSlot*, UWidget*>> NewSlots;
+		for (int32 ItemIndex = 0; ItemIndex < InItems.Num(); ++ItemIndex)
+		{
+			UObject* Item = InItems[ItemIndex];
+			const TTuple<UPanelSlot*, TScriptInterface<INotifyFieldValueChanged>>* FoundObject = PreviousSlots.FindByPredicate([Item](const auto& Other) 
+			{ 
+				return Other.Value.GetObject() == Item; 
+			});
+
+			if (Item && OwningUserWidget && Item->GetClass() != SelectedVMClass)
+			{
+				UE::MVVM::FMessageLog Log(OwningUserWidget);
+				Log.Warning(FText::Format(LOCTEXT("SetPanelWidgetItemsViewmodelTypeMismatch", "The item {0} passed as an entry of widget {1} is not a viewmodel of the selected type {2}.")
+					, FText::FromString(Item->GetName()), FText::FromString(PanelWidgetPtr->GetName()), FText::FromString(SelectedVMClass->GetName())
+				));
+			}
+
+			if (!FoundObject)
+			{
+				UUserWidget* EntryWidget = UUserWidget::CreateWidgetInstance(*PanelWidgetPtr, EntryWidgetClass, NAME_None);
+				ensure(EntryWidget);
+
+				SetViewModelOnEntryWidget(EntryWidget, Item, OwningUserWidget);
+				NewSlots.Add(TTuple<UPanelSlot*, UWidget*>(SlotTemplate, EntryWidget));
+			}
+			else
+			{
+				NewSlots.Add(TTuple<UPanelSlot*, UWidget*>(FoundObject->Key, FoundObject->Key->Content));
+			}
+		}
+
+		ReplaceAllSlots(NewSlots);
+	}
+}
+
+void UMVVMViewPanelWidgetExtension::SetViewModelOnEntryWidget(UUserWidget* EntryWidget, UObject* ViewModelObject, UUserWidget* OwningUserWidget)
+{
+	if (UMVVMView* View = UMVVMSubsystem::GetViewFromUserWidget(EntryWidget))
+	{
+		if (ViewModelObject->Implements<UNotifyFieldValueChanged>())
+		{
+			View->SetViewModel(EntryViewModelName, ViewModelObject);
+		}
+		else
+		{
+			if (OwningUserWidget)
+			{
+				UE::MVVM::FMessageLog Log(OwningUserWidget);
+				Log.Error(FText::Format(LOCTEXT("SetViewModelOnEntryWidgetFailNotViewModel", "Trying to set an object that is not a viewmodel on entries of panel-type widget '{0}'. If you do not wish to set viewmodels on the entries of this widget, please remove the corresonding Viewmodel extension from it.")
+					, FText::FromName(WidgetName)
+				));
+			}
+		}
+	}
+}
+
+void UMVVMViewPanelWidgetExtension::ReplaceAllSlots(TArrayView<TTuple<UPanelSlot*, UWidget*>> NewSlots)
+{
+	if (UPanelWidget* PanelWidgetPtr = CachedPanelWidget.Get())
+	{
+		const TArray<UPanelSlot*>& OldSlots = PanelWidgetPtr->GetSlots();
+		const int32 OldSlotsNum = OldSlots.Num();
+		const int32 NewSlotsNum = NewSlots.Num();
+		const int32 MinSlotNum = FMath::Min(OldSlotsNum, NewSlotsNum);
+
+		// as long as we're within the boundaries of both arrays, compare and replace elements
+		for (int32 SlotIndex = 0; SlotIndex < MinSlotNum; SlotIndex++)
+		{
+			if (OldSlots[SlotIndex] != NewSlots[SlotIndex].Key)
+			{
+				PanelWidgetPtr->RemoveChildAt(SlotIndex);
+				UPanelSlot* NewSlot = PanelWidgetPtr->InsertChildAt(SlotIndex, NewSlots[SlotIndex].Value, NewSlots[SlotIndex].Key);
+			}
+		}
+
+		// If we have more old slots than new ones, remove all the extra ones.
+		if (OldSlotsNum > NewSlotsNum)
+		{
+			for (int32 SlotIndex = OldSlotsNum - 1; SlotIndex >= NewSlotsNum; SlotIndex--)
+			{
+				PanelWidgetPtr->RemoveChildAt(SlotIndex);
+			}
+		}
+		// If we have more new slots than old ones, simply append the new ones.
+		else if (NewSlotsNum > OldSlotsNum)
+		{
+			for (int32 SlotIndex = OldSlotsNum; SlotIndex < NewSlotsNum; SlotIndex++)
+			{
+				UPanelSlot* NewSlot = PanelWidgetPtr->AddChild(NewSlots[SlotIndex].Value, NewSlots[SlotIndex].Key);
+			}
+		}
+	}
+}
+
+#undef LOCTEXT_NAMESPACE
\ No newline at end of file
diff --git a/Engine/Plugins/Runtime/ModelViewViewModel/Source/ModelViewViewModel/Private/View/MVVMView.cpp b/Engine/Plugins/Runtime/ModelViewViewModel/Source/ModelViewViewModel/Private/View/MVVMView.cpp
--- a/Engine/Plugins/Runtime/ModelViewViewModel/Source/ModelViewViewModel/Private/View/MVVMView.cpp
+++ b/Engine/Plugins/Runtime/ModelViewViewModel/Source/ModelViewViewModel/Private/View/MVVMView.cpp
@@ -75,6 +75,11 @@ void UMVVMView::Construct()
 	check(GeneratedViewClass);
 	check(bConstructed == false);
 
+	for (UMVVMViewClassExtension* Extension : GeneratedViewClass->GetViewClassExtensions())
+	{
+		Extension->OnViewConstructed(GetUserWidget(), this);
+	}
+
 	if (GeneratedViewClass->DoesInitializeSourcesOnConstruct())
 	{
 		InitializeSources();
@@ -104,6 +109,11 @@ void UMVVMView::Destruct()
 
 	UninitializeEvents();
 	UninitializeSources(); // and bindings
+
+	for (UMVVMViewClassExtension* Extension : GeneratedViewClass->GetViewClassExtensions())
+	{
+		Extension->OnViewDestructed(GetUserWidget(), this);
+	}
 }
 
 
diff --git a/Engine/Plugins/Runtime/ModelViewViewModel/Source/ModelViewViewModel/Public/Extensions/MVVMViewClassExtension.h b/Engine/Plugins/Runtime/ModelViewViewModel/Source/ModelViewViewModel/Public/Extensions/MVVMViewClassExtension.h
--- a/Engine/Plugins/Runtime/ModelViewViewModel/Source/ModelViewViewModel/Public/Extensions/MVVMViewClassExtension.h
+++ b/Engine/Plugins/Runtime/ModelViewViewModel/Source/ModelViewViewModel/Public/Extensions/MVVMViewClassExtension.h
@@ -21,10 +21,12 @@ class UMVVMViewClassExtension : public UObject
 
 public:
 	//~ Functions to be overriden in a user-defined UMVVMViewMyWidgetExtension class
+	virtual void OnViewConstructed(UUserWidget* UserWidget, UMVVMView* View) {};
 	virtual void OnSourcesInitialized(UUserWidget* UserWidget, UMVVMView* View) {};
 	virtual void OnBindingsInitialized(UUserWidget* UserWidget, UMVVMView* View) {};
 	virtual void OnEventsInitialized(UUserWidget* UserWidget, UMVVMView* View) {};
 	virtual void OnEventsUninitialized(UUserWidget* UserWidget, UMVVMView* View) {};
 	virtual void OnBindingsUninitialized(UUserWidget* UserWidget, UMVVMView* View) {};
 	virtual void OnSourcesUninitialized(UUserWidget* UserWidget, UMVVMView* View) {};
+	virtual void OnViewDestructed(UUserWidget* UserWidget, UMVVMView* View) {};
 };
\ No newline at end of file
diff --git a/Engine/Plugins/Runtime/ModelViewViewModel/Source/ModelViewViewModel/Public/Extensions/MVVMViewPanelWidgetExtension.h b/Engine/Plugins/Runtime/ModelViewViewModel/Source/ModelViewViewModel/Public/Extensions/MVVMViewPanelWidgetExtension.h
new file mode 100644
--- /dev/null
+++ b/Engine/Plugins/Runtime/ModelViewViewModel/Source/ModelViewViewModel/Public/Extensions/MVVMViewPanelWidgetExtension.h
@@ -0,0 +1,90 @@
+// Copyright Epic Games, Inc. All Rights Reserved.
+
+#pragma once
+
+#include "Bindings/MVVMCompiledBindingLibrary.h"
+#include "Blueprint/UserWidget.h"
+#include "MVVMViewClassExtension.h"
+#include "MVVMViewPanelWidgetExtension.generated.h"
+
+class UMVVMView;
+class UPanelWidget;
+class UUserWidget;
+
+UCLASS()
+class MODELVIEWVIEWMODEL_API UMVVMViewPanelWidgetExtension : public UMVVMViewClassExtension
+{
+	GENERATED_BODY()
+
+public:
+	//~ Begin UMVVMViewClassExtension overrides
+	virtual void OnViewConstructed(UUserWidget* UserWidget, UMVVMView* View) override;
+	virtual void OnViewDestructed(UUserWidget* UserWidget, UMVVMView* View) override;
+	//~ End UMVVMViewClassExtension overrides
+
+#if WITH_EDITOR
+	struct FInitPanelWidgetExtensionArgs
+	{
+		FInitPanelWidgetExtensionArgs(FName InWidgetName, FName InEntryViewModelName, const FMVVMVCompiledFieldPath& InWidgetPath, const TSubclassOf<UUserWidget>& InEntryWidgetClass, UPanelSlot* InSlotTemplate, FName InPanelPropertyName, UClass* InEntryViewModelClass)
+			: WidgetName(InWidgetName),
+			EntryViewModelName(InEntryViewModelName),
+			WidgetPath(InWidgetPath),
+			EntryWidgetClass(InEntryWidgetClass),
+			SlotTemplate(InSlotTemplate),
+			PanelPropertyName(InPanelPropertyName),
+			EntryViewModelClass(InEntryViewModelClass)
+		{}
+
+		FName WidgetName;
+		FName EntryViewModelName;
+		FMVVMVCompiledFieldPath WidgetPath;
+		TSubclassOf<UUserWidget> EntryWidgetClass;
+		TObjectPtr<UPanelSlot> SlotTemplate;
+		FName PanelPropertyName;
+		TObjectPtr<UClass> EntryViewModelClass;
+	};
+
+	void Initialize(FInitPanelWidgetExtensionArgs InArgs);
+
+#endif
+
+	FName GetEntryViewModelName() const
+	{ 
+		return EntryViewModelName; 
+	}
+	
+	UFUNCTION(BlueprintCallable, Category = PanelWidget, meta = (AllowPrivateAccess = true, DisplayName = "Set Items", ViewmodelBlueprintWidgetExtension = "EntryViewModel"))
+	virtual void BP_SetItems(const TArray<UObject*>& InItems);
+
+private:
+	void SetViewModelOnEntryWidget(UUserWidget* EntryWidget, UObject* ViewModelObject, UUserWidget* OwningUserWidget);
+	void ReplaceAllSlots(TArrayView<TTuple<UPanelSlot*, UWidget*>> NewSlots);
+
+private:
+	UPROPERTY(VisibleAnywhere, Category = "MVVM Extension")
+	FName WidgetName;
+
+	UPROPERTY(VisibleAnywhere, Category = "MVVM Extension")
+	FName EntryViewModelName;
+
+	UPROPERTY(VisibleAnywhere, Category = "MVVM Extension")
+	TSubclassOf<UUserWidget> EntryWidgetClass;
+
+	UPROPERTY(VisibleAnywhere, Category = "MVVM Extension")
+	TObjectPtr<UPanelSlot> SlotTemplate;
+
+	UPROPERTY()
+	FName PanelPropertyName;
+
+	UPROPERTY()
+	TObjectPtr<UClass> EntryViewModelClass = nullptr;
+
+	UPROPERTY()
+	FMVVMVCompiledFieldPath WidgetPath;
+
+	UPROPERTY(Transient)
+	TWeakObjectPtr<UPanelWidget> CachedPanelWidget;
+
+	UPROPERTY(Transient)
+	TWeakObjectPtr<UUserWidget> CachedOwningUserWidget;
+};
\ No newline at end of file
diff --git a/Engine/Plugins/Runtime/ModelViewViewModel/Source/ModelViewViewModelBlueprint/Private/Extensions/MVVMViewBlueprintPanelWidgetExtension.cpp b/Engine/Plugins/Runtime/ModelViewViewModel/Source/ModelViewViewModelBlueprint/Private/Extensions/MVVMViewBlueprintPanelWidgetExtension.cpp
new file mode 100644
--- /dev/null
+++ b/Engine/Plugins/Runtime/ModelViewViewModel/Source/ModelViewViewModelBlueprint/Private/Extensions/MVVMViewBlueprintPanelWidgetExtension.cpp
@@ -0,0 +1,236 @@
+// Copyright Epic Games, Inc. All Rights Reserved.
+
+#include "Extensions/MVVMViewBlueprintPanelWidgetExtension.h"
+
+#include "Bindings/MVVMCompiledBindingLibraryCompiler.h"
+#include "Blueprint/WidgetBlueprintGeneratedClass.h"
+#include "Components/PanelWidget.h"
+#include "Extensions/MVVMViewPanelWidgetExtension.h"
+#include "MVVMBlueprintView.h"
+#include "MVVMBlueprintViewCompilerInterface.h"
+#include "MVVMBlueprintViewModelContext.h"
+#include "MVVMPropertyPath.h"
+#include "MVVMWidgetBlueprintExtension_View.h"
+#include "Templates/ValueOrError.h"
+#include "UObject/UnrealType.h"
+#include "View/MVVMViewClass.h"
+
+#include UE_INLINE_GENERATED_CPP_BY_NAME(MVVMViewBlueprintPanelWidgetExtension)
+
+#define LOCTEXT_NAMESPACE "MVVMViewBlueprintPanelWidgetExtension"
+
+TArray<UE::MVVM::Compiler::FBlueprintViewUserWidgetProperty> UMVVMViewBlueprintPanelWidgetExtension::AddProperties()
+{
+	TArray<UE::MVVM::Compiler::FBlueprintViewUserWidgetProperty> PropertiesToAdd;
+	PanelPropertyName = FName(FText::Format(LOCTEXT("ExtensionVariableName", "{0}_{1}"), FText::FromName(WidgetName), FText::FromString(TEXT("Viewmodel_Extension"))).ToString());
+
+	UE::MVVM::Compiler::FBlueprintViewUserWidgetProperty Property;
+	Property.AuthoritativeClass = UMVVMViewPanelWidgetExtension::StaticClass();
+	Property.DisplayName = FText::FromName(PanelPropertyName);
+	Property.Name = PanelPropertyName;
+	Property.CategoryName = TEXT("PanelWidgetExtension");
+	Property.bReadOnly = true;
+
+	PropertiesToAdd.Add(Property);
+	return PropertiesToAdd;
+}
+
+void UMVVMViewBlueprintPanelWidgetExtension::Precompile(UE::MVVM::Compiler::IMVVMBlueprintViewPrecompile* Compiler, UWidgetBlueprintGeneratedClass* Class)
+{
+	check (Compiler);
+	auto VerifyViewmodelTypeMatch = [this, Compiler](const TSubclassOf<UUserWidget> InEntryWidgetClass, const FMVVMBlueprintPropertyPath& EntryViewModelPath, const FName& PropertyName)
+	{
+		bool bFoundSetterBinding = false;
+
+		// Find bindings that have write fields with metadata "ViewmodelBlueprintWidgetExtension", these are setter functions for assigning a viewmodel array.
+		// Use their corresponding read field to get the expected type of entry viewmodel, and compare it against the one that user selected in the details panel.
+		for (const UE::MVVM::Compiler::FCompilerBindingHandle& Binding : Compiler->GetAllBindings())
+		{
+			TArray<UE::MVVM::FMVVMConstFieldVariant> WriteFields = Compiler->GetBindingWriteFields(Binding);
+			for (int32 FieldIndex = 1; FieldIndex < WriteFields.Num(); FieldIndex++)
+			{
+				UE::MVVM::FMVVMConstFieldVariant& PathField = WriteFields[FieldIndex];
+				UE::MVVM::FMVVMConstFieldVariant& ParentPathField = WriteFields[FieldIndex - 1];
+				if (PathField.IsFunction() && ParentPathField.IsProperty() && PropertyName == ParentPathField.GetName())
+				{
+					const FString MetaData = PathField.GetFunction()->GetMetaData("ViewmodelBlueprintWidgetExtension");
+					if (MetaData.Equals("EntryViewmodel", ESearchCase::IgnoreCase))
+					{
+						if (const FProperty* SourceProperty = Compiler->GetBindingSourceProperty(Binding))
+						{
+							if (const FArrayProperty* ArrayProperty = CastField<FArrayProperty>(SourceProperty))
+							{
+								// This is the viewmodel that the user selected from the details panel
+								TArray<UE::MVVM::FMVVMConstFieldVariant> SelectedVMFields = EntryViewModelPath.GetFields(InEntryWidgetClass);
+								if (SelectedVMFields.Num() > 0)
+								{
+									UE::MVVM::FMVVMConstFieldVariant SelectedVMPropertyField = SelectedVMFields.Last();
+
+									if (SelectedVMPropertyField.IsProperty())
+									{
+										bFoundSetterBinding = true;
+										const FProperty* SelectedVMProperty = SelectedVMPropertyField.GetProperty();
+										if (!ArrayProperty->Inner->SameType(SelectedVMProperty))
+										{
+											const FText UserSelectedViewmodelType = FText::FromString(SelectedVMProperty->GetCPPType());
+											const FText ArraySetterFunctionViewmodelType = FText::FromString(ArrayProperty->Inner->GetCPPType());
+											Compiler->AddMessageForBinding(Binding, FText::Format(LOCTEXT("EntryViewModelTypeMismatch", "The entry viewmodel type {0} does not match the array of viewmodels of type {1}."), UserSelectedViewmodelType, ArraySetterFunctionViewmodelType), UE::MVVM::Compiler::EMessageType::Error);
+											Compiler->MarkPrecompileStepInvalid();
+											break;
+										}
+									}
+								}
+							}
+						}
+
+						// No need to look into the next fields of this write path, we have found the field with ViewmodelBlueprintWidgetExtension metadata.
+						// We will go to the next binding
+						break;
+					}
+				}
+			}
+		}
+
+		if (!bFoundSetterBinding && EntryViewModelId.IsValid())
+		{
+			Compiler->AddMessage(FText::Format(LOCTEXT("PreCompileMVVMWidgetExtensionNoSetterBindingFound", "No binding found from an array of viewmodels to {0} -> SetItems. Please add this binding or remove the Viewmodel extension on widget {0} from its details panel."), FText::FromName(WidgetName)), UE::MVVM::Compiler::EMessageType::Warning);
+		}
+	};
+	 
+	if (UWidget* const* FoundPanel = Compiler->GetWidgetNameToWidgetPointerMap().Find(WidgetName))
+	{
+		if (UPanelWidget* PanelWidget = Cast<UPanelWidget>(*FoundPanel))
+		{
+			if (EntryWidgetClass.Get())
+			{
+				UE::MVVM::Compiler::IMVVMBlueprintViewPrecompile::FObjectFieldPathArgs AddObjFieldPathArgs(Class, WidgetName.ToString(), UPanelWidget::StaticClass());
+				TValueOrError<UE::MVVM::FCompiledBindingLibraryCompiler::FFieldPathHandle, FText> FieldPathResult = Compiler->AddObjectFieldPath(AddObjFieldPathArgs);
+				if (FieldPathResult.HasError() || !FieldPathResult.GetValue().IsValid())
+				{
+					Compiler->AddMessage(FText::Format(LOCTEXT("CouldNotCreateSourceFieldPath", "Couldn't create the source field path '{0}'. {1}. Make sure '{0}' is marked as 'Is Variable'."), FText::FromName(WidgetName), FieldPathResult.GetError()), UE::MVVM::Compiler::EMessageType::Error);
+					Compiler->MarkPrecompileStepInvalid();
+					return;
+				}
+				WidgetPathHandle = FieldPathResult.StealValue();
+
+				if (const UUserWidget* EntryUserWidget = Cast<UUserWidget>(EntryWidgetClass->ClassDefaultObject))
+				{
+					if (const UMVVMBlueprintView* EntryBPView = GetEntryWidgetBlueprintView(EntryUserWidget))
+					{
+						if (const FMVVMBlueprintViewModelContext* ViewModelContext = EntryBPView->FindViewModel(EntryViewModelId))
+						{
+							const FName EntryViewModelName = ViewModelContext->GetViewModelName();
+
+							if (const UWidgetBlueprint* EntryBlueprint = Cast<UWidgetBlueprint>(EntryUserWidget->GetClass()->ClassGeneratedBy))
+							{
+								// find the property path of the entry viewmodel
+								FMVVMBlueprintPropertyPath SelectedViewModel;
+								SelectedViewModel.SetViewModelId(EntryViewModelId);
+								FProperty* ViewModelProperty = FindFProperty<FProperty>(*EntryWidgetClass, EntryViewModelName);
+								SelectedViewModel.AppendPropertyPath(EntryBlueprint, UE::MVVM::FMVVMConstFieldVariant(ViewModelProperty));
+
+								VerifyViewmodelTypeMatch(EntryWidgetClass, SelectedViewModel, PanelPropertyName);
+							}
+						}
+						else
+						{
+							// If the stored viewmodel ID invalid, it means the viewmodel was deleted so we clear the value.
+							// We need to manually clear this value here because the panel widget doesn't get notified when the entry widget viewmodels change.
+							EntryViewModelId = FGuid();
+							Compiler->AddMessage(FText::Format(LOCTEXT("PreCompileMVVMWidgetExtensionEntryVMDeleted", "No viewmodel selected for Entry widget {0}. Please select a viewmodel for it via the details panel or remove the Viewmodel extension on the containing widget {1}."), FText::FromName(EntryUserWidget->GetFName()), FText::FromName(WidgetName)), UE::MVVM::Compiler::EMessageType::Error);
+
+						}
+					}
+					else
+					{
+						// If no view is found in the entry widget, we clear the entry viewmodel value.
+						// We need to manually clear this value here because the panel widget doesn't get notified when the entry widget view is added/removed.
+						EntryViewModelId = FGuid();
+						Compiler->AddMessage(FText::Format(LOCTEXT("PreCompileMVVMWidgetExtensionEntryHasNoView", "Entry widget {0} doesn't have a View. Consider adding a binding to it or remove the MVVM extension on the containing widget {1}."), FText::FromName(EntryUserWidget->GetFName()), FText::FromName(WidgetName)), UE::MVVM::Compiler::EMessageType::Error);
+					}
+				}
+			}
+			else
+			{
+				Compiler->AddMessage(FText::Format(LOCTEXT("PreCompileMVVMWidgetExtensionWidgetNoEntryClass", "Widget {0} doesn't have an entry widget class. Consider assigning it in the details panel or remove the Viewmodel extension on the widget."), FText::FromName(WidgetName)), UE::MVVM::Compiler::EMessageType::Error);
+			}
+		}
+		else
+		{
+			Compiler->AddMessage(FText::Format(LOCTEXT("PreCompileMVVMWidgetExtensionWidgetNotPanelWidget", "Widget {0} is not a UPanelWidget but has a MVVMViewBlueprintPanelWidgetExtension."), FText::FromName(WidgetName)), UE::MVVM::Compiler::EMessageType::Error);
+		}
+	}
+	else
+	{
+		Compiler->AddMessage(FText::Format(LOCTEXT("PreCompileMVVMWidgetExtensionInvalidWidgetName", "Widget with name {0} doesn't exist in the widget blueprint but a viewmodel widget extension exists that is attached to it."), FText::FromName(WidgetName)),UE::MVVM::Compiler::EMessageType::Error);
+	}
+}
+
+void UMVVMViewBlueprintPanelWidgetExtension::Compile(UE::MVVM::Compiler::IMVVMBlueprintViewCompile* Compiler, UWidgetBlueprintGeneratedClass* Class, UMVVMViewClass* ViewExtension)
+{
+	check(Compiler);
+
+	if (UWidget* const* FoundPanel = Compiler->GetWidgetNameToWidgetPointerMap().Find(WidgetName))
+	{
+		if (UPanelWidget* PanelWidget = Cast<UPanelWidget>(*FoundPanel))
+		{
+			if (EntryWidgetClass.Get())
+			{
+				check(WidgetPathHandle.IsValid());
+
+				// Verify the widget property
+				const TValueOrError<FMVVMVCompiledFieldPath, void> CompiledFieldPath = Compiler->GetFieldPath(WidgetPathHandle);
+				if (CompiledFieldPath.HasError())
+				{
+					Compiler->AddMessage(FText::Format(LOCTEXT("CompiledFieldPathForWidgetNotGenerated", "The field path for widget {0} was not generated."), FText::FromName(WidgetName)), UE::MVVM::Compiler::EMessageType::Error);
+					Compiler->MarkCompileStepInvalid();
+					return;
+				}
+
+				// Check that we have a valid entry view and valid entry viewmodel
+				// Otherwise, no runtime extension will be created.
+				if (const UUserWidget* EntryUserWidget = Cast<UUserWidget>(EntryWidgetClass->ClassDefaultObject))
+				{
+					if (const UMVVMBlueprintView* EntryBPView = GetEntryWidgetBlueprintView(EntryUserWidget))
+					{
+						if (const FMVVMBlueprintViewModelContext* ViewModelContext = EntryBPView->FindViewModel(EntryViewModelId))
+						{
+							// Create the corresponding runtime extension
+							UMVVMViewClassExtension* NewExtensionObj = Compiler->CreateViewClassExtension(UMVVMViewPanelWidgetExtension::StaticClass());
+							UMVVMViewPanelWidgetExtension* NewExtension = CastChecked<UMVVMViewPanelWidgetExtension>(NewExtensionObj);
+
+							const FName EntryViewModelName = ViewModelContext->GetViewModelName();
+							NewExtension->Initialize(UMVVMViewPanelWidgetExtension::FInitPanelWidgetExtensionArgs(WidgetName, EntryViewModelName, CompiledFieldPath.GetValue(), EntryWidgetClass, SlotObj, PanelPropertyName, ViewModelContext->NotifyFieldValueClass));
+						}
+					}
+				}
+			}
+		}
+	}	
+}
+
+const UMVVMBlueprintView* UMVVMViewBlueprintPanelWidgetExtension::GetEntryWidgetBlueprintView(const UUserWidget* EntryUserWidget) const
+{
+	if (const UWidgetBlueprint* EntryBlueprint = Cast<UWidgetBlueprint>(EntryUserWidget->GetClass()->ClassGeneratedBy))
+	{
+		if (const UMVVMWidgetBlueprintExtension_View* EntryWidgetExtension = UMVVMWidgetBlueprintExtension_View::GetExtension<UMVVMWidgetBlueprintExtension_View>(EntryBlueprint))
+		{
+			if (const UMVVMBlueprintView* EntryWidgetView = EntryWidgetExtension->GetBlueprintView())
+			{
+				return EntryWidgetView;
+			}
+		}
+	}
+	return nullptr;
+}
+
+void UMVVMViewBlueprintPanelWidgetExtension::WidgetRenamed(FName OldName, FName NewName)
+{
+	if (WidgetName == OldName)
+	{
+		WidgetName = NewName;
+		Modify();
+	}
+}
+
+#undef LOCTEXT_NAMESPACE
\ No newline at end of file
diff --git a/Engine/Plugins/Runtime/ModelViewViewModel/Source/ModelViewViewModelBlueprint/Private/MVVMViewBlueprintCompiler.cpp b/Engine/Plugins/Runtime/ModelViewViewModel/Source/ModelViewViewModelBlueprint/Private/MVVMViewBlueprintCompiler.cpp
--- a/Engine/Plugins/Runtime/ModelViewViewModel/Source/ModelViewViewModelBlueprint/Private/MVVMViewBlueprintCompiler.cpp
+++ b/Engine/Plugins/Runtime/ModelViewViewModel/Source/ModelViewViewModelBlueprint/Private/MVVMViewBlueprintCompiler.cpp
@@ -1036,6 +1036,27 @@ void FMVVMViewBlueprintCompiler::CreateRequiredProperties(const FWidgetBlueprint
 			}
 		}
 	}
+
+	for (FMVVMExtensionItem& Extension : BlueprintView->GetOuterUMVVMWidgetBlueprintExtension_View()->BlueprintExtensions)
+	{
+		if (ensure(Extension.ExtensionObj))
+		{
+			TArray<Compiler::FBlueprintViewUserWidgetProperty> Properties = Extension.ExtensionObj->AddProperties();
+			for (const Compiler::FBlueprintViewUserWidgetProperty& Property : Properties)
+			{
+				FCompilerUserWidgetProperty& CompilerUserWidgetProperty = NeededUserWidgetProperties.AddDefaulted_GetRef();
+
+				CompilerUserWidgetProperty.AuthoritativeClass = Property.AuthoritativeClass;
+				CompilerUserWidgetProperty.Name = Property.Name;
+				CompilerUserWidgetProperty.DisplayName = Property.DisplayName;
+				CompilerUserWidgetProperty.bExposeOnSpawn = Property.bExposeOnSpawn;
+				CompilerUserWidgetProperty.bPrivate = Property.bPrivate;
+				CompilerUserWidgetProperty.bReadOnly = Property.bReadOnly;
+				CompilerUserWidgetProperty.CategoryName = Property.CategoryName;
+				CompilerUserWidgetProperty.CreationType = FCompilerUserWidgetProperty::ECreationType::CreateOnlyIfDoesntExist;
+			}
+		}
+	}
 }
 
 
diff --git a/Engine/Plugins/Runtime/ModelViewViewModel/Source/ModelViewViewModelBlueprint/Private/MVVMViewBlueprintCompiler.h b/Engine/Plugins/Runtime/ModelViewViewModel/Source/ModelViewViewModelBlueprint/Private/MVVMViewBlueprintCompiler.h
--- a/Engine/Plugins/Runtime/ModelViewViewModel/Source/ModelViewViewModelBlueprint/Private/MVVMViewBlueprintCompiler.h
+++ b/Engine/Plugins/Runtime/ModelViewViewModel/Source/ModelViewViewModelBlueprint/Private/MVVMViewBlueprintCompiler.h
@@ -138,7 +138,7 @@ private:
 	 * They can be a Widget, a viewmodel, or any object owned by the UserWidget.
 	 * They can be source or destination.
 	 */
-	struct FCompilerUserWidgetProperty
+	struct FCompilerUserWidgetProperty : Compiler::FBlueprintViewUserWidgetProperty
 	{
 		enum class ECreationType
 		{
@@ -146,19 +146,10 @@ private:
 			CreateIfDoesntExist,
 			CreateOnlyIfDoesntExist,
 		};
-		UClass* AuthoritativeClass = nullptr;
-		FName Name;
-		FText DisplayName;
-		FString CategoryName;
 		FString BlueprintSetter;
 		ECreationType CreationType = ECreationType::None;
-		bool bExposeOnSpawn = false;
-		bool bPrivate = false;
-		bool bReadOnly = false;
 		bool bInstanced = false;
 		bool bInstanceExposed = false;
-
-		const FProperty* Property = nullptr;
 	};
 	TArray<FCompilerUserWidgetProperty> NeededUserWidgetProperties;
 
diff --git a/Engine/Plugins/Runtime/ModelViewViewModel/Source/ModelViewViewModelBlueprint/Public/Extensions/MVVMBlueprintViewExtension.h b/Engine/Plugins/Runtime/ModelViewViewModel/Source/ModelViewViewModelBlueprint/Public/Extensions/MVVMBlueprintViewExtension.h
--- a/Engine/Plugins/Runtime/ModelViewViewModel/Source/ModelViewViewModelBlueprint/Public/Extensions/MVVMBlueprintViewExtension.h
+++ b/Engine/Plugins/Runtime/ModelViewViewModel/Source/ModelViewViewModelBlueprint/Public/Extensions/MVVMBlueprintViewExtension.h
@@ -3,19 +3,13 @@
 #pragma once
 
 #include "CoreMinimal.h"
+#include "MVVMBlueprintViewCompilerInterface.h"
 
 #include "MVVMBlueprintViewExtension.generated.h"
 
-
 class UMVVMViewClass;
 class UWidgetBlueprintGeneratedClass;
 
-namespace UE::MVVM::Compiler
-{
-	class IMVVMBlueprintViewPrecompile;
-	class IMVVMBlueprintViewCompile;
-}
-
 /**
  * An extension class to define MVVM-related properties and behaviour. When WBP compiled, this information is copied 
  * into UMVVMViewClassExtension. This class provides a hook into the MVVM compiler and partially exposes it for injecting
@@ -28,7 +22,12 @@ class UMVVMBlueprintViewExtension : public UObject
 
 public:
 	//~ Functions to be overriden in a user-defined UMVVMViewBlueprintMyWidgetExtension class
+	virtual TArray<UE::MVVM::Compiler::FBlueprintViewUserWidgetProperty> AddProperties()
+	{
+		return TArray<UE::MVVM::Compiler::FBlueprintViewUserWidgetProperty>();
+	};
 	virtual void Precompile(UE::MVVM::Compiler::IMVVMBlueprintViewPrecompile* Compiler, UWidgetBlueprintGeneratedClass* Class) {}
 	virtual void Compile(UE::MVVM::Compiler::IMVVMBlueprintViewCompile* Compiler, UWidgetBlueprintGeneratedClass* Class, UMVVMViewClass* ViewExtension) {}
 	virtual void WidgetRenamed(FName OldName, FName NewName) {}
+
 };
\ No newline at end of file
diff --git a/Engine/Plugins/Runtime/ModelViewViewModel/Source/ModelViewViewModelBlueprint/Public/Extensions/MVVMViewBlueprintPanelWidgetExtension.h b/Engine/Plugins/Runtime/ModelViewViewModel/Source/ModelViewViewModelBlueprint/Public/Extensions/MVVMViewBlueprintPanelWidgetExtension.h
new file mode 100644
--- /dev/null
+++ b/Engine/Plugins/Runtime/ModelViewViewModel/Source/ModelViewViewModelBlueprint/Public/Extensions/MVVMViewBlueprintPanelWidgetExtension.h
@@ -0,0 +1,72 @@
+// Copyright Epic Games, Inc. All Rights Reserved.
+
+#pragma once
+
+#include "Blueprint/UserWidget.h"
+#include "Bindings/MVVMCompiledBindingLibraryCompiler.h"
+#include "MVVMBlueprintViewExtension.h"
+#include "Components/PanelSlot.h"
+#include "MVVMViewBlueprintPanelWidgetExtension.generated.h"
+
+class UMVVMBlueprintView;
+class UMVVMViewClass;
+class UMVVMViewPanelWidgetExtension;
+class UUserWidget;
+class UWidgetBlueprintGeneratedClass;
+
+namespace UE::MVVM
+{
+	class FMVVMPanelWidgetExtensionCustomizationExtender;
+}
+
+namespace UE::MVVM::Compiler
+{
+	class IMVVMBlueprintViewPrecompile;
+	class IMVVMBlueprintViewCompile;
+	struct FBlueprintViewUserWidgetProperty;
+}
+
+UCLASS()
+class MODELVIEWVIEWMODELBLUEPRINT_API UMVVMViewBlueprintPanelWidgetExtension : public UMVVMBlueprintViewExtension
+{
+	GENERATED_BODY()
+
+public:
+	//~ Begin UMVVMBlueprintViewExtension overrides
+	virtual TArray<UE::MVVM::Compiler::FBlueprintViewUserWidgetProperty> AddProperties() override;
+	virtual void Precompile(UE::MVVM::Compiler::IMVVMBlueprintViewPrecompile* Compiler, UWidgetBlueprintGeneratedClass* Class) override;
+	virtual void Compile(UE::MVVM::Compiler::IMVVMBlueprintViewCompile* Compiler, UWidgetBlueprintGeneratedClass* Class, UMVVMViewClass* ViewExtension) override;
+	virtual void WidgetRenamed(FName OldName, FName NewName) override;
+	//~ End UMVVMBlueprintViewExtension overrides
+
+	FGuid GetEntryViewModelId() const
+	{
+		return EntryViewModelId;
+	}
+
+private:
+	const UMVVMBlueprintView* GetEntryWidgetBlueprintView(const UUserWidget* EntryUserWidget) const;
+
+private:
+	UPROPERTY()
+	FName WidgetName;
+
+	UPROPERTY()
+	FGuid EntryViewModelId;
+
+	UPROPERTY(EditAnywhere, BlueprintReadWrite, Category = "EntryClass", meta = (AllowPrivateAccess = "true"))
+	TSubclassOf<UUserWidget> EntryWidgetClass;
+
+	UPROPERTY(EditAnywhere, BlueprintReadWrite, Instanced, Category="Slot", meta = (AllowPrivateAccess = "true"))
+	TObjectPtr<UPanelSlot> SlotObj;
+
+	UPROPERTY()
+	FName PanelPropertyName;
+
+	UPROPERTY()
+	TObjectPtr<UMVVMViewPanelWidgetExtension> ExtensionObj;
+
+	UE::MVVM::FCompiledBindingLibraryCompiler::FFieldPathHandle WidgetPathHandle;
+
+	friend UE::MVVM::FMVVMPanelWidgetExtensionCustomizationExtender;
+};
diff --git a/Engine/Plugins/Runtime/ModelViewViewModel/Source/ModelViewViewModelBlueprint/Public/MVVMBlueprintViewCompilerInterface.h b/Engine/Plugins/Runtime/ModelViewViewModel/Source/ModelViewViewModelBlueprint/Public/MVVMBlueprintViewCompilerInterface.h
--- a/Engine/Plugins/Runtime/ModelViewViewModel/Source/ModelViewViewModelBlueprint/Public/MVVMBlueprintViewCompilerInterface.h
+++ b/Engine/Plugins/Runtime/ModelViewViewModel/Source/ModelViewViewModelBlueprint/Public/MVVMBlueprintViewCompilerInterface.h
@@ -61,6 +61,19 @@ enum class EMessageType
 	Error = 2
 };
 
+struct FBlueprintViewUserWidgetProperty
+{
+	UClass* AuthoritativeClass = nullptr;
+	FName Name;
+	FText DisplayName;
+	FString CategoryName;
+	bool bExposeOnSpawn = false;
+	bool bPrivate = false;
+	bool bReadOnly = false;
+
+	const FProperty* Property = nullptr;
+};
+
 /**
  * Exposed interface of MVVMViewBlueprintCompiler to be used in view extensions (MVVMViewBlueprintViewExtension) at precompile step
  */
diff --git a/Engine/Plugins/Runtime/ModelViewViewModel/Source/ModelViewViewModelEditor/Private/Customizations/MVVMPanelWidgetExtensionCustomizationExtender.cpp b/Engine/Plugins/Runtime/ModelViewViewModel/Source/ModelViewViewModelEditor/Private/Customizations/MVVMPanelWidgetExtensionCustomizationExtender.cpp
new file mode 100644
--- /dev/null
+++ b/Engine/Plugins/Runtime/ModelViewViewModel/Source/ModelViewViewModelEditor/Private/Customizations/MVVMPanelWidgetExtensionCustomizationExtender.cpp
@@ -0,0 +1,378 @@
+// Copyright Epic Games, Inc. All Rights Reserved.
+
+#include "MVVMPanelWidgetExtensionCustomizationExtender.h"
+
+#include "AssetRegistry/AssetRegistryModule.h"
+#include "Components/PanelWidget.h"
+#include "DetailCategoryBuilder.h"
+#include "DetailLayoutBuilder.h"
+#include "Extensions/MVVMViewBlueprintPanelWidgetExtension.h"
+#include "Kismet2/BlueprintEditorUtils.h"
+#include "MVVMWidgetBlueprintExtension_View.h"
+#include "MVVMBlueprintView.h"
+#include "MVVMBlueprintViewModel.h"
+#include "PropertyCustomizationHelpers.h"
+#include "PropertyHandle.h"
+#include "WidgetBlueprint.h"
+#include "WidgetBlueprintEditor.h"
+#include "Widgets/Input/SButton.h"
+#include "Widgets/Input/SComboButton.h"
+#include "Widgets/Text/STextBlock.h"
+
+#define LOCTEXT_NAMESPACE "MVVMPanelWidgetExtensionCustomizationExtender"
+
+namespace UE::MVVM
+{
+
+TSharedPtr<FMVVMPanelWidgetExtensionCustomizationExtender> FMVVMPanelWidgetExtensionCustomizationExtender::MakeInstance()
+{
+	return MakeShared<FMVVMPanelWidgetExtensionCustomizationExtender>();
+}
+
+void FMVVMPanelWidgetExtensionCustomizationExtender::CustomizeDetails(IDetailLayoutBuilder& InDetailLayout, const TArrayView<UWidget*> InWidgets, const TSharedRef<FWidgetBlueprintEditor>& InWidgetBlueprintEditor)
+{
+	// multi-selection not supported for the data
+	if (InWidgets.Num() == 1)
+	{
+		if (UPanelWidget* Panel = Cast<UPanelWidget>(InWidgets[0]))
+		{
+			if (Panel->CanHaveMultipleChildren())
+			{
+				Widget = Panel;
+				WidgetBlueprintEditor = InWidgetBlueprintEditor;
+
+				// Only do a customization if we have a MVVM blueprint view class on this blueprint.
+				if (GetExtensionViewForSelectedWidgetBlueprint())
+				{
+					IDetailCategoryBuilder& MVVMCategory = InDetailLayout.EditCategory("Viewmodel");
+
+					bIsExtensionAdded = GetPanelWidgetExtension() != nullptr;
+
+					// Add a button that controls adding/removing the extension on the panel widget
+					MVVMCategory.AddCustomRow(FText::FromString(TEXT("Viewmodel")))
+					.NameContent()
+					[
+						SNew(STextBlock)
+						.Text(LOCTEXT("VMSupport", "Viewmodel Support"))
+						.Font(IDetailLayoutBuilder::GetDetailFont())
+					]
+					.ValueContent()
+					.HAlign(HAlign_Fill)
+					[
+						SNew(SHorizontalBox)
+						+ SHorizontalBox::Slot()
+						.AutoWidth()
+						[
+							SNew(SButton)
+							.OnClicked(this, &FMVVMPanelWidgetExtensionCustomizationExtender::ModifyExtension)
+							[
+								SNew(SHorizontalBox)
+								+ SHorizontalBox::Slot()
+								.HAlign(HAlign_Center)
+								.VAlign(VAlign_Center)
+								.AutoWidth()
+								[
+									SNew(SImage)
+									.Image(this, &FMVVMPanelWidgetExtensionCustomizationExtender::GetExtensionButtonIcon)
+								]
+								+ SHorizontalBox::Slot()
+								.Padding(FMargin(3, 0, 0, 0))
+								.VAlign(VAlign_Center)
+								.AutoWidth()
+								[
+									SNew(STextBlock)
+									.TextStyle(FAppStyle::Get(), "SmallButtonText")
+									.Text(this, &FMVVMPanelWidgetExtensionCustomizationExtender::GetExtensionButtonText)
+								]
+							]
+						]
+					];
+
+					if (UMVVMViewBlueprintPanelWidgetExtension* PanelExtension = GetPanelWidgetExtension())
+					{
+						IDetailPropertyRow* PanelExtensionPropertyRow = MVVMCategory.AddExternalObjects({ PanelExtension});
+						TSharedPtr<IPropertyHandle> PanelExtensionObjectHandle = PanelExtensionPropertyRow->GetPropertyHandle();
+						PanelExtensionPropertyRow->Visibility(EVisibility::Collapsed);
+
+						// "Entry Widget Class" property row
+						EntryClassHandle = PanelExtensionObjectHandle->GetChildHandle("EntryWidgetClass");
+						EntryClassHandle->SetOnPropertyValueChanged(FSimpleDelegate::CreateSP(this, &FMVVMPanelWidgetExtensionCustomizationExtender::HandleEntryClassChanged, false));
+						HandleEntryClassChanged(true);
+						IDetailPropertyRow& EntryClassRow = MVVMCategory.AddProperty(EntryClassHandle);
+
+						// "Entry Viewmodel" property row
+						MVVMCategory.AddCustomRow(FText::FromString(TEXT("Viewmodel")))
+							.NameContent()
+							[
+								SNew(STextBlock)
+								.Text(LOCTEXT("EntryVM", "Entry Viewmodel"))
+								.Font(IDetailLayoutBuilder::GetDetailFont())
+							]
+							.ValueContent()
+							.HAlign(HAlign_Fill)
+							[
+								SNew(SHorizontalBox)
+								+ SHorizontalBox::Slot()
+								.AutoWidth()
+								[
+									SNew(SComboButton)
+									.OnGetMenuContent(this, &FMVVMPanelWidgetExtensionCustomizationExtender::OnGetViewModelsMenuContent)
+									.ButtonContent()
+									[
+										SNew(STextBlock)
+										.Text(this, &FMVVMPanelWidgetExtensionCustomizationExtender::OnGetSelectedViewModel)
+									]
+								]
+								+ SHorizontalBox::Slot()
+								.AutoWidth()
+								[
+									PropertyCustomizationHelpers::MakeClearButton(
+										FSimpleDelegate::CreateSP(this, &FMVVMPanelWidgetExtensionCustomizationExtender::ClearEntryViewModel))
+								]
+							];
+
+						// "Slot template" property row
+						IDetailPropertyRow* SlotDetailRow = MVVMCategory.AddExternalObjects({ PanelExtension->SlotObj }, EPropertyLocation::Default,
+							FAddPropertyParams()
+							.CreateCategoryNodes(false)
+							.AllowChildren(true)
+							.HideRootObjectNode(false)
+						);
+
+						TSharedPtr<IPropertyHandle> SlotPropertyHandle = SlotDetailRow->GetPropertyHandle();
+
+						SlotDetailRow->CustomWidget(true)
+							.NameContent()
+							[
+								SNew(STextBlock)
+								.Text(LOCTEXT("SlotTemplate", "Slot Template"))
+								.Font(IDetailLayoutBuilder::GetDetailFont())
+							]
+						.ValueContent()
+							[
+								SlotPropertyHandle->CreatePropertyValueWidget()
+							];
+
+						// Because AddExternalObjects was used the property system will not add a reset to default widget by default 
+						SlotDetailRow->OverrideResetToDefault(
+							FResetToDefaultOverride::Create(
+								FIsResetToDefaultVisible::CreateLambda([SlotPropertyHandle](TSharedPtr<IPropertyHandle> Handle)
+									{
+										return SlotPropertyHandle->CanResetToDefault();
+									}),
+								FResetToDefaultHandler::CreateLambda([SlotPropertyHandle](TSharedPtr<IPropertyHandle> Handle)
+									{
+										return SlotPropertyHandle->ResetToDefault();
+									})
+								)
+						);
+					}
+				}
+			}
+		}
+	}
+}
+
+FReply FMVVMPanelWidgetExtensionCustomizationExtender::ModifyExtension()
+{
+	if (UMVVMViewBlueprintPanelWidgetExtension* PanelExtension = GetPanelWidgetExtension())
+	{
+		if (UPanelWidget* WidgetPtr = Widget.Get())
+		{
+			GetExtensionViewForSelectedWidgetBlueprint()->RemoveBlueprintWidgetExtension(PanelExtension, WidgetPtr->GetFName());
+			bIsExtensionAdded = false;
+		}
+	}
+	else
+	{
+		CreatePanelWidgetViewExtensionIfNotExisting();
+		bIsExtensionAdded = true;
+	}
+	return FReply::Handled();
+}
+
+void FMVVMPanelWidgetExtensionCustomizationExtender::CreatePanelWidgetViewExtensionIfNotExisting()
+{
+	if (UMVVMWidgetBlueprintExtension_View* Extension = GetExtensionViewForSelectedWidgetBlueprint())
+	{
+		if (UPanelWidget* WidgetPtr = Widget.Get())
+		{
+			if (Extension->GetBlueprintExtensionsForWidget(WidgetPtr->GetFName()).IsEmpty())
+			{
+				UMVVMBlueprintViewExtension* NewExtension = Extension->CreateBlueprintWidgetExtension(UMVVMViewBlueprintPanelWidgetExtension::StaticClass(), WidgetPtr->GetFName());
+				UMVVMViewBlueprintPanelWidgetExtension* NewPanelWidgetExtension = CastChecked<UMVVMViewBlueprintPanelWidgetExtension>(NewExtension);
+				NewPanelWidgetExtension->WidgetName = WidgetPtr->GetFName();
+
+				UPanelSlot* SlotObj = NewObject<UPanelSlot>(NewPanelWidgetExtension, WidgetPtr->GetSlotClass(), NAME_None);
+				NewPanelWidgetExtension->SlotObj = SlotObj;
+			}
+		}
+	}
+}
+
+UMVVMViewBlueprintPanelWidgetExtension* FMVVMPanelWidgetExtensionCustomizationExtender::GetPanelWidgetExtension() const
+{
+	if (UMVVMWidgetBlueprintExtension_View* ViewClass = GetExtensionViewForSelectedWidgetBlueprint())
+	{
+		if (UPanelWidget* WidgetPtr = Widget.Get())
+		{
+			for (UMVVMBlueprintViewExtension* Extension : ViewClass->GetBlueprintExtensionsForWidget(WidgetPtr->GetFName()))
+			{
+				if (UMVVMViewBlueprintPanelWidgetExtension* PanelWidgetExtension = Cast<UMVVMViewBlueprintPanelWidgetExtension>(Extension))
+				{
+					return PanelWidgetExtension;
+				}
+			}
+		}
+	}
+
+	return nullptr;
+}
+
+UMVVMWidgetBlueprintExtension_View* FMVVMPanelWidgetExtensionCustomizationExtender::GetExtensionViewForSelectedWidgetBlueprint() const
+{
+	if (const TSharedPtr<FWidgetBlueprintEditor> BPEditor = WidgetBlueprintEditor.Pin())
+	{
+		if (const UWidgetBlueprint* Blueprint = BPEditor->GetWidgetBlueprintObj())
+		{
+			return UMVVMWidgetBlueprintExtension_View::GetExtension<UMVVMWidgetBlueprintExtension_View>(Blueprint);
+		}
+	}
+
+	return nullptr;
+}
+
+void FMVVMPanelWidgetExtensionCustomizationExtender::ClearEntryViewModel()
+{
+	SetEntryViewModel(FGuid());
+}
+
+void FMVVMPanelWidgetExtensionCustomizationExtender::HandleEntryClassChanged(bool bIsInit)
+{
+	// Update the cached value of entry class
+	void* EntryClassPtr = nullptr;
+	TSubclassOf<UUserWidget>* EntryClassValue = nullptr;
+	bool bEntryClassChanged = false;
+
+	if (EntryClassHandle->IsValidHandle() && EntryClassHandle->GetValueData(EntryClassPtr) == FPropertyAccess::Success)
+	{
+		EntryClassValue = reinterpret_cast<TSubclassOf<UUserWidget>*>(EntryClassPtr);
+	}
+
+	if (!EntryClassValue || !EntryClass.Get() || EntryClass.Get() != *EntryClassValue)
+	{
+		bEntryClassChanged = true;
+	}
+	EntryClass = EntryClassValue ? *EntryClassValue : nullptr;
+
+	// Update other values that depend on the entry class (only if the cached value actually changed)
+	if (bEntryClassChanged && EntryClass)
+	{
+		if (UUserWidget* EntryCDO = Cast<UUserWidget>(EntryClass->ClassDefaultObject))
+		{
+			EntryWidgetBlueprint = Cast<UWidgetBlueprint>(EntryCDO->GetClass()->ClassGeneratedBy);
+		}
+
+		// Clear the saved entry viewmodel if we're not calling this from CustomizeDetails (not initializing the customizer)
+		if (!bIsInit)
+		{
+			SetEntryViewModel(FGuid(), false);
+		}
+	}
+}
+
+FText FMVVMPanelWidgetExtensionCustomizationExtender::OnGetSelectedViewModel() const
+{
+	if (const UPanelWidget* WidgetPtr = Widget.Get())
+	{
+		if (EntryClass)
+		{
+			if (UMVVMViewBlueprintPanelWidgetExtension* PanelWidgetExtension = GetPanelWidgetExtension())
+			{
+				if (const UUserWidget* EntryUserWidget = Cast<UUserWidget>(EntryClass->ClassDefaultObject))
+				{
+					if (const UWidgetBlueprint* EntryBlueprint = Cast<UWidgetBlueprint>(EntryUserWidget->GetClass()->ClassGeneratedBy))
+					{
+						if (const UMVVMWidgetBlueprintExtension_View* EntryWidgetExtension = UMVVMWidgetBlueprintExtension_View::GetExtension<UMVVMWidgetBlueprintExtension_View>(EntryBlueprint))
+						{
+							if (const UMVVMBlueprintView* EntryWidgetView = EntryWidgetExtension->GetBlueprintView())
+							{
+								if (const FMVVMBlueprintViewModelContext* ViewModelContext = EntryWidgetView->FindViewModel(PanelWidgetExtension->GetEntryViewModelId()))
+								{
+									return FText::FromName(ViewModelContext->GetViewModelName());
+								}
+							}
+						}
+					}
+				}
+			}
+		}
+	}
+
+	return LOCTEXT("NoViewmodel", "No Viewmodel");
+}
+
+FText FMVVMPanelWidgetExtensionCustomizationExtender::GetExtensionButtonText() const
+{
+	return bIsExtensionAdded ? LOCTEXT("RemoveVMExt", "Remove Viewmodel Extension") : LOCTEXT("AddVMExt", "Add Viewmodel Extension");
+}
+
+const FSlateBrush* FMVVMPanelWidgetExtensionCustomizationExtender::GetExtensionButtonIcon() const
+{
+	return bIsExtensionAdded ? FAppStyle::Get().GetBrush("Icons.X") : FAppStyle::Get().GetBrush("Icons.Plus");
+}
+
+TSharedRef<SWidget> FMVVMPanelWidgetExtensionCustomizationExtender::OnGetViewModelsMenuContent()
+{
+	FMenuBuilder MenuBuilder(true, NULL);
+
+	if (EntryClass)
+	{
+		// Find all viewmodels in the entry widget
+		if (const UWidgetBlueprint* EntryWidgetBlueprintPtr = EntryWidgetBlueprint.Get())
+		{
+			if (const UMVVMWidgetBlueprintExtension_View* EntryWidgetExtension = UMVVMWidgetBlueprintExtension_View::GetExtension<UMVVMWidgetBlueprintExtension_View>(EntryWidgetBlueprintPtr))
+			{
+				if (const UMVVMBlueprintView* EntryWidgetView = EntryWidgetExtension->GetBlueprintView())
+				{
+					const TArrayView<const FMVVMBlueprintViewModelContext> ViewModels = EntryWidgetView->GetViewModels();
+					for (const FMVVMBlueprintViewModelContext& EntryViewModel : ViewModels)
+					{
+						// Create the menu action for this entry viewmodel
+						FUIAction ItemAction(FExecuteAction::CreateSP(this, &FMVVMPanelWidgetExtensionCustomizationExtender::SetEntryViewModel, EntryViewModel.GetViewModelId(), true));
+						MenuBuilder.AddMenuEntry(FText::FromName(EntryViewModel.GetViewModelName()), TAttribute<FText>(), FSlateIcon(), ItemAction);
+					}
+				}
+			}
+		}
+	}
+	return MenuBuilder.MakeWidget();
+}
+
+void FMVVMPanelWidgetExtensionCustomizationExtender::SetEntryViewModel(FGuid InEntryViewModelId, bool bMarkModified)
+{
+	if (UMVVMWidgetBlueprintExtension_View* Extension = GetExtensionViewForSelectedWidgetBlueprint())
+	{
+		if (const UPanelWidget* WidgetPtr = Widget.Get())
+		{
+			if (UMVVMViewBlueprintPanelWidgetExtension* PanelWidgetExtension = GetPanelWidgetExtension())
+			{
+				PanelWidgetExtension->Modify();
+				PanelWidgetExtension->EntryViewModelId = InEntryViewModelId;
+				if (bMarkModified)
+				{
+					if (const TSharedPtr<FWidgetBlueprintEditor> BPEditor = WidgetBlueprintEditor.Pin())
+					{
+						if (UWidgetBlueprint* Blueprint = BPEditor->GetWidgetBlueprintObj())
+						{
+							FBlueprintEditorUtils::MarkBlueprintAsStructurallyModified(Blueprint);
+						}
+					}
+				}
+			}
+		}
+	}
+}
+
+}
+#undef LOCTEXT_NAMESPACE
\ No newline at end of file
diff --git a/Engine/Plugins/Runtime/ModelViewViewModel/Source/ModelViewViewModelEditor/Private/Customizations/MVVMPanelWidgetExtensionCustomizationExtender.h b/Engine/Plugins/Runtime/ModelViewViewModel/Source/ModelViewViewModelEditor/Private/Customizations/MVVMPanelWidgetExtensionCustomizationExtender.h
new file mode 100644
--- /dev/null
+++ b/Engine/Plugins/Runtime/ModelViewViewModel/Source/ModelViewViewModelEditor/Private/Customizations/MVVMPanelWidgetExtensionCustomizationExtender.h
@@ -0,0 +1,76 @@
+// Copyright Epic Games, Inc. All Rights Reserved.
+
+#pragma once
+
+#include "Customizations/IBlueprintWidgetCustomizationExtender.h"
+#include "Blueprint/UserWidget.h"
+#include "Layout/Visibility.h"
+#include "MVVMPropertyPath.h"
+#include "UObject/WeakObjectPtr.h"
+
+class FReply;
+class IPropertyHandle;
+class SWidget;
+class UPanelWidget;
+class UMVVMViewBlueprintPanelWidgetExtension;
+class UMVVMWidgetBlueprintExtension_View;
+
+namespace  UE::MVVM
+{
+
+class FMVVMPanelWidgetExtensionCustomizationExtender : public IBlueprintWidgetCustomizationExtender
+{
+public:
+	static TSharedPtr<FMVVMPanelWidgetExtensionCustomizationExtender> MakeInstance();
+	virtual void CustomizeDetails(IDetailLayoutBuilder& InDetailLayout, const TArrayView<UWidget*> InWidgets, const TSharedRef<FWidgetBlueprintEditor>& InWidgetBlueprintEditor);
+
+private:
+	/** Set the entry viewmodel on the MVVMViewBlueprintPanelWidgetExtension for this widget. */
+	void SetEntryViewModel(FGuid InEntryViewModelId, bool bMarkModified = true);
+
+	/** Get a list of all viewmodels in the entry class. */
+	TSharedRef<SWidget> OnGetViewModelsMenuContent();
+
+	/** Get the name of the currently-selected entry viewmodel from the extension. */
+	FText OnGetSelectedViewModel() const;
+
+	/** Clear the entry viewmodel on the MVVMViewBlueprintPanelWidgetExtension for this widget. */
+	void ClearEntryViewModel();
+
+	/** Update the cached variables when the entry class property changes. */
+	void HandleEntryClassChanged(bool bIsInit);
+
+	/** Update the cached variables when the number of grid columns changes. */
+	void HandleNumGridColumns();
+
+	/** Create a new MVVMViewBlueprintPanelWidgetExtension for this widget in the blueprint view class. */
+	void CreatePanelWidgetViewExtensionIfNotExisting();
+
+	/** Get the MVVMViewBlueprintPanelWidgetExtension for this widget in the blueprint view class. */
+	UMVVMViewBlueprintPanelWidgetExtension* GetPanelWidgetExtension() const;
+
+	/** Get the MVVM blueprint view class of this widget blueprint. */
+	UMVVMWidgetBlueprintExtension_View* GetExtensionViewForSelectedWidgetBlueprint() const;
+
+	/** Add/Remove the MVVMViewBlueprintPanelWidgetExtension for this widget on button click. */
+	FReply ModifyExtension();
+
+	/** Get + or X icon for the MVVM extension button. */
+	const FSlateBrush* GetExtensionButtonIcon() const;
+
+	/** Get display text for the MVVM extension button. */
+	FText GetExtensionButtonText() const;
+
+private:
+	/** The selected panel widget in the details panel. */
+	TWeakObjectPtr<UPanelWidget> Widget;
+	TWeakPtr<FWidgetBlueprintEditor> WidgetBlueprintEditor; 
+
+	TSubclassOf<UUserWidget> EntryClass;
+	TSharedPtr<IPropertyHandle> EntryClassHandle;
+	TWeakObjectPtr<UWidgetBlueprint> EntryWidgetBlueprint;
+
+	/** Keep track of whether we have a MVVMViewBlueprintPanelWidgetExtension for this widget. */
+	bool bIsExtensionAdded = false;
+};
+}
\ No newline at end of file
diff --git a/Engine/Plugins/Runtime/ModelViewViewModel/Source/ModelViewViewModelEditor/Private/ModelViewViewModelEditorModule.cpp b/Engine/Plugins/Runtime/ModelViewViewModel/Source/ModelViewViewModelEditor/Private/ModelViewViewModelEditorModule.cpp
--- a/Engine/Plugins/Runtime/ModelViewViewModel/Source/ModelViewViewModelEditor/Private/ModelViewViewModelEditorModule.cpp
+++ b/Engine/Plugins/Runtime/ModelViewViewModel/Source/ModelViewViewModelEditor/Private/ModelViewViewModelEditorModule.cpp
@@ -7,6 +7,7 @@
 #include "BlueprintModes/WidgetBlueprintApplicationModes.h"
 #include "Customizations/MVVMBlueprintViewModelContextCustomization.h"
 #include "Customizations/MVVMListViewBaseExtensionCustomizationExtender.h"
+#include "Customizations/MVVMPanelWidgetExtensionCustomizationExtender.h"
 #include "Customizations/MVVMPropertyBindingExtension.h"
 #include "Extensions/MVVMBlueprintViewExtension.h"
 #include "Framework/Application/SlateApplication.h"
@@ -49,6 +50,9 @@ void FModelViewViewModelEditorModule::StartupModule()
 	ListViewBaseCustomizationExtender = UE::MVVM::FMVVMListViewBaseExtensionCustomizationExtender::MakeInstance();
 	UMGEditorModule.AddWidgetCustomizationExtender(ListViewBaseCustomizationExtender.ToSharedRef());
 
+	PanelWidgetCustomizationExtender = UE::MVVM::FMVVMPanelWidgetExtensionCustomizationExtender::MakeInstance();
+	UMGEditorModule.AddWidgetCustomizationExtender(PanelWidgetCustomizationExtender.ToSharedRef());
+
 	UMGEditorModule.RegisterInstancedCustomPropertyTypeLayout(
 		FMVVMBlueprintViewModelContext::StaticStruct()->GetStructPathName()
 		, IUMGEditorModule::FOnGetInstancePropertyTypeCustomizationInstance::CreateStatic(UE::MVVM::FBlueprintViewModelContextDetailCustomization::MakeInstance)
diff --git a/Engine/Plugins/Runtime/ModelViewViewModel/Source/ModelViewViewModelEditor/Private/ModelViewViewModelEditorModule.h b/Engine/Plugins/Runtime/ModelViewViewModel/Source/ModelViewViewModelEditor/Private/ModelViewViewModelEditorModule.h
--- a/Engine/Plugins/Runtime/ModelViewViewModel/Source/ModelViewViewModelEditor/Private/ModelViewViewModelEditorModule.h
+++ b/Engine/Plugins/Runtime/ModelViewViewModel/Source/ModelViewViewModelEditor/Private/ModelViewViewModelEditorModule.h
@@ -9,6 +9,7 @@ namespace UE::MVVM
 {
 	class FAssetTypeActions_ViewModelBlueprint;
 	class FMVVMListViewBaseExtensionCustomizationExtender;
+	class FMVVMPanelWidgetExtensionCustomizationExtender;
 	class FMVVMPropertyBindingExtension;
 }
 
@@ -46,4 +47,5 @@ private:
 	TSharedPtr<UE::MVVM::FMVVMPropertyBindingExtension> PropertyBindingExtension;
 	TSharedPtr<UE::MVVM::FAssetTypeActions_ViewModelBlueprint> ViewModelBlueprintActions;
 	TSharedPtr<UE::MVVM::FMVVMListViewBaseExtensionCustomizationExtender> ListViewBaseCustomizationExtender;
+	TSharedPtr<UE::MVVM::FMVVMPanelWidgetExtensionCustomizationExtender> PanelWidgetCustomizationExtender;
 };
diff --git a/Engine/Source/Runtime/UMG/Private/Components/PanelWidget.cpp b/Engine/Source/Runtime/UMG/Private/Components/PanelWidget.cpp
--- a/Engine/Source/Runtime/UMG/Private/Components/PanelWidget.cpp
+++ b/Engine/Source/Runtime/UMG/Private/Components/PanelWidget.cpp
@@ -126,7 +126,12 @@ bool UPanelWidget::RemoveChildAt(int32 Index)
 
 UPanelSlot* UPanelWidget::AddChild(UWidget* Content)
 {
-	if ( Content == nullptr )
+	return AddChild(Content, nullptr);
+}
+
+UPanelSlot* UPanelWidget::AddChild(UWidget* Content, UPanelSlot* SlotTemplate)
+{
+	if (Content == nullptr)
 	{
 		return nullptr;
 	}
@@ -144,7 +149,24 @@ UPanelSlot* UPanelWidget::AddChild(UWidget* Content)
 		NewObjectFlags |= RF_Transient;
 	}
 
-	UPanelSlot* PanelSlot = NewObject<UPanelSlot>(this, GetSlotClass(), NAME_None, NewObjectFlags);
+	UPanelSlot* PanelSlot; 
+	if (SlotTemplate && GetSlotClass() == SlotTemplate->GetClass())
+	{
+		UWidget* TempContent = SlotTemplate->Content;
+		UPanelWidget* TempParent = SlotTemplate->Parent;
+		SlotTemplate->Content = nullptr;
+		SlotTemplate->Parent = nullptr;
+
+		PanelSlot = NewObject<UPanelSlot>(this, GetSlotClass(), NAME_None, NewObjectFlags, SlotTemplate);
+
+		SlotTemplate->Content = TempContent;
+		SlotTemplate->Parent = TempParent;
+	}
+	else
+	{
+		PanelSlot = NewObject<UPanelSlot>(this, GetSlotClass(), NAME_None, NewObjectFlags);
+	}
+
 	PanelSlot->Content = Content;
 	PanelSlot->Parent = this;
 
@@ -159,6 +181,34 @@ UPanelSlot* UPanelWidget::AddChild(UWidget* Content)
 	return PanelSlot;
 }
 
+UPanelSlot* UPanelWidget::InsertChildAt(int32 Index, UWidget* Content)
+{
+	return InsertChildAt(Index, Content, nullptr);
+}
+
+UPanelSlot* UPanelWidget::InsertChildAt(int32 Index, UWidget* Content, UPanelSlot* SlotTemplate)
+{
+	if (Content)
+	{
+		UPanelSlot* NewSlot = AddChild(Content, SlotTemplate);
+		ShiftChild(Index, Content);
+		return NewSlot;
+	}
+	return nullptr;
+}
+
+void UPanelWidget::ShiftChild(int32 Index, UWidget* Child)
+{
+	int32 CurrentIndex = GetChildIndex(Child);
+	if (CurrentIndex != INDEX_NONE)
+	{
+		Slots.RemoveAt(CurrentIndex);
+		Slots.Insert(Child->Slot, FMath::Clamp(Index, 0, Slots.Num()));
+
+		InvalidateLayoutAndVolatility();
+	}
+}
+
 #if WITH_EDITOR
 
 bool UPanelWidget::ReplaceChildAt(int32 Index, UWidget* Content)
@@ -191,27 +241,6 @@ bool UPanelWidget::ReplaceChild(UWidget* CurrentChild, UWidget* NewChild)
 	return false;
 }
 
-UPanelSlot* UPanelWidget::InsertChildAt(int32 Index, UWidget* Content)
-{
-	if (Content)
-	{
-		UPanelSlot* NewSlot = AddChild(Content);
-		ShiftChild(Index, Content);
-		return NewSlot;
-	}
-	return nullptr;
-}
-
-void UPanelWidget::ShiftChild(int32 Index, UWidget* Child)
-{
-	int32 CurrentIndex = GetChildIndex(Child);
-	if (CurrentIndex != INDEX_NONE)
-	{
-		Slots.RemoveAt(CurrentIndex);
-		Slots.Insert(Child->Slot, FMath::Clamp(Index, 0, Slots.Num()));
-	}
-}
-
 void UPanelWidget::SetDesignerFlags(EWidgetDesignFlags NewFlags)
 {
 	Super::SetDesignerFlags(NewFlags);
diff --git a/Engine/Source/Runtime/UMG/Public/Components/PanelWidget.h b/Engine/Source/Runtime/UMG/Public/Components/PanelWidget.h
--- a/Engine/Source/Runtime/UMG/Public/Components/PanelWidget.h
+++ b/Engine/Source/Runtime/UMG/Public/Components/PanelWidget.h
@@ -58,6 +58,31 @@ public:
 	UFUNCTION(BlueprintCallable, Category="Widget|Panel")
 	UMG_API UPanelSlot* AddChild(UWidget* Content);
 
+	/**
+	 * Adds a new child widget to the container. Creates the slot using the
+	 * template passed to it. Note that a new UPanelSlot object is created
+	 * and the SlotTemplate is only used as a template.
+	 */
+	UMG_API UPanelSlot* AddChild(UWidget* Content, UPanelSlot* SlotTemplate);
+
+	/**
+	 * Inserts a widget at a specific index.  This does not update the live slate version, it requires
+	 * a rebuild of the whole UI to see a change.
+	 */
+	UMG_API UPanelSlot* InsertChildAt(int32 Index, UWidget* Content);
+
+	/**
+	 * Inserts a widget at a specific index. Creates the slot using the
+	 * template passed to it. Note that a new UPanelSlot object is created
+	 * and the SlotTemplate is only used as a template.
+	 */
+	UMG_API UPanelSlot* InsertChildAt(int32 Index, UWidget* Content, UPanelSlot* SlotTemplate);
+
+	/**
+	 * Moves the child widget from its current index to the new index provided.
+	 */
+	UMG_API void ShiftChild(int32 Index, UWidget* Child);
+
 #if WITH_EDITOR
 
 	/**
@@ -77,17 +102,6 @@ public:
 	 */
 	UMG_API virtual bool ReplaceChild(UWidget* CurrentChild, UWidget* NewChild);
 
-	/**
-	 * Inserts a widget at a specific index.  This does not update the live slate version, it requires
-	 * a rebuild of the whole UI to see a change.
-	 */
-	UMG_API UPanelSlot* InsertChildAt(int32 Index, UWidget* Content);
-
-	/**
-	 * Moves the child widget from its current index to the new index provided.
-	 */
-	UMG_API void ShiftChild(int32 Index, UWidget* Child);
-
 	UMG_API virtual void SetDesignerFlags(EWidgetDesignFlags NewFlags) override;
 
 #endif
@@ -147,17 +161,17 @@ public:
 	UMG_API virtual void PostLoad() override;
 	// End UObject
 
+	virtual UClass* GetSlotClass() const
+	{
+		return UPanelSlot::StaticClass();
+	}
+
 protected:
 
 #if WITH_EDITOR
 	UMG_API virtual TSharedRef<SWidget> RebuildDesignWidget(TSharedRef<SWidget> Content) override;
 #endif
 
-	virtual UClass* GetSlotClass() const
-	{
-		return UPanelSlot::StaticClass();
-	}
-
 	virtual void OnSlotAdded(UPanelSlot* InSlot)
 	{
 
