MVVM: Extension improvements:
Support undo/redo
Support copy/paste
Add setting for supported widget types

diff --git a/Engine/Plugins/Runtime/ModelViewViewModel/Source/ModelViewViewModelBlueprint/Private/Extensions/MVVMViewBlueprintListViewBaseExtension.cpp b/Engine/Plugins/Runtime/ModelViewViewModel/Source/ModelViewViewModelBlueprint/Private/Extensions/MVVMViewBlueprintListViewBaseExtension.cpp
--- a/Engine/Plugins/Runtime/ModelViewViewModel/Source/ModelViewViewModelBlueprint/Private/Extensions/MVVMViewBlueprintListViewBaseExtension.cpp
+++ b/Engine/Plugins/Runtime/ModelViewViewModel/Source/ModelViewViewModelBlueprint/Private/Extensions/MVVMViewBlueprintListViewBaseExtension.cpp
@@ -207,13 +207,15 @@ const UMVVMBlueprintView* UMVVMViewBlueprintListViewBaseExtension::GetEntryWidge
 	return nullptr;
 }
 
-void UMVVMViewBlueprintListViewBaseExtension::WidgetRenamed(FName OldName, FName NewName)
+bool UMVVMViewBlueprintListViewBaseExtension::WidgetRenamed(FName OldName, FName NewName)
 {
 	if (WidgetName == OldName)
 	{
-		WidgetName = NewName;
 		Modify();
+		WidgetName = NewName;
+		return true;
 	}
+	return false;
 }
 
 #undef LOCTEXT_NAMESPACE
\ No newline at end of file
diff --git a/Engine/Plugins/Runtime/ModelViewViewModel/Source/ModelViewViewModelBlueprint/Private/Extensions/MVVMViewBlueprintPanelWidgetExtension.cpp b/Engine/Plugins/Runtime/ModelViewViewModel/Source/ModelViewViewModelBlueprint/Private/Extensions/MVVMViewBlueprintPanelWidgetExtension.cpp
--- a/Engine/Plugins/Runtime/ModelViewViewModel/Source/ModelViewViewModelBlueprint/Private/Extensions/MVVMViewBlueprintPanelWidgetExtension.cpp
+++ b/Engine/Plugins/Runtime/ModelViewViewModel/Source/ModelViewViewModelBlueprint/Private/Extensions/MVVMViewBlueprintPanelWidgetExtension.cpp
@@ -224,13 +224,15 @@ const UMVVMBlueprintView* UMVVMViewBlueprintPanelWidgetExtension::GetEntryWidget
 	return nullptr;
 }
 
-void UMVVMViewBlueprintPanelWidgetExtension::WidgetRenamed(FName OldName, FName NewName)
+bool UMVVMViewBlueprintPanelWidgetExtension::WidgetRenamed(FName OldName, FName NewName)
 {
 	if (WidgetName == OldName)
 	{
-		WidgetName = NewName;
 		Modify();
+		WidgetName = NewName;
+		return true;
 	}
+	return false;
 }
 
 #undef LOCTEXT_NAMESPACE
\ No newline at end of file
diff --git a/Engine/Plugins/Runtime/ModelViewViewModel/Source/ModelViewViewModelBlueprint/Private/MVVMDeveloperProjectSettings.cpp b/Engine/Plugins/Runtime/ModelViewViewModel/Source/ModelViewViewModelBlueprint/Private/MVVMDeveloperProjectSettings.cpp
--- a/Engine/Plugins/Runtime/ModelViewViewModel/Source/ModelViewViewModelBlueprint/Private/MVVMDeveloperProjectSettings.cpp
+++ b/Engine/Plugins/Runtime/ModelViewViewModel/Source/ModelViewViewModelBlueprint/Private/MVVMDeveloperProjectSettings.cpp
@@ -13,6 +13,13 @@
 #include "K2Node_FormatText.h"
 #include "Kismet/BlueprintFunctionLibrary.h"
 
+#include "Components/HorizontalBox.h"
+#include "Components/ListView.h"
+#include "Components/ScrollBox.h"
+#include "Components/StackBox.h"
+#include "Components/VerticalBox.h"
+#include "Components/WrapBox.h"
+
 #define LOCTEXT_NAMESPACE "MVVMDeveloperProjectSettings"
 
 UMVVMDeveloperProjectSettings::UMVVMDeveloperProjectSettings()
@@ -33,6 +40,13 @@ UMVVMDeveloperProjectSettings::UMVVMDeveloperProjectSettings()
 	FTopLevelAssetPath FormatText = FTopLevelAssetPath("/Script/BlueprintGraph", "K2Node_FormatText");
 	AllowedClassForConversionFunctions.Add(FSoftClassPath(BlueprintFunctionLibrary.ToString()));
 	AllowedClassForConversionFunctions.Add(FSoftClassPath(FormatText.ToString()));
+
+	SupportedListViewBaseClassesForExtension.Add(UListView::StaticClass());
+	SupportedPanelClassesForExtension.Add(UHorizontalBox::StaticClass());
+	SupportedPanelClassesForExtension.Add(UVerticalBox::StaticClass());
+	SupportedPanelClassesForExtension.Add(UScrollBox::StaticClass());
+	SupportedPanelClassesForExtension.Add(UStackBox::StaticClass());
+	SupportedPanelClassesForExtension.Add(UWrapBox::StaticClass());
 }
 
 FName UMVVMDeveloperProjectSettings::GetCategoryName() const
@@ -312,4 +326,42 @@ TArray<const UClass*> UMVVMDeveloperProjectSettings::GetDeniedConversionFunction
 	return Result;
 }
 
+bool UMVVMDeveloperProjectSettings::IsExtensionSupportedForPanelClass(TSubclassOf<UPanelWidget> ClassToSupport) const
+{
+	if (ClassToSupport.Get())
+	{
+		for (const TSoftClassPtr<UPanelWidget>& SoftClass : SupportedPanelClassesForExtension)
+		{
+			if (UClass* Class = SoftClass.Get())
+			{
+				if (ClassToSupport->IsChildOf(Class))
+				{
+					return true;
+				}
+			}
+		}
+	}
+
+	return false;
+}
+
+bool UMVVMDeveloperProjectSettings::IsExtensionSupportedForListViewBaseClass(TSubclassOf<UListViewBase> ClassToSupport) const
+{
+	if (ClassToSupport.Get())
+	{
+		for (const TSoftClassPtr<UListViewBase>& SoftClass : SupportedListViewBaseClassesForExtension)
+		{
+			if (UClass* Class = SoftClass.Get())
+			{
+				if (ClassToSupport->IsChildOf(Class))
+				{
+					return true;
+				}
+			}
+		}
+	}
+
+	return false;
+}
+
 #undef LOCTEXT_NAMESPACE
diff --git a/Engine/Plugins/Runtime/ModelViewViewModel/Source/ModelViewViewModelBlueprint/Private/MVVMWidgetBlueprintExtension_View.cpp b/Engine/Plugins/Runtime/ModelViewViewModel/Source/ModelViewViewModelBlueprint/Private/MVVMWidgetBlueprintExtension_View.cpp
--- a/Engine/Plugins/Runtime/ModelViewViewModel/Source/ModelViewViewModelBlueprint/Private/MVVMWidgetBlueprintExtension_View.cpp
+++ b/Engine/Plugins/Runtime/ModelViewViewModel/Source/ModelViewViewModelBlueprint/Private/MVVMWidgetBlueprintExtension_View.cpp
@@ -5,14 +5,15 @@
 #include "Blueprint/WidgetTree.h"
 #include "Extensions/MVVMBlueprintViewExtension.h"
 #include "FindInBlueprintManager.h"
+#include "Kismet2/BlueprintEditorUtils.h"
+#include "Kismet2/KismetEditorUtilities.h"
 #include "MVVMBlueprintInstancedViewModel.h"
 #include "MVVMBlueprintView.h"
 #include "MVVMBlueprintViewModel.h"
 #include "MVVMBlueprintViewConversionFunction.h"
 #include "MVVMViewBlueprintCompiler.h"
+#include "ScopedTransaction.h"
 #include "View/MVVMViewClass.h"
-#include "Kismet2/BlueprintEditorUtils.h"
-#include "Kismet2/KismetEditorUtilities.h"
 
 
 #include UE_INLINE_GENERATED_CPP_BY_NAME(MVVMWidgetBlueprintExtension_View)
@@ -56,6 +57,11 @@ void UMVVMWidgetBlueprintExtension_View::DestroyBlueprintViewInstance()
 void UMVVMWidgetBlueprintExtension_View::PostLoad()
 {
 	Super::PostLoad();
+
+	if (!HasAnyFlags(RF_Transactional))
+	{
+		SetFlags(RF_Transactional);
+	}
 }
 
 UMVVMBlueprintViewExtension* UMVVMWidgetBlueprintExtension_View::CreateBlueprintWidgetExtension(TSubclassOf<UMVVMBlueprintViewExtension> ExtensionClass, FName WidgetName)
@@ -64,8 +70,10 @@ UMVVMBlueprintViewExtension* UMVVMWidgetBlueprintExtension_View::CreateBlueprint
 	{
 		UObject* ExtensionObj = NewObject<UObject>(this, ExtensionClass.Get(), NAME_None, RF_Transactional);
 		UMVVMBlueprintViewExtension* NewExtension = CastChecked<UMVVMBlueprintViewExtension>(ExtensionObj);
-		NewExtension->Modify();
 
+		const FScopedTransaction Transaction(LOCTEXT("AddViewModelExtension", "Add viewmodel extension"));
+		NewExtension->Modify();
+		Modify();
 		FMVVMExtensionItem ExtensionToAdd;
 		ExtensionToAdd.WidgetName = WidgetName;
 		ExtensionToAdd.ExtensionObj = NewExtension;
@@ -82,6 +90,9 @@ void UMVVMWidgetBlueprintExtension_View::RemoveBlueprintWidgetExtension(UMVVMBlu
 	FMVVMExtensionItem Extension;
 	Extension.WidgetName = WidgetName;
 	Extension.ExtensionObj = ExtensionToRemove;
+
+	const FScopedTransaction Transaction(LOCTEXT("RemoveViewModelExtension", "Remove viewmodel extension"));
+	Modify();
 	BlueprintExtensions.Remove(Extension);
 	FBlueprintEditorUtils::MarkBlueprintAsStructurallyModified(GetWidgetBlueprint());
 }
@@ -128,6 +139,10 @@ void UMVVMWidgetBlueprintExtension_View::VerifyWidgetExtensions()
 				}
 			});
 
+			if (WidgetNamesToRemove.Num() > 0)
+			{
+				Modify();
+			}
 			for (int32 Index = BlueprintExtensions.Num() - 1; Index >= 0; Index--)
 			{
 				if (WidgetNamesToRemove.Contains(BlueprintExtensions[Index].WidgetName))
@@ -145,6 +160,7 @@ void UMVVMWidgetBlueprintExtension_View::RenameWidgetExtensions(FName OldName, F
 	{
 		if (Extension.WidgetName == OldName)
 		{
+			Modify();
 			Extension.WidgetName = NewName;
 			if (Extension.ExtensionObj)
 			{
diff --git a/Engine/Plugins/Runtime/ModelViewViewModel/Source/ModelViewViewModelBlueprint/Public/Extensions/MVVMBlueprintViewExtension.h b/Engine/Plugins/Runtime/ModelViewViewModel/Source/ModelViewViewModelBlueprint/Public/Extensions/MVVMBlueprintViewExtension.h
--- a/Engine/Plugins/Runtime/ModelViewViewModel/Source/ModelViewViewModelBlueprint/Public/Extensions/MVVMBlueprintViewExtension.h
+++ b/Engine/Plugins/Runtime/ModelViewViewModel/Source/ModelViewViewModelBlueprint/Public/Extensions/MVVMBlueprintViewExtension.h
@@ -28,6 +28,8 @@ public:
 	};
 	virtual void Precompile(UE::MVVM::Compiler::IMVVMBlueprintViewPrecompile* Compiler, UWidgetBlueprintGeneratedClass* Class) {}
 	virtual void Compile(UE::MVVM::Compiler::IMVVMBlueprintViewCompile* Compiler, UWidgetBlueprintGeneratedClass* Class, UMVVMViewClass* ViewExtension) {}
-	virtual void WidgetRenamed(FName OldName, FName NewName) {}
-
+	virtual bool WidgetRenamed(FName OldName, FName NewName) 
+	{ 
+		return false; 
+	}
 };
\ No newline at end of file
diff --git a/Engine/Plugins/Runtime/ModelViewViewModel/Source/ModelViewViewModelBlueprint/Public/Extensions/MVVMViewBlueprintListViewBaseExtension.h b/Engine/Plugins/Runtime/ModelViewViewModel/Source/ModelViewViewModelBlueprint/Public/Extensions/MVVMViewBlueprintListViewBaseExtension.h
--- a/Engine/Plugins/Runtime/ModelViewViewModel/Source/ModelViewViewModelBlueprint/Public/Extensions/MVVMViewBlueprintListViewBaseExtension.h
+++ b/Engine/Plugins/Runtime/ModelViewViewModel/Source/ModelViewViewModelBlueprint/Public/Extensions/MVVMViewBlueprintListViewBaseExtension.h
@@ -31,7 +31,7 @@ public:
 	//~ Begin UMVVMBlueprintViewExtension overrides
 	virtual void Precompile(UE::MVVM::Compiler::IMVVMBlueprintViewPrecompile* Compiler, UWidgetBlueprintGeneratedClass* Class) override;
 	virtual void Compile(UE::MVVM::Compiler::IMVVMBlueprintViewCompile* Compiler, UWidgetBlueprintGeneratedClass* Class, UMVVMViewClass* ViewExtension) override;
-	virtual void WidgetRenamed(FName OldName, FName NewName) override;
+	virtual bool WidgetRenamed(FName OldName, FName NewName) override;
 	//~ End UMVVMBlueprintViewExtension overrides
 
 	FGuid GetEntryViewModelId() const
diff --git a/Engine/Plugins/Runtime/ModelViewViewModel/Source/ModelViewViewModelBlueprint/Public/Extensions/MVVMViewBlueprintPanelWidgetExtension.h b/Engine/Plugins/Runtime/ModelViewViewModel/Source/ModelViewViewModelBlueprint/Public/Extensions/MVVMViewBlueprintPanelWidgetExtension.h
--- a/Engine/Plugins/Runtime/ModelViewViewModel/Source/ModelViewViewModelBlueprint/Public/Extensions/MVVMViewBlueprintPanelWidgetExtension.h
+++ b/Engine/Plugins/Runtime/ModelViewViewModel/Source/ModelViewViewModelBlueprint/Public/Extensions/MVVMViewBlueprintPanelWidgetExtension.h
@@ -36,7 +36,7 @@ public:
 	virtual TArray<UE::MVVM::Compiler::FBlueprintViewUserWidgetProperty> AddProperties() override;
 	virtual void Precompile(UE::MVVM::Compiler::IMVVMBlueprintViewPrecompile* Compiler, UWidgetBlueprintGeneratedClass* Class) override;
 	virtual void Compile(UE::MVVM::Compiler::IMVVMBlueprintViewCompile* Compiler, UWidgetBlueprintGeneratedClass* Class, UMVVMViewClass* ViewExtension) override;
-	virtual void WidgetRenamed(FName OldName, FName NewName) override;
+	virtual bool WidgetRenamed(FName OldName, FName NewName) override;
 	//~ End UMVVMBlueprintViewExtension overrides
 
 	FGuid GetEntryViewModelId() const
diff --git a/Engine/Plugins/Runtime/ModelViewViewModel/Source/ModelViewViewModelBlueprint/Public/MVVMDeveloperProjectSettings.h b/Engine/Plugins/Runtime/ModelViewViewModel/Source/ModelViewViewModelBlueprint/Public/MVVMDeveloperProjectSettings.h
--- a/Engine/Plugins/Runtime/ModelViewViewModel/Source/ModelViewViewModelBlueprint/Public/MVVMDeveloperProjectSettings.h
+++ b/Engine/Plugins/Runtime/ModelViewViewModel/Source/ModelViewViewModelBlueprint/Public/MVVMDeveloperProjectSettings.h
@@ -7,6 +7,8 @@
 #include "MVVMDeveloperProjectSettings.generated.h"
 
 class UK2Node;
+class UListViewBase;
+class UPanelWidget;
 enum class EMVVMBlueprintViewModelContextCreationType : uint8;
 enum class EMVVMExecutionMode : uint8;
 
@@ -94,6 +96,9 @@ public:
 	TArray<const UClass*> GetAllowedConversionFunctionClasses() const;
 	TArray<const UClass*> GetDeniedConversionFunctionClasses() const;
 
+	bool IsExtensionSupportedForPanelClass(TSubclassOf<UPanelWidget> ClassToSupport) const;
+	bool IsExtensionSupportedForListViewBaseClass(TSubclassOf<UListViewBase> ClassToSupport) const;
+
 private:
 	/** Permission list for filtering which properties are visible in UI. */
 	UPROPERTY(EditAnywhere, config, Category = "UX", meta=(AllowAbstract = true))
@@ -177,6 +182,14 @@ public:
 	/** Settings for filtering the list of available properties and functions on binding creation. */
 	UPROPERTY(EditAnywhere, config, Category = "UX")
 	FMVVMViewBindingFilterSettings FilterSettings;
+
+	/** Sub-classes of panel widget that are supported to have an extension for binding their entries to viewmodels. */
+	UPROPERTY(EditAnywhere, config, Category = "Widget Extension")
+	TSet<TSoftClassPtr<UPanelWidget>> SupportedPanelClassesForExtension;
+
+	/** Sub-classes of ListViewBase that are supported to have an extension for binding their entries to viewmodels. */
+	UPROPERTY(EditAnywhere, config, Category = "Widget Extension")
+	TSet<TSoftClassPtr<UListViewBase>> SupportedListViewBaseClassesForExtension;
 };
 
 #if UE_ENABLE_INCLUDE_ORDER_DEPRECATED_IN_5_2
diff --git a/Engine/Plugins/Runtime/ModelViewViewModel/Source/ModelViewViewModelBlueprint/Public/MVVMWidgetBlueprintExtension_View.h b/Engine/Plugins/Runtime/ModelViewViewModel/Source/ModelViewViewModelBlueprint/Public/MVVMWidgetBlueprintExtension_View.h
--- a/Engine/Plugins/Runtime/ModelViewViewModel/Source/ModelViewViewModelBlueprint/Public/MVVMWidgetBlueprintExtension_View.h
+++ b/Engine/Plugins/Runtime/ModelViewViewModel/Source/ModelViewViewModelBlueprint/Public/MVVMWidgetBlueprintExtension_View.h
@@ -12,6 +12,11 @@ class UMVVMBlueprintViewExtension;
 class FWidgetBlueprintCompilerContext;
 class UWidgetBlueprintGeneratedClass;
 
+namespace UE::MVVM
+{
+	class FClipboardExtension;
+}
+
 namespace UE::MVVM::Private
 {
 	struct FMVVMViewBlueprintCompiler;
@@ -114,6 +119,7 @@ private:
 	TArray<FMVVMExtensionItem> BlueprintExtensions;
 
 	friend UE::MVVM::Private::FMVVMViewBlueprintCompiler;
+	friend UE::MVVM::FClipboardExtension;
 };
 
 #if UE_ENABLE_INCLUDE_ORDER_DEPRECATED_IN_5_2
diff --git a/Engine/Plugins/Runtime/ModelViewViewModel/Source/ModelViewViewModelEditor/Private/Customizations/MVVMClipboardExtension.cpp b/Engine/Plugins/Runtime/ModelViewViewModel/Source/ModelViewViewModelEditor/Private/Customizations/MVVMClipboardExtension.cpp
new file mode 100644
--- /dev/null
+++ b/Engine/Plugins/Runtime/ModelViewViewModel/Source/ModelViewViewModelEditor/Private/Customizations/MVVMClipboardExtension.cpp
@@ -0,0 +1,104 @@
+// Copyright Epic Games, Inc. All Rights Reserved.
+
+#include "Customizations/MVVMClipboardExtension.h"
+
+#include "Exporters/Exporter.h"
+#include "Extensions/MVVMBlueprintViewExtension.h"
+#include "MVVMDeveloperProjectSettings.h"
+#include "MVVMWidgetBlueprintExtension_View.h"
+#include "WidgetBlueprintEditorUtils.h"
+
+#define LOCTEXT_NAMESPACE "MVVMClipboardExtension"
+
+namespace UE::MVVM
+{
+
+void FClipboardExtension::AppendToClipboard(const UWidget* Widget, const FExportArgs& ExportArgs)
+{
+	check(ExportArgs.Out);
+
+	if (UWidgetBlueprint* WidgetBlueprint = FWidgetBlueprintEditorUtils::GetWidgetBlueprintFromWidget(Widget))
+	{
+		if (UMVVMWidgetBlueprintExtension_View* MVVMExtensionViewPtr = UMVVMWidgetBlueprintExtension_View::GetExtension<UMVVMWidgetBlueprintExtension_View>(WidgetBlueprint))
+		{
+			for (UMVVMBlueprintViewExtension* Extension : MVVMExtensionViewPtr->GetBlueprintExtensionsForWidget(Widget->GetFName()))
+			{
+				UExporter::ExportToOutputDevice(ExportArgs.Context, Extension, ExportArgs.Exporter, *ExportArgs.Out, ExportArgs.FileType, ExportArgs.Indent, ExportArgs.PortFlags, ExportArgs.bSelectedOnly, ExportArgs.ExportRootScope);
+			}
+		}
+	}
+}
+
+bool FClipboardExtension::CanAppendToClipboard(const UWidget* Widget) const
+{
+	if (UWidgetBlueprint* WidgetBlueprint = FWidgetBlueprintEditorUtils::GetWidgetBlueprintFromWidget(Widget))
+	{
+		UMVVMWidgetBlueprintExtension_View* MVVMExtensionViewPtr = UMVVMWidgetBlueprintExtension_View::GetExtension<UMVVMWidgetBlueprintExtension_View>(WidgetBlueprint);
+		return MVVMExtensionViewPtr && MVVMExtensionViewPtr->GetBlueprintExtensionsForWidget(Widget->GetFName()).Num() > 0;
+	}
+	return false;
+}
+
+void FClipboardExtension::ProcessImportedText(const UWidgetBlueprint* WidgetBlueprint, const FString& TextToImport, UPackage*& TempPackage)
+{
+	Factory.NewExtensions.Empty();
+	if (UMVVMWidgetBlueprintExtension_View::GetExtension<UMVVMWidgetBlueprintExtension_View>(WidgetBlueprint))
+	{
+		Factory.ProcessBuffer(TempPackage, RF_Transactional, TextToImport);
+	}
+}
+
+bool FClipboardExtension::CanImportFromClipboard(const UWidget* Widget) const
+{
+	return Factory.NewExtensions.Num() > 0;
+}
+
+void FClipboardExtension::ImportDataToWidget(const UWidget* Widget, FName OldWidgetName)
+{
+	if (UWidgetBlueprint* WidgetBlueprint = FWidgetBlueprintEditorUtils::GetWidgetBlueprintFromWidget(Widget))
+	{
+		if (UMVVMWidgetBlueprintExtension_View* MVVMExtensionViewPtr = UMVVMWidgetBlueprintExtension_View::GetExtension<UMVVMWidgetBlueprintExtension_View>(WidgetBlueprint))
+		{
+			for (int32 Index = Factory.NewExtensions.Num() - 1; Index >= 0; Index--)
+			{
+				if (UMVVMBlueprintViewExtension* NewExtension = Factory.NewExtensions[Index].Get())
+				{
+					if (NewExtension->WidgetRenamed(OldWidgetName, Widget->GetFName()))
+					{
+						NewExtension->Rename(nullptr, MVVMExtensionViewPtr, NewExtension->GetFlags());
+						MVVMExtensionViewPtr->Modify();
+						NewExtension->Modify();
+						FMVVMExtensionItem CopiedExtension;
+						CopiedExtension.ExtensionObj = NewExtension;
+						CopiedExtension.WidgetName = Widget->GetFName();
+						MVVMExtensionViewPtr->BlueprintExtensions.Add(CopiedExtension);
+						Factory.NewExtensions.RemoveAtSwap(Index);
+					}
+				}
+			}
+		}
+	}
+}
+
+FClipboardExtension::FExtensionTextFactory::FExtensionTextFactory()
+	: FCustomizableTextObjectFactory(GWarn)
+{
+}
+
+bool FClipboardExtension::FExtensionTextFactory::CanCreateClass(UClass* ObjectClass, bool& bOmitSubObjs) const
+{
+	return ObjectClass && ObjectClass->IsChildOf(UMVVMBlueprintViewExtension::StaticClass());
+}
+
+void FClipboardExtension::FExtensionTextFactory::ProcessConstructedObject(UObject* NewObject)
+{
+	check(NewObject);
+
+	if (UMVVMBlueprintViewExtension* Extension = Cast<UMVVMBlueprintViewExtension>(NewObject))
+	{
+		NewExtensions.Add(Extension);
+	}
+}
+
+}
+#undef LOCTEXT_NAMESPACE
diff --git a/Engine/Plugins/Runtime/ModelViewViewModel/Source/ModelViewViewModelEditor/Private/Customizations/MVVMClipboardExtension.h b/Engine/Plugins/Runtime/ModelViewViewModel/Source/ModelViewViewModelEditor/Private/Customizations/MVVMClipboardExtension.h
new file mode 100644
--- /dev/null
+++ b/Engine/Plugins/Runtime/ModelViewViewModel/Source/ModelViewViewModelEditor/Private/Customizations/MVVMClipboardExtension.h
@@ -0,0 +1,45 @@
+// Copyright Epic Games, Inc. All Rights Reserved.
+
+#pragma once
+
+#include "Factories.h"
+#include "IHasClipboardExtensibility.h"
+
+class UMVVMBlueprintViewExtension;
+class UWidget;
+class UWidgetBlueprint;
+
+namespace UE::MVVM
+{
+
+class FClipboardExtension
+	: public IClipboardExtension
+{
+	//~ Begin IClipboardExtension overrides
+	virtual void AppendToClipboard(const UWidget* Widget, const FExportArgs& ExportArgs) override;
+	virtual bool CanAppendToClipboard(const UWidget* Widget) const override;
+	virtual void ProcessImportedText(const UWidgetBlueprint* WidgetBlueprint, const FString& TextToImport, UPackage*& TempPackage) override;
+	virtual bool CanImportFromClipboard(const UWidget* Widget) const override;
+	virtual void ImportDataToWidget(const UWidget* Widget, FName OldWidgetName) override;
+	//~ End IClipboardExtension overrides
+
+private:
+
+	class FExtensionTextFactory : public FCustomizableTextObjectFactory
+	{
+	public:
+		FExtensionTextFactory();
+
+		//~ Begin FCustomizableTextObjectFactory overrides
+		virtual bool CanCreateClass(UClass* ObjectClass, bool& bOmitSubObjs) const override;
+		virtual void ProcessConstructedObject(UObject* NewObject) override;
+		//~ End FCustomizableTextObjectFactory overrides
+
+	public:
+		TArray<TWeakObjectPtr<UMVVMBlueprintViewExtension>> NewExtensions;
+	};
+
+	FExtensionTextFactory Factory;
+};
+
+}
diff --git a/Engine/Plugins/Runtime/ModelViewViewModel/Source/ModelViewViewModelEditor/Private/Customizations/MVVMListViewBaseExtensionCustomizationExtender.cpp b/Engine/Plugins/Runtime/ModelViewViewModel/Source/ModelViewViewModelEditor/Private/Customizations/MVVMListViewBaseExtensionCustomizationExtender.cpp
--- a/Engine/Plugins/Runtime/ModelViewViewModel/Source/ModelViewViewModelEditor/Private/Customizations/MVVMListViewBaseExtensionCustomizationExtender.cpp
+++ b/Engine/Plugins/Runtime/ModelViewViewModel/Source/ModelViewViewModelEditor/Private/Customizations/MVVMListViewBaseExtensionCustomizationExtender.cpp
@@ -11,8 +11,10 @@
 #include "MVVMWidgetBlueprintExtension_View.h"
 #include "MVVMBlueprintView.h"
 #include "MVVMBlueprintViewModel.h"
+#include "MVVMDeveloperProjectSettings.h"
 #include "PropertyCustomizationHelpers.h"
 #include "PropertyHandle.h"
+#include "ScopedTransaction.h"
 #include "WidgetBlueprint.h"
 #include "WidgetBlueprintEditor.h"
 #include "Widgets/Input/SButton.h"
@@ -36,99 +38,102 @@ void FMVVMListViewBaseExtensionCustomizationExtender::CustomizeDetails(IDetailLa
 	{
 		if (UListViewBase* ListView = Cast<UListViewBase>(InWidgets[0]))
 		{
-			Widget = ListView;
-			WidgetBlueprintEditor = InWidgetBlueprintEditor;
-
-			// Only do a customization if we have a MVVM blueprint view class on this blueprint.
-			if (GetExtensionViewForSelectedWidgetBlueprint())
+			if (GetDefault<UMVVMDeveloperProjectSettings>()->IsExtensionSupportedForListViewBaseClass(ListView->GetClass()))
 			{
-				IDetailCategoryBuilder& MVVMCategory = InDetailLayout.EditCategory("ListEntries");
-
-				// Fetch the entry widget class handle.
-				TArray<TSharedRef<IPropertyHandle>> ListEntryProperties;
-				MVVMCategory.GetDefaultProperties(ListEntryProperties);
-				TSharedRef<IPropertyHandle>* EntryClassPtr = ListEntryProperties.FindByPredicate( [](const TSharedRef<IPropertyHandle> Property) 
-				{ 
-					return Property->GetPropertyDisplayName().EqualToCaseIgnored(FText::FromString(TEXT("Entry Widget Class"))); 
-				});
-
-				check(EntryClassPtr);
-				EntryClassHandle = *EntryClassPtr;
-				EntryClassHandle->SetOnPropertyValueChanged(FSimpleDelegate::CreateSP(this, &FMVVMListViewBaseExtensionCustomizationExtender::HandleEntryClassChanged, false));
-				bIsExtensionAdded = GetListBaseViewExtension() != nullptr;
-				HandleEntryClassChanged(true);
-
-				// Add a button that controls adding/removing the extension on the ListViewBase widget
-				MVVMCategory.AddCustomRow(FText::FromString(TEXT("Viewmodel")))
-				.NameContent()
-				[
-					SNew(STextBlock)
-					.Text(LOCTEXT("VMSupport", "Viewmodel Support"))
-					.Font(IDetailLayoutBuilder::GetDetailFont())
-				]
-				.ValueContent()
-				.HAlign(HAlign_Fill)
-				[
-					SNew(SHorizontalBox)
-					+ SHorizontalBox::Slot()
-					.AutoWidth()
+				Widget = ListView;
+				WidgetBlueprintEditor = InWidgetBlueprintEditor;
+
+				// Only do a customization if we have a MVVM blueprint view class on this blueprint.
+				if (GetExtensionViewForSelectedWidgetBlueprint())
+				{
+					IDetailCategoryBuilder& MVVMCategory = InDetailLayout.EditCategory("ListEntries");
+
+					// Fetch the entry widget class handle.
+					TArray<TSharedRef<IPropertyHandle>> ListEntryProperties;
+					MVVMCategory.GetDefaultProperties(ListEntryProperties);
+					TSharedRef<IPropertyHandle>* EntryClassPtr = ListEntryProperties.FindByPredicate( [](const TSharedRef<IPropertyHandle> Property) 
+					{ 
+						return Property->GetPropertyDisplayName().EqualToCaseIgnored(FText::FromString(TEXT("Entry Widget Class"))); 
+					});
+
+					check(EntryClassPtr);
+					EntryClassHandle = *EntryClassPtr;
+					EntryClassHandle->SetOnPropertyValueChanged(FSimpleDelegate::CreateSP(this, &FMVVMListViewBaseExtensionCustomizationExtender::HandleEntryClassChanged, false));
+					bIsExtensionAdded = GetListBaseViewExtension() != nullptr;
+					HandleEntryClassChanged(true);
+
+					// Add a button that controls adding/removing the extension on the ListViewBase widget
+					MVVMCategory.AddCustomRow(FText::FromString(TEXT("Viewmodel")))
+					.NameContent()
+					[
+						SNew(STextBlock)
+						.Text(LOCTEXT("VMSupport", "Viewmodel Support"))
+						.Font(IDetailLayoutBuilder::GetDetailFont())
+					]
+					.ValueContent()
+					.HAlign(HAlign_Fill)
 					[
-						SNew(SButton)
-						.OnClicked(this, &FMVVMListViewBaseExtensionCustomizationExtender::ModifyExtension)
+						SNew(SHorizontalBox)
+						+ SHorizontalBox::Slot()
+						.AutoWidth()
 						[
-							SNew(SHorizontalBox)
-							+ SHorizontalBox::Slot()
-							.HAlign(HAlign_Center)
-							.VAlign(VAlign_Center)
-							.AutoWidth()
+							SNew(SButton)
+							.OnClicked(this, &FMVVMListViewBaseExtensionCustomizationExtender::ModifyExtension)
 							[
-								SNew(SImage)
-								.Image(this, &FMVVMListViewBaseExtensionCustomizationExtender::GetExtensionButtonIcon)
+								SNew(SHorizontalBox)
+								+ SHorizontalBox::Slot()
+								.HAlign(HAlign_Center)
+								.VAlign(VAlign_Center)
+								.AutoWidth()
+								[
+									SNew(SImage)
+									.Image(this, &FMVVMListViewBaseExtensionCustomizationExtender::GetExtensionButtonIcon)
+								]
+								+ SHorizontalBox::Slot()
+								.Padding(FMargin(3.0f, 0.0f, 0.0f, 0.0f))
+								.VAlign(VAlign_Center)
+								.AutoWidth()
+								[
+									SNew(STextBlock)
+									.TextStyle(FAppStyle::Get(), "SmallButtonText")
+									.Text(this, &FMVVMListViewBaseExtensionCustomizationExtender::GetExtensionButtonText)
+								]
 							]
-							+ SHorizontalBox::Slot()
-							.Padding(FMargin(3, 0, 0, 0))
-							.VAlign(VAlign_Center)
-							.AutoWidth()
+						]
+					];
+
+					// Add a combobox that allows selecting from the viewmodels in the entry widgets
+					MVVMCategory.AddCustomRow(FText::FromString(TEXT("Viewmodel")))
+					.Visibility(TAttribute<EVisibility>::Create(TAttribute<EVisibility>::FGetter::CreateSP(this, &FMVVMListViewBaseExtensionCustomizationExtender::GetEntryViewModelVisibility)))
+					.NameContent()
+					[
+						SNew(STextBlock)
+						.Text(LOCTEXT("EntryVM", "Entry Viewmodel"))
+						.Font(IDetailLayoutBuilder::GetDetailFont())
+					]
+					.ValueContent()
+					.HAlign(HAlign_Fill)
+					[
+						SNew(SHorizontalBox)
+						+ SHorizontalBox::Slot()
+						.AutoWidth()
+						[
+							SNew(SComboButton)
+							.OnGetMenuContent(this, &FMVVMListViewBaseExtensionCustomizationExtender::OnGetViewModelsMenuContent)
+							.ButtonContent()
 							[
 								SNew(STextBlock)
-								.TextStyle(FAppStyle::Get(), "SmallButtonText")
-								.Text(this, &FMVVMListViewBaseExtensionCustomizationExtender::GetExtensionButtonText)
+								.Text(this, &FMVVMListViewBaseExtensionCustomizationExtender::OnGetSelectedViewModel)
 							]
 						]
-					]
-				];
-
-				// Add a combobox that allows selecting from the viewmodels in the entry widgets
-				MVVMCategory.AddCustomRow(FText::FromString(TEXT("Viewmodel")))
-				.Visibility(TAttribute<EVisibility>::Create(TAttribute<EVisibility>::FGetter::CreateSP(this, &FMVVMListViewBaseExtensionCustomizationExtender::GetEntryViewModelVisibility)))
-				.NameContent()
-				[
-					SNew(STextBlock)
-					.Text(LOCTEXT("EntryVM", "Entry Viewmodel"))
-					.Font(IDetailLayoutBuilder::GetDetailFont())
-				]
-				.ValueContent()
-				.HAlign(HAlign_Fill)
-				[
-					SNew(SHorizontalBox)
-					+ SHorizontalBox::Slot()
-					.AutoWidth()
-					[
-						SNew(SComboButton)
-						.OnGetMenuContent(this, &FMVVMListViewBaseExtensionCustomizationExtender::OnGetViewModelsMenuContent)
-						.ButtonContent()
+						+ SHorizontalBox::Slot()
+						.AutoWidth()
 						[
-							SNew(STextBlock)
-							.Text(this, &FMVVMListViewBaseExtensionCustomizationExtender::OnGetSelectedViewModel)
+							PropertyCustomizationHelpers::MakeClearButton(
+							FSimpleDelegate::CreateSP(this, &FMVVMListViewBaseExtensionCustomizationExtender::ClearEntryViewModel))
 						]
-					]
-					+ SHorizontalBox::Slot()
-					.AutoWidth()
-					[
-						PropertyCustomizationHelpers::MakeClearButton(
-						FSimpleDelegate::CreateSP(this, &FMVVMListViewBaseExtensionCustomizationExtender::ClearEntryViewModel))
-					]
-				];
+					];
+				}
 			}
 		}
 	}
@@ -316,15 +321,19 @@ void FMVVMListViewBaseExtensionCustomizationExtender::SetEntryViewModel(FGuid In
 		{
 			if (UMVVMViewBlueprintListViewBaseExtension* ListBaseViewExtension = GetListBaseViewExtension())
 			{
-				ListBaseViewExtension->Modify();
-				ListBaseViewExtension->EntryViewModelId = InEntryViewModelId;
-				if (bMarkModified)
+				if (ListBaseViewExtension->EntryViewModelId != InEntryViewModelId)
 				{
-					if (const TSharedPtr<FWidgetBlueprintEditor> BPEditor = WidgetBlueprintEditor.Pin())
+					const FScopedTransaction Transaction(LOCTEXT("SetEntryViewModel", "Set Entry ViewModel"));
+					ListBaseViewExtension->Modify();
+					ListBaseViewExtension->EntryViewModelId = InEntryViewModelId;
+					if (bMarkModified)
 					{
-						if (UWidgetBlueprint* Blueprint = BPEditor->GetWidgetBlueprintObj())
+						if (const TSharedPtr<FWidgetBlueprintEditor> BPEditor = WidgetBlueprintEditor.Pin())
 						{
-							FBlueprintEditorUtils::MarkBlueprintAsStructurallyModified(Blueprint);
+							if (UWidgetBlueprint* Blueprint = BPEditor->GetWidgetBlueprintObj())
+							{
+								FBlueprintEditorUtils::MarkBlueprintAsStructurallyModified(Blueprint);
+							}
 						}
 					}
 				}
diff --git a/Engine/Plugins/Runtime/ModelViewViewModel/Source/ModelViewViewModelEditor/Private/Customizations/MVVMPanelWidgetExtensionCustomizationExtender.cpp b/Engine/Plugins/Runtime/ModelViewViewModel/Source/ModelViewViewModelEditor/Private/Customizations/MVVMPanelWidgetExtensionCustomizationExtender.cpp
--- a/Engine/Plugins/Runtime/ModelViewViewModel/Source/ModelViewViewModelEditor/Private/Customizations/MVVMPanelWidgetExtensionCustomizationExtender.cpp
+++ b/Engine/Plugins/Runtime/ModelViewViewModel/Source/ModelViewViewModelEditor/Private/Customizations/MVVMPanelWidgetExtensionCustomizationExtender.cpp
@@ -11,8 +11,10 @@
 #include "MVVMWidgetBlueprintExtension_View.h"
 #include "MVVMBlueprintView.h"
 #include "MVVMBlueprintViewModel.h"
+#include "MVVMDeveloperProjectSettings.h"
 #include "PropertyCustomizationHelpers.h"
 #include "PropertyHandle.h"
+#include "ScopedTransaction.h"
 #include "WidgetBlueprint.h"
 #include "WidgetBlueprintEditor.h"
 #include "Widgets/Input/SButton.h"
@@ -38,134 +40,137 @@ void FMVVMPanelWidgetExtensionCustomizationExtender::CustomizeDetails(IDetailLay
 		{
 			if (Panel->CanHaveMultipleChildren())
 			{
-				Widget = Panel;
-				WidgetBlueprintEditor = InWidgetBlueprintEditor;
-
-				// Only do a customization if we have a MVVM blueprint view class on this blueprint.
-				if (GetExtensionViewForSelectedWidgetBlueprint())
+				if (GetDefault<UMVVMDeveloperProjectSettings>()->IsExtensionSupportedForPanelClass(Panel->GetClass()))
 				{
-					IDetailCategoryBuilder& MVVMCategory = InDetailLayout.EditCategory("Viewmodel");
-
-					bIsExtensionAdded = GetPanelWidgetExtension() != nullptr;
-
-					// Add a button that controls adding/removing the extension on the panel widget
-					MVVMCategory.AddCustomRow(FText::FromString(TEXT("Viewmodel")))
-					.NameContent()
-					[
-						SNew(STextBlock)
-						.Text(LOCTEXT("VMSupport", "Viewmodel Support"))
-						.Font(IDetailLayoutBuilder::GetDetailFont())
-					]
-					.ValueContent()
-					.HAlign(HAlign_Fill)
-					[
-						SNew(SHorizontalBox)
-						+ SHorizontalBox::Slot()
-						.AutoWidth()
-						[
-							SNew(SButton)
-							.OnClicked(this, &FMVVMPanelWidgetExtensionCustomizationExtender::ModifyExtension)
-							[
-								SNew(SHorizontalBox)
-								+ SHorizontalBox::Slot()
-								.HAlign(HAlign_Center)
-								.VAlign(VAlign_Center)
-								.AutoWidth()
-								[
-									SNew(SImage)
-									.Image(this, &FMVVMPanelWidgetExtensionCustomizationExtender::GetExtensionButtonIcon)
-								]
-								+ SHorizontalBox::Slot()
-								.Padding(FMargin(3, 0, 0, 0))
-								.VAlign(VAlign_Center)
-								.AutoWidth()
-								[
-									SNew(STextBlock)
-									.TextStyle(FAppStyle::Get(), "SmallButtonText")
-									.Text(this, &FMVVMPanelWidgetExtensionCustomizationExtender::GetExtensionButtonText)
-								]
-							]
-						]
-					];
+					Widget = Panel;
+					WidgetBlueprintEditor = InWidgetBlueprintEditor;
 
-					if (UMVVMViewBlueprintPanelWidgetExtension* PanelExtension = GetPanelWidgetExtension())
+					// Only do a customization if we have a MVVM blueprint view class on this blueprint.
+					if (GetExtensionViewForSelectedWidgetBlueprint())
 					{
-						IDetailPropertyRow* PanelExtensionPropertyRow = MVVMCategory.AddExternalObjects({ PanelExtension});
-						TSharedPtr<IPropertyHandle> PanelExtensionObjectHandle = PanelExtensionPropertyRow->GetPropertyHandle();
-						PanelExtensionPropertyRow->Visibility(EVisibility::Collapsed);
+						IDetailCategoryBuilder& MVVMCategory = InDetailLayout.EditCategory("Viewmodel");
 
-						// "Entry Widget Class" property row
-						EntryClassHandle = PanelExtensionObjectHandle->GetChildHandle("EntryWidgetClass");
-						EntryClassHandle->SetOnPropertyValueChanged(FSimpleDelegate::CreateSP(this, &FMVVMPanelWidgetExtensionCustomizationExtender::HandleEntryClassChanged, false));
-						HandleEntryClassChanged(true);
-						IDetailPropertyRow& EntryClassRow = MVVMCategory.AddProperty(EntryClassHandle);
+						bIsExtensionAdded = GetPanelWidgetExtension() != nullptr;
 
-						// "Entry Viewmodel" property row
+						// Add a button that controls adding/removing the extension on the panel widget
 						MVVMCategory.AddCustomRow(FText::FromString(TEXT("Viewmodel")))
-							.NameContent()
-							[
-								SNew(STextBlock)
-								.Text(LOCTEXT("EntryVM", "Entry Viewmodel"))
-								.Font(IDetailLayoutBuilder::GetDetailFont())
-							]
-							.ValueContent()
-							.HAlign(HAlign_Fill)
+						.NameContent()
+						[
+							SNew(STextBlock)
+							.Text(LOCTEXT("VMSupport", "Viewmodel Support"))
+							.Font(IDetailLayoutBuilder::GetDetailFont())
+						]
+						.ValueContent()
+						.HAlign(HAlign_Fill)
+						[
+							SNew(SHorizontalBox)
+							+ SHorizontalBox::Slot()
+							.AutoWidth()
 							[
-								SNew(SHorizontalBox)
-								+ SHorizontalBox::Slot()
-								.AutoWidth()
+								SNew(SButton)
+								.OnClicked(this, &FMVVMPanelWidgetExtensionCustomizationExtender::ModifyExtension)
 								[
-									SNew(SComboButton)
-									.OnGetMenuContent(this, &FMVVMPanelWidgetExtensionCustomizationExtender::OnGetViewModelsMenuContent)
-									.ButtonContent()
+									SNew(SHorizontalBox)
+									+ SHorizontalBox::Slot()
+									.HAlign(HAlign_Center)
+									.VAlign(VAlign_Center)
+									.AutoWidth()
+									[
+										SNew(SImage)
+										.Image(this, &FMVVMPanelWidgetExtensionCustomizationExtender::GetExtensionButtonIcon)
+									]
+									+ SHorizontalBox::Slot()
+									.Padding(FMargin(3.0f, 0.0f, 0.0f, 0.0f))
+									.VAlign(VAlign_Center)
+									.AutoWidth()
 									[
 										SNew(STextBlock)
-										.Text(this, &FMVVMPanelWidgetExtensionCustomizationExtender::OnGetSelectedViewModel)
+										.TextStyle(FAppStyle::Get(), "SmallButtonText")
+										.Text(this, &FMVVMPanelWidgetExtensionCustomizationExtender::GetExtensionButtonText)
 									]
 								]
-								+ SHorizontalBox::Slot()
-								.AutoWidth()
+							]
+						];
+
+						if (UMVVMViewBlueprintPanelWidgetExtension* PanelExtension = GetPanelWidgetExtension())
+						{
+							IDetailPropertyRow* PanelExtensionPropertyRow = MVVMCategory.AddExternalObjects({ PanelExtension});
+							TSharedPtr<IPropertyHandle> PanelExtensionObjectHandle = PanelExtensionPropertyRow->GetPropertyHandle();
+							PanelExtensionPropertyRow->Visibility(EVisibility::Collapsed);
+
+							// "Entry Widget Class" property row
+							EntryClassHandle = PanelExtensionObjectHandle->GetChildHandle("EntryWidgetClass");
+							EntryClassHandle->SetOnPropertyValueChanged(FSimpleDelegate::CreateSP(this, &FMVVMPanelWidgetExtensionCustomizationExtender::HandleEntryClassChanged, false));
+							HandleEntryClassChanged(true);
+							IDetailPropertyRow& EntryClassRow = MVVMCategory.AddProperty(EntryClassHandle);
+
+							// "Entry Viewmodel" property row
+							MVVMCategory.AddCustomRow(FText::FromString(TEXT("Viewmodel")))
+								.NameContent()
 								[
-									PropertyCustomizationHelpers::MakeClearButton(
-										FSimpleDelegate::CreateSP(this, &FMVVMPanelWidgetExtensionCustomizationExtender::ClearEntryViewModel))
+									SNew(STextBlock)
+									.Text(LOCTEXT("EntryVM", "Entry Viewmodel"))
+									.Font(IDetailLayoutBuilder::GetDetailFont())
 								]
-							];
+								.ValueContent()
+								.HAlign(HAlign_Fill)
+								[
+									SNew(SHorizontalBox)
+									+ SHorizontalBox::Slot()
+									.AutoWidth()
+									[
+										SNew(SComboButton)
+										.OnGetMenuContent(this, &FMVVMPanelWidgetExtensionCustomizationExtender::OnGetViewModelsMenuContent)
+										.ButtonContent()
+										[
+											SNew(STextBlock)
+											.Text(this, &FMVVMPanelWidgetExtensionCustomizationExtender::OnGetSelectedViewModel)
+										]
+									]
+									+ SHorizontalBox::Slot()
+									.AutoWidth()
+									[
+										PropertyCustomizationHelpers::MakeClearButton(
+											FSimpleDelegate::CreateSP(this, &FMVVMPanelWidgetExtensionCustomizationExtender::ClearEntryViewModel))
+									]
+								];
 
-						// "Slot template" property row
-						IDetailPropertyRow* SlotDetailRow = MVVMCategory.AddExternalObjects({ PanelExtension->SlotObj }, EPropertyLocation::Default,
-							FAddPropertyParams()
-							.CreateCategoryNodes(false)
-							.AllowChildren(true)
-							.HideRootObjectNode(false)
-						);
+							// "Slot template" property row
+							IDetailPropertyRow* SlotDetailRow = MVVMCategory.AddExternalObjects({ PanelExtension->SlotObj }, EPropertyLocation::Default,
+								FAddPropertyParams()
+								.CreateCategoryNodes(false)
+								.AllowChildren(true)
+								.HideRootObjectNode(false)
+							);
 
-						TSharedPtr<IPropertyHandle> SlotPropertyHandle = SlotDetailRow->GetPropertyHandle();
+							TSharedPtr<IPropertyHandle> SlotPropertyHandle = SlotDetailRow->GetPropertyHandle();
 
-						SlotDetailRow->CustomWidget(true)
-							.NameContent()
-							[
-								SNew(STextBlock)
-								.Text(LOCTEXT("SlotTemplate", "Slot Template"))
-								.Font(IDetailLayoutBuilder::GetDetailFont())
-							]
-						.ValueContent()
-							[
-								SlotPropertyHandle->CreatePropertyValueWidget()
-							];
-
-						// Because AddExternalObjects was used the property system will not add a reset to default widget by default 
-						SlotDetailRow->OverrideResetToDefault(
-							FResetToDefaultOverride::Create(
-								FIsResetToDefaultVisible::CreateLambda([SlotPropertyHandle](TSharedPtr<IPropertyHandle> Handle)
-									{
-										return SlotPropertyHandle->CanResetToDefault();
-									}),
-								FResetToDefaultHandler::CreateLambda([SlotPropertyHandle](TSharedPtr<IPropertyHandle> Handle)
-									{
-										return SlotPropertyHandle->ResetToDefault();
-									})
-								)
-						);
+							SlotDetailRow->CustomWidget(true)
+								.NameContent()
+								[
+									SNew(STextBlock)
+									.Text(LOCTEXT("SlotTemplate", "Slot Template"))
+									.Font(IDetailLayoutBuilder::GetDetailFont())
+								]
+							.ValueContent()
+								[
+									SlotPropertyHandle->CreatePropertyValueWidget()
+								];
+
+							// Because AddExternalObjects was used the property system will not add a reset to default widget by default 
+							SlotDetailRow->OverrideResetToDefault(
+								FResetToDefaultOverride::Create(
+									FIsResetToDefaultVisible::CreateLambda([SlotPropertyHandle](TSharedPtr<IPropertyHandle> Handle)
+										{
+											return SlotPropertyHandle->CanResetToDefault();
+										}),
+									FResetToDefaultHandler::CreateLambda([SlotPropertyHandle](TSharedPtr<IPropertyHandle> Handle)
+										{
+											return SlotPropertyHandle->ResetToDefault();
+										})
+									)
+							);
+						}
 					}
 				}
 			}
@@ -203,7 +208,7 @@ void FMVVMPanelWidgetExtensionCustomizationExtender::CreatePanelWidgetViewExtens
 				UMVVMViewBlueprintPanelWidgetExtension* NewPanelWidgetExtension = CastChecked<UMVVMViewBlueprintPanelWidgetExtension>(NewExtension);
 				NewPanelWidgetExtension->WidgetName = WidgetPtr->GetFName();
 
-				UPanelSlot* SlotObj = NewObject<UPanelSlot>(NewPanelWidgetExtension, WidgetPtr->GetSlotClass(), NAME_None);
+				UPanelSlot* SlotObj = NewObject<UPanelSlot>(NewPanelWidgetExtension, WidgetPtr->GetSlotClass(), NAME_None, RF_Transactional);
 				NewPanelWidgetExtension->SlotObj = SlotObj;
 			}
 		}
@@ -357,15 +362,19 @@ void FMVVMPanelWidgetExtensionCustomizationExtender::SetEntryViewModel(FGuid InE
 		{
 			if (UMVVMViewBlueprintPanelWidgetExtension* PanelWidgetExtension = GetPanelWidgetExtension())
 			{
-				PanelWidgetExtension->Modify();
-				PanelWidgetExtension->EntryViewModelId = InEntryViewModelId;
-				if (bMarkModified)
+				if (PanelWidgetExtension->EntryViewModelId != InEntryViewModelId)
 				{
-					if (const TSharedPtr<FWidgetBlueprintEditor> BPEditor = WidgetBlueprintEditor.Pin())
+					const FScopedTransaction Transaction(LOCTEXT("SetEntryViewModel", "Set Entry ViewModel"));
+					PanelWidgetExtension->Modify();
+					PanelWidgetExtension->EntryViewModelId = InEntryViewModelId;
+					if (bMarkModified)
 					{
-						if (UWidgetBlueprint* Blueprint = BPEditor->GetWidgetBlueprintObj())
+						if (const TSharedPtr<FWidgetBlueprintEditor> BPEditor = WidgetBlueprintEditor.Pin())
 						{
-							FBlueprintEditorUtils::MarkBlueprintAsStructurallyModified(Blueprint);
+							if (UWidgetBlueprint* Blueprint = BPEditor->GetWidgetBlueprintObj())
+							{
+								FBlueprintEditorUtils::MarkBlueprintAsStructurallyModified(Blueprint);
+							}
 						}
 					}
 				}
diff --git a/Engine/Plugins/Runtime/ModelViewViewModel/Source/ModelViewViewModelEditor/Private/MVVMEditorSubsystem.cpp b/Engine/Plugins/Runtime/ModelViewViewModel/Source/ModelViewViewModelEditor/Private/MVVMEditorSubsystem.cpp
--- a/Engine/Plugins/Runtime/ModelViewViewModel/Source/ModelViewViewModelEditor/Private/MVVMEditorSubsystem.cpp
+++ b/Engine/Plugins/Runtime/ModelViewViewModel/Source/ModelViewViewModelEditor/Private/MVVMEditorSubsystem.cpp
@@ -135,6 +135,8 @@ namespace UE::MVVM::Private
 UMVVMBlueprintView* UMVVMEditorSubsystem::RequestView(UWidgetBlueprint* WidgetBlueprint) const
 {
 	UMVVMWidgetBlueprintExtension_View* Extension = UMVVMWidgetBlueprintExtension_View::RequestExtension<UMVVMWidgetBlueprintExtension_View>(WidgetBlueprint);
+	Extension->SetFlags(RF_Transactional);
+
 	if (Extension->GetBlueprintView() == nullptr)
 	{
 		Extension->CreateBlueprintViewInstance();
diff --git a/Engine/Plugins/Runtime/ModelViewViewModel/Source/ModelViewViewModelEditor/Private/ModelViewViewModelEditorModule.cpp b/Engine/Plugins/Runtime/ModelViewViewModel/Source/ModelViewViewModelEditor/Private/ModelViewViewModelEditorModule.cpp
--- a/Engine/Plugins/Runtime/ModelViewViewModel/Source/ModelViewViewModelEditor/Private/ModelViewViewModelEditorModule.cpp
+++ b/Engine/Plugins/Runtime/ModelViewViewModel/Source/ModelViewViewModelEditor/Private/ModelViewViewModelEditorModule.cpp
@@ -7,6 +7,7 @@
 #include "BlueprintModes/WidgetBlueprintApplicationModes.h"
 #include "Customizations/MVVMBlueprintViewModelContextCustomization.h"
 #include "Customizations/MVVMListViewBaseExtensionCustomizationExtender.h"
+#include "Customizations/MVVMClipboardExtension.h"
 #include "Customizations/MVVMPanelWidgetExtensionCustomizationExtender.h"
 #include "Customizations/MVVMPropertyBindingExtension.h"
 #include "Extensions/MVVMBlueprintViewExtension.h"
@@ -47,6 +48,9 @@ void FModelViewViewModelEditorModule::StartupModule()
 	PropertyBindingExtension = MakeShared<UE::MVVM::FMVVMPropertyBindingExtension>();
 	UMGEditorModule.GetPropertyBindingExtensibilityManager()->AddExtension(PropertyBindingExtension.ToSharedRef());
 
+	ClipboardExtension = MakeShared<UE::MVVM::FClipboardExtension>();
+	UMGEditorModule.GetClipboardExtensibilityManager()->AddExtension(ClipboardExtension.ToSharedRef());
+
 	ListViewBaseCustomizationExtender = UE::MVVM::FMVVMListViewBaseExtensionCustomizationExtender::MakeInstance();
 	UMGEditorModule.AddWidgetCustomizationExtender(ListViewBaseCustomizationExtender.ToSharedRef());
 
@@ -93,6 +97,7 @@ void FModelViewViewModelEditorModule::ShutdownModule()
 	if (IUMGEditorModule* UMGEditorModule = FModuleManager::GetModulePtr<IUMGEditorModule>("UMGEditor"))
 	{
 		UMGEditorModule->OnRegisterTabsForEditor().RemoveAll(this);
+		UMGEditorModule->GetClipboardExtensibilityManager()->RemoveExtension(ClipboardExtension.ToSharedRef());
 		UMGEditorModule->GetPropertyBindingExtensibilityManager()->RemoveExtension(PropertyBindingExtension.ToSharedRef());
 		if (UObjectInitialized())
 		{
diff --git a/Engine/Plugins/Runtime/ModelViewViewModel/Source/ModelViewViewModelEditor/Private/ModelViewViewModelEditorModule.h b/Engine/Plugins/Runtime/ModelViewViewModel/Source/ModelViewViewModelEditor/Private/ModelViewViewModelEditorModule.h
--- a/Engine/Plugins/Runtime/ModelViewViewModel/Source/ModelViewViewModelEditor/Private/ModelViewViewModelEditorModule.h
+++ b/Engine/Plugins/Runtime/ModelViewViewModel/Source/ModelViewViewModelEditor/Private/ModelViewViewModelEditorModule.h
@@ -8,6 +8,7 @@
 namespace UE::MVVM
 {
 	class FAssetTypeActions_ViewModelBlueprint;
+	class FClipboardExtension;
 	class FMVVMListViewBaseExtensionCustomizationExtender;
 	class FMVVMPanelWidgetExtensionCustomizationExtender;
 	class FMVVMPropertyBindingExtension;
@@ -45,6 +46,7 @@ private:
 
 private:
 	TSharedPtr<UE::MVVM::FMVVMPropertyBindingExtension> PropertyBindingExtension;
+	TSharedPtr<UE::MVVM::FClipboardExtension> ClipboardExtension;
 	TSharedPtr<UE::MVVM::FAssetTypeActions_ViewModelBlueprint> ViewModelBlueprintActions;
 	TSharedPtr<UE::MVVM::FMVVMListViewBaseExtensionCustomizationExtender> ListViewBaseCustomizationExtender;
 	TSharedPtr<UE::MVVM::FMVVMPanelWidgetExtensionCustomizationExtender> PanelWidgetCustomizationExtender;
diff --git a/Engine/Plugins/Runtime/ModelViewViewModel/Source/ModelViewViewModelEditor/Private/Widgets/SMVVMViewModelPanel.cpp b/Engine/Plugins/Runtime/ModelViewViewModel/Source/ModelViewViewModelEditor/Private/Widgets/SMVVMViewModelPanel.cpp
--- a/Engine/Plugins/Runtime/ModelViewViewModel/Source/ModelViewViewModelEditor/Private/Widgets/SMVVMViewModelPanel.cpp
+++ b/Engine/Plugins/Runtime/ModelViewViewModel/Source/ModelViewViewModelEditor/Private/Widgets/SMVVMViewModelPanel.cpp
@@ -371,8 +371,7 @@ void SMVVMViewModelPanel::HandleViewUpdated(UBlueprintExtension*)
 				UMVVMBlueprintView* CurrentBlueprintView = GEditor->GetEditorSubsystem<UMVVMEditorSubsystem>()->GetView(WidgetBlueprint);
 				if (CurrentBlueprintView == nullptr)
 				{
-					UMVVMWidgetBlueprintExtension_View::RequestExtension<UMVVMWidgetBlueprintExtension_View>(WidgetBlueprint)->CreateBlueprintViewInstance();
-					CurrentBlueprintView = GEditor->GetEditorSubsystem<UMVVMEditorSubsystem>()->GetView(WidgetBlueprint);
+					CurrentBlueprintView = GEditor->GetEditorSubsystem<UMVVMEditorSubsystem>()->RequestView(WidgetBlueprint);
 				}
 				WeakBlueprintView = CurrentBlueprintView;
 
diff --git a/Engine/Source/Editor/UMGEditor/Private/UMGEditorModule.cpp b/Engine/Source/Editor/UMGEditor/Private/UMGEditorModule.cpp
--- a/Engine/Source/Editor/UMGEditor/Private/UMGEditorModule.cpp
+++ b/Engine/Source/Editor/UMGEditor/Private/UMGEditorModule.cpp
@@ -85,6 +85,7 @@ public:
 		DesignerExtensibilityManager->AddDesignerExtensionFactory(SWidgetDesignerNavigation::MakeDesignerExtension());
 
 		PropertyBindingExtensibilityManager = MakeShared<FPropertyBindingExtensibilityManager>();
+		ClipboardExtensibilityManager = MakeShared<FClipboardExtensibilityManager>();
 
 		// Register widget blueprint compiler we do this no matter what.
 		IKismetCompilerInterface& KismetCompilerModule = FModuleManager::LoadModuleChecked<IKismetCompilerInterface>("KismetCompiler");
@@ -212,6 +213,7 @@ public:
 	virtual TSharedPtr<FExtensibilityManager> GetToolBarExtensibilityManager() override { return ToolBarExtensibilityManager; }
 	virtual TSharedPtr<FDesignerExtensibilityManager> GetDesignerExtensibilityManager() override { return DesignerExtensibilityManager; }
 	virtual TSharedPtr<FPropertyBindingExtensibilityManager> GetPropertyBindingExtensibilityManager() override { return PropertyBindingExtensibilityManager; }
+	virtual TSharedPtr<FClipboardExtensibilityManager> GetClipboardExtensibilityManager() override { return ClipboardExtensibilityManager; }
 
 	/** Register settings objects. */
 	void RegisterSettings()
@@ -384,6 +386,7 @@ private:
 	TSharedPtr<FExtensibilityManager> ToolBarExtensibilityManager;
 	TSharedPtr<FDesignerExtensibilityManager> DesignerExtensibilityManager;
 	TSharedPtr<FPropertyBindingExtensibilityManager> PropertyBindingExtensibilityManager;
+	TSharedPtr<FClipboardExtensibilityManager> ClipboardExtensibilityManager;
 	TSharedPtr<FGraphPanelPinFactory> GraphPanelPinFactory;
 
 	FDelegateHandle SequenceEditorHandle;
diff --git a/Engine/Source/Editor/UMGEditor/Private/WidgetBlueprintEditorUtils.cpp b/Engine/Source/Editor/UMGEditor/Private/WidgetBlueprintEditorUtils.cpp
--- a/Engine/Source/Editor/UMGEditor/Private/WidgetBlueprintEditorUtils.cpp
+++ b/Engine/Source/Editor/UMGEditor/Private/WidgetBlueprintEditorUtils.cpp
@@ -64,6 +64,7 @@
 #include "Interfaces/ISlate3DRenderer.h"
 #include "Rendering/SlateDrawBuffer.h"
 #include "Slate/WidgetRenderer.h"
+#include "UMGEditorModule.h"
 #include "Widgets/SVirtualWindow.h"
 #include "GraphEditorActions.h"
 #include "WidgetEditingProjectSettings.h"
@@ -1613,6 +1614,17 @@ void FWidgetBlueprintEditorUtils::ExportWidgetsToText(TArray<UWidget*> WidgetsTo
 	}
 
 	const FExportObjectInnerContext Context(WidgetsToIgnore);
+
+	IUMGEditorModule& EditorModule = FModuleManager::LoadModuleChecked<IUMGEditorModule>("UMGEditor");
+	const TArrayView<const TSharedPtr<IClipboardExtension>> ClipboardExtensions = EditorModule.GetClipboardExtensibilityManager()->GetExtensions();
+
+	// Get the widget blueprint containing the exported widgets
+	UWidgetBlueprint* WidgetBlueprint = nullptr;
+	if (WidgetsToExport.Num() > 0)
+	{
+		WidgetBlueprint = FWidgetBlueprintEditorUtils::GetWidgetBlueprintFromWidget(WidgetsToExport[0]);
+	}
+
 	// Export each of the selected nodes
 	for ( UWidget* Widget : WidgetsToExport )
 	{
@@ -1641,6 +1653,26 @@ void FWidgetBlueprintEditorUtils::ExportWidgetsToText(TArray<UWidget*> WidgetsTo
 
 			UExporter::ExportToOutputDevice(&Context, SlotMetaData, nullptr, Archive, TEXT("copy"), 0, PPF_ExportsNotFullyQualified | PPF_Copy | PPF_Delimited, false, nullptr);
 		}
+
+		if (WidgetBlueprint)
+		{
+			for (const TSharedPtr<IClipboardExtension>& ClipboardExtension : ClipboardExtensions)
+			{
+				if (ClipboardExtension->CanAppendToClipboard(Widget))
+				{
+					IClipboardExtension::FExportArgs ExportArgs;
+					ExportArgs.Context = &Context;
+					ExportArgs.Exporter = nullptr;
+					ExportArgs.FileType = TEXT("copy");
+					ExportArgs.Indent = 0;
+					ExportArgs.PortFlags = PPF_ExportsNotFullyQualified | PPF_Copy | PPF_Delimited;
+					ExportArgs.bSelectedOnly = false;
+					ExportArgs.ExportRootScope = nullptr;
+					ExportArgs.Out = &Archive;
+					ClipboardExtension->AppendToClipboard(Widget, ExportArgs);
+				}
+			}
+		}
 	}
 
 	ExportedText = Archive;
@@ -1989,6 +2021,13 @@ void FWidgetBlueprintEditorUtils::ImportWidgetsFromText(UWidgetBlueprint* BP, co
 	FWidgetObjectTextFactory Factory = ProcessImportedText(BP, TextToImport, TempPackage);
 	TGCObjectScopeGuard<UPackage> TempPackageGCGuard(TempPackage);
 
+	IUMGEditorModule& EditorModule = FModuleManager::LoadModuleChecked<IUMGEditorModule>("UMGEditor");
+	const TArrayView<const TSharedPtr<IClipboardExtension>> ClipboardExtensions = EditorModule.GetClipboardExtensibilityManager()->GetExtensions();
+	for (const TSharedPtr<IClipboardExtension>& ClipboardExtension : ClipboardExtensions)
+	{
+		ClipboardExtension->ProcessImportedText(BP, TextToImport, TempPackage);
+	}
+
 	PastedExtraSlotData = Factory.MissingSlotData;
 
 	for ( auto& Entry : Factory.NewWidgetMap )
@@ -2041,6 +2080,14 @@ void FWidgetBlueprintEditorUtils::ImportWidgetsFromText(UWidgetBlueprint* BP, co
 		{
 			Widget->Rename(*WidgetOldName, BP->WidgetTree);
 		}
+
+		for (const TSharedPtr<IClipboardExtension>& Extension : ClipboardExtensions)
+		{
+			if (Extension->CanImportFromClipboard(Widget))
+			{
+				Extension->ImportDataToWidget(Widget, FName(WidgetOldName));
+			}
+		}
 	}
 }
 
@@ -2896,4 +2943,28 @@ FText FWidgetBlueprintEditorUtils::GetPaletteCategory(const FAssetData& WidgetAs
 	}
 }
 
+UWidgetBlueprint* FWidgetBlueprintEditorUtils::GetWidgetBlueprintFromWidget(const UWidget* Widget)
+{
+	if (Widget)
+	{
+		if (UObject* WidgetTree = Widget->GetOuter())
+		{
+			UWidgetBlueprint* WidgetBlueprint = Cast<UWidgetBlueprint>(WidgetTree->GetOuter());
+			if (WidgetBlueprint)
+			{
+				return WidgetBlueprint;
+			}
+			else
+			{
+				WidgetBlueprint = Cast<UWidgetBlueprint>(WidgetTree->GetOuter()->GetClass()->ClassGeneratedBy);
+				if (WidgetBlueprint)
+				{
+					return WidgetBlueprint;
+				}
+			}
+		}
+	}
+	return nullptr;
+}
+
 #undef LOCTEXT_NAMESPACE
diff --git a/Engine/Source/Editor/UMGEditor/Public/IHasClipboardExtensibility.h b/Engine/Source/Editor/UMGEditor/Public/IHasClipboardExtensibility.h
new file mode 100644
--- /dev/null
+++ b/Engine/Source/Editor/UMGEditor/Public/IHasClipboardExtensibility.h
@@ -0,0 +1,74 @@
+// Copyright Epic Games, Inc. All Rights Reserved.
+
+#pragma once
+
+#include "Containers/Array.h"
+#include "Misc/OutputDevice.h"
+#include "Templates/SharedPointer.h"
+#include "UObject/PropertyPortFlags.h"
+
+class FExportObjectInnerContext;
+class UExporter;
+class UWidget;
+class UWidgetBlueprint;
+
+class UMGEDITOR_API IClipboardExtension
+{
+public:
+	struct FExportArgs
+	{
+		const FExportObjectInnerContext* Context = nullptr;
+		UExporter* Exporter = nullptr; 
+		FOutputDevice* Out = nullptr; 
+		const TCHAR* FileType = nullptr;
+		int32 Indent = 0;
+		uint32 PortFlags = PPF_None;
+		bool bSelectedOnly = false;
+		UObject* ExportRootScope = nullptr;
+	};
+
+	virtual void AppendToClipboard(const UWidget* Widget, const FExportArgs& ExportArgs) = 0;
+	virtual bool CanAppendToClipboard(const UWidget* Widget) const = 0;
+	virtual void ProcessImportedText(const UWidgetBlueprint* WidgetBlueprint, const FString& TextToImport, UPackage*& TempPackage) = 0;
+	virtual bool CanImportFromClipboard(const UWidget* Widget) const = 0;
+	virtual void ImportDataToWidget(const UWidget* Widget, FName OldWidgetName) = 0;
+
+	virtual ~IClipboardExtension() { }
+};
+
+/**
+ * Clipboard extensibility manager holds a list of registered clipboard extensions.
+ */
+class UMGEDITOR_API FClipboardExtensibilityManager
+{
+public:
+	void AddExtension(const TSharedRef<IClipboardExtension>& Extension)
+	{
+		if (ensure(!Extensions.Contains(Extension)))
+		{
+			Extensions.Add(Extension);
+		}
+	}
+
+	void RemoveExtension(const TSharedRef<IClipboardExtension>& Extension)
+	{
+		int32 NumRemoved = Extensions.RemoveSingleSwap(Extension);
+		ensure(NumRemoved == 1);
+	}
+
+	const TArrayView<const TSharedPtr<IClipboardExtension>> GetExtensions() const
+	{
+		return Extensions;
+	}
+
+private:
+	TArray<TSharedPtr<IClipboardExtension>> Extensions;
+};
+
+/** Indicates that a class has data to append to clipboard */
+class IHasClipboardExtensibility
+{
+public:
+	virtual TSharedPtr<FClipboardExtensibilityManager> GetClipboardExtensibilityManager() = 0;
+};
+
diff --git a/Engine/Source/Editor/UMGEditor/Public/UMGEditorModule.h b/Engine/Source/Editor/UMGEditor/Public/UMGEditorModule.h
--- a/Engine/Source/Editor/UMGEditor/Public/UMGEditorModule.h
+++ b/Engine/Source/Editor/UMGEditor/Public/UMGEditorModule.h
@@ -4,6 +4,7 @@
 
 #include "Modules/ModuleInterface.h"
 #include "Toolkits/AssetEditorToolkit.h"
+#include "IHasClipboardExtensibility.h"
 #include "IHasDesignerExtensibility.h"
 #include "IHasPropertyBindingExtensibility.h"
 #include "UObject/TopLevelAssetPath.h"
@@ -22,7 +23,8 @@ class IUMGEditorModule :
 	public IHasMenuExtensibility, 
 	public IHasToolBarExtensibility, 
 	public IHasDesignerExtensibility,
-	public IHasPropertyBindingExtensibility
+	public IHasPropertyBindingExtensibility,
+	public IHasClipboardExtensibility
 {
 public:
 	virtual class FWidgetBlueprintCompiler* GetRegisteredCompiler() = 0;
diff --git a/Engine/Source/Editor/UMGEditor/Public/WidgetBlueprintEditorUtils.h b/Engine/Source/Editor/UMGEditor/Public/WidgetBlueprintEditorUtils.h
--- a/Engine/Source/Editor/UMGEditor/Public/WidgetBlueprintEditorUtils.h
+++ b/Engine/Source/Editor/UMGEditor/Public/WidgetBlueprintEditorUtils.h
@@ -124,6 +124,8 @@ public:
 	static UWidgetEditingProjectSettings* GetRelevantMutableSettings(TWeakPtr<FWidgetBlueprintEditor> CurrentEditor);
 	static const UWidgetEditingProjectSettings* GetRelevantSettings(TWeakPtr<FWidgetBlueprintEditor> CurrentEditor);
 
+	static UWidgetBlueprint* GetWidgetBlueprintFromWidget(const UWidget* Widget);
+
 private:
 
 	static FString CopyWidgetsInternal(UWidgetBlueprint* BP, TSet<FWidgetReference> Widgets);
