MVVM: Fix warning "Failed to read file '../../../Engine/Plugins/Runtime/ModelViewViewModel/Content/Editor/Common/ButtonHoverHint.png'"

diff --git a/Engine/Plugins/Runtime/ModelViewViewModel/Source/ModelViewViewModelEditor/Private/Styling/MVVMEditorStyle.cpp b/Engine/Plugins/Runtime/ModelViewViewModel/Source/ModelViewViewModelEditor/Private/Styling/MVVMEditorStyle.cpp
--- a/Engine/Plugins/Runtime/ModelViewViewModel/Source/ModelViewViewModelEditor/Private/Styling/MVVMEditorStyle.cpp
+++ b/Engine/Plugins/Runtime/ModelViewViewModel/Source/ModelViewViewModelEditor/Private/Styling/MVVMEditorStyle.cpp
@@ -30,7 +30,7 @@ FMVVMEditorStyle::FMVVMEditorStyle()
 	if (ensure(MVVMPlugin))
 	{
 		SetContentRoot(MVVMPlugin->GetContentDir() / TEXT("Editor"));
-		SetCoreContentRoot(FPaths::EngineContentDir() / TEXT("Slate"));
+		SetCoreContentRoot(FPaths::EngineContentDir() / TEXT("Editor") / TEXT("Slate"));
 	}
 
 	// Class Icons
@@ -66,9 +66,9 @@ FMVVMEditorStyle::FMVVMEditorStyle()
 	);
 
 	FButtonStyle NoStyleComboButtonButtonStyle = FButtonStyle()
-		.SetNormal(FSlateBoxBrush(RootToContentDir("Common/ButtonHoverHint.png"), FMargin(4.0f / 16.0f), FLinearColor(1.0f, 1.0f, 1.0f, 0.15f)))
-		.SetHovered(FSlateBoxBrush(RootToContentDir("Common/ButtonHoverHint.png"), FMargin(4.0f / 16.0f), FLinearColor(1.0f, 1.0f, 1.0f, 0.25f)))
-		.SetPressed(FSlateBoxBrush(RootToContentDir("Common/ButtonHoverHint.png"), FMargin(4.0f / 16.0f), FLinearColor(1.0f, 1.0f, 1.0f, 0.30f)))
+		.SetNormal(FSlateBoxBrush(RootToCoreContentDir("Common/ButtonHoverHint.png"), FMargin(4.0f / 16.0f), FLinearColor(1.0f, 1.0f, 1.0f, 0.15f)))
+		.SetHovered(FSlateBoxBrush(RootToCoreContentDir("Common/ButtonHoverHint.png"), FMargin(4.0f / 16.0f), FLinearColor(1.0f, 1.0f, 1.0f, 0.25f)))
+		.SetPressed(FSlateBoxBrush(RootToCoreContentDir("Common/ButtonHoverHint.png"), FMargin(4.0f / 16.0f), FLinearColor(1.0f, 1.0f, 1.0f, 0.30f)))
 		.SetNormalPadding(FMargin(0.0f, 0.0f, 0.0f, 1.0f))
 		.SetPressedPadding(FMargin(0.0f, 1.0f, 0.0f, 0.0f));
 	Set("NoStyleComboButton", FComboButtonStyle()
