MVVM Extension: Fix issue with multiple instances of the same userwidget sharing the same extension objects at runtime.
We make a copy of extensions from the ViewClass in each View so that Views don't invoke the extensions directly on the class.

diff --git a/Engine/Plugins/Runtime/ModelViewViewModel/Source/ModelViewViewModel/Private/Extensions/MVVMViewPanelWidgetExtension.cpp b/Engine/Plugins/Runtime/ModelViewViewModel/Source/ModelViewViewModel/Private/Extensions/MVVMViewPanelWidgetExtension.cpp
--- a/Engine/Plugins/Runtime/ModelViewViewModel/Source/ModelViewViewModel/Private/Extensions/MVVMViewPanelWidgetExtension.cpp
+++ b/Engine/Plugins/Runtime/ModelViewViewModel/Source/ModelViewViewModel/Private/Extensions/MVVMViewPanelWidgetExtension.cpp
@@ -116,31 +116,40 @@ void UMVVMViewPanelWidgetExtension::BP_SetItems(const TArray<UObject*>& InItems)
 		TArray<TTuple<UPanelSlot*, UWidget*>> NewSlots;
 		for (int32 ItemIndex = 0; ItemIndex < InItems.Num(); ++ItemIndex)
 		{
-			UObject* Item = InItems[ItemIndex];
-			const TTuple<UPanelSlot*, TScriptInterface<INotifyFieldValueChanged>>* FoundObject = PreviousSlots.FindByPredicate([Item](const auto& Other) 
-			{ 
-				return Other.Value.GetObject() == Item; 
-			});
-
-			if (Item && OwningUserWidget && Item->GetClass() != SelectedVMClass)
+			if (UObject* Item = InItems[ItemIndex])
 			{
-				UE::MVVM::FMessageLog Log(OwningUserWidget);
-				Log.Warning(FText::Format(LOCTEXT("SetPanelWidgetItemsViewmodelTypeMismatch", "The item {0} passed as an entry of widget {1} is not a viewmodel of the selected type {2}.")
-					, FText::FromString(Item->GetName()), FText::FromString(PanelWidgetPtr->GetName()), FText::FromString(SelectedVMClass->GetName())
-				));
-			}
+				const TTuple<UPanelSlot*, TScriptInterface<INotifyFieldValueChanged>>* FoundObject = PreviousSlots.FindByPredicate([Item](const auto& Other)
+					{
+						return Other.Value.GetObject() == Item;
+					});
 
-			if (!FoundObject)
-			{
-				UUserWidget* EntryWidget = UUserWidget::CreateWidgetInstance(*PanelWidgetPtr, EntryWidgetClass, NAME_None);
-				ensure(EntryWidget);
+				if (Item && OwningUserWidget && Item->GetClass() != SelectedVMClass)
+				{
+					UE::MVVM::FMessageLog Log(OwningUserWidget);
+					Log.Warning(FText::Format(LOCTEXT("SetPanelWidgetItemsViewmodelTypeMismatch", "The item {0} passed as an entry of widget {1} is not a viewmodel of the selected type {2}.")
+						, FText::FromString(Item->GetName()), FText::FromString(PanelWidgetPtr->GetName()), FText::FromString(SelectedVMClass->GetName())
+					));
+				}
 
-				SetViewModelOnEntryWidget(EntryWidget, Item, OwningUserWidget);
-				NewSlots.Add(TTuple<UPanelSlot*, UWidget*>(SlotTemplate, EntryWidget));
+				if (!FoundObject)
+				{
+					UUserWidget* EntryWidget = UUserWidget::CreateWidgetInstance(*PanelWidgetPtr, EntryWidgetClass, NAME_None);
+					ensure(EntryWidget);
+
+					SetViewModelOnEntryWidget(EntryWidget, Item, OwningUserWidget);
+					NewSlots.Add(TTuple<UPanelSlot*, UWidget*>(SlotTemplate, EntryWidget));
+				}
+				else
+				{
+					NewSlots.Add(TTuple<UPanelSlot*, UWidget*>(FoundObject->Key, FoundObject->Key->Content));
+				}
 			}
 			else
 			{
-				NewSlots.Add(TTuple<UPanelSlot*, UWidget*>(FoundObject->Key, FoundObject->Key->Content));
+				UE::MVVM::FMessageLog Log(OwningUserWidget);
+				Log.Warning(FText::Format(LOCTEXT("SetPanelWidgetItemsViewmodelNullObject", "The item at index {0} passed as an entry of widget {1} is null. An entry widget won't be generated for this item.")
+					, ItemIndex, FText::FromString(PanelWidgetPtr->GetName())
+				));
 			}
 		}
 
diff --git a/Engine/Plugins/Runtime/ModelViewViewModel/Source/ModelViewViewModel/Private/View/MVVMView.cpp b/Engine/Plugins/Runtime/ModelViewViewModel/Source/ModelViewViewModel/Private/View/MVVMView.cpp
--- a/Engine/Plugins/Runtime/ModelViewViewModel/Source/ModelViewViewModel/Private/View/MVVMView.cpp
+++ b/Engine/Plugins/Runtime/ModelViewViewModel/Source/ModelViewViewModel/Private/View/MVVMView.cpp
@@ -75,7 +75,19 @@ void UMVVMView::Construct()
 	check(GeneratedViewClass);
 	check(bConstructed == false);
 
+	// Create an independent copy of extensions per view.
+	ensure(Extensions.IsEmpty());
+	Extensions.Reset(GeneratedViewClass->GetViewClassExtensions().Num());
 	for (UMVVMViewClassExtension* Extension : GeneratedViewClass->GetViewClassExtensions())
+	{
+		UMVVMViewClassExtension* ViewExtension = DuplicateObject<UMVVMViewClassExtension>(Extension, this);
+		if (ensure(ViewExtension))
+		{
+			Extensions.Add(ViewExtension);
+		}
+	}
+
+	for (UMVVMViewClassExtension* Extension : Extensions)
 	{
 		Extension->OnViewConstructed(GetUserWidget(), this);
 	}
@@ -110,7 +122,7 @@ void UMVVMView::Destruct()
 	UninitializeEvents();
 	UninitializeSources(); // and bindings
 
-	for (UMVVMViewClassExtension* Extension : GeneratedViewClass->GetViewClassExtensions())
+	for (UMVVMViewClassExtension* Extension : Extensions)
 	{
 		Extension->OnViewDestructed(GetUserWidget(), this);
 	}
@@ -158,14 +170,14 @@ void UMVVMView::InitializeSources()
 		InitializeSourceBindingsCommon();
 	}
 
-	for (UMVVMViewClassExtension* Extension : GeneratedViewClass->GetViewClassExtensions())
+	for (UMVVMViewClassExtension* Extension : Extensions)
 	{
 		Extension->OnSourcesInitialized(GetUserWidget(), this);
 	}
 
 	if (GeneratedViewClass->DoesInitializeBindingsOnConstruct())
 	{
-		for (UMVVMViewClassExtension* Extension : GeneratedViewClass->GetViewClassExtensions())
+		for (UMVVMViewClassExtension* Extension : Extensions)
 		{
 			Extension->OnBindingsInitialized(GetUserWidget(), this);
 		}
@@ -189,7 +201,7 @@ void UMVVMView::UninitializeSources()
 		UninitializeBindings();
 	}
 
-	for (UMVVMViewClassExtension* Extension : GeneratedViewClass->GetViewClassExtensions())
+	for (UMVVMViewClassExtension* Extension : Extensions)
 	{
 		Extension->OnSourcesUninitialized(GetUserWidget(), this);
 	}
@@ -322,7 +334,7 @@ void UMVVMView::InitializeBindings()
 
 	InitializeSourceBindingsCommon();
 
-	for (UMVVMViewClassExtension* Extension : GeneratedViewClass->GetViewClassExtensions())
+	for (UMVVMViewClassExtension* Extension : Extensions)
 	{
 		Extension->OnBindingsInitialized(GetUserWidget(), this);
 	}
@@ -359,7 +371,7 @@ void UMVVMView::UninitializeBindings()
 
 	bBindingsInitialized = false;
 
-	for (UMVVMViewClassExtension* Extension : GeneratedViewClass->GetViewClassExtensions())
+	for (UMVVMViewClassExtension* Extension : Extensions)
 	{
 		Extension->OnBindingsUninitialized(GetUserWidget(), this);
 	}
@@ -1015,7 +1027,7 @@ void UMVVMView::InitializeEvents()
 		BindEvent(ClassEvent, FMVVMViewClass_EventKey(Index));
 	}
 
-	for (UMVVMViewClassExtension* Extension : GeneratedViewClass->GetViewClassExtensions())
+	for (UMVVMViewClassExtension* Extension : Extensions)
 	{
 		Extension->OnEventsInitialized(GetUserWidget(), this);
 	}
@@ -1035,7 +1047,7 @@ void UMVVMView::UninitializeEvents()
 	
 	bEventsInitialized = false;
 
-	for (UMVVMViewClassExtension* Extension : GeneratedViewClass->GetViewClassExtensions())
+	for (UMVVMViewClassExtension* Extension : Extensions)
 	{
 		Extension->OnEventsUninitialized(GetUserWidget(), this);
 	}
diff --git a/Engine/Plugins/Runtime/ModelViewViewModel/Source/ModelViewViewModel/Public/View/MVVMView.h b/Engine/Plugins/Runtime/ModelViewViewModel/Source/ModelViewViewModel/Public/View/MVVMView.h
--- a/Engine/Plugins/Runtime/ModelViewViewModel/Source/ModelViewViewModel/Public/View/MVVMView.h
+++ b/Engine/Plugins/Runtime/ModelViewViewModel/Source/ModelViewViewModel/Public/View/MVVMView.h
@@ -13,6 +13,7 @@ namespace UE::FieldNotification { struct FFieldId; }
 template <typename InterfaceType> class TScriptInterface;
 
 class UMVVMViewClass;
+class UMVVMViewClassExtension;
 struct FMVVMViewClass_Binding;
 struct FMVVMViewClass_Event;
 struct FMVVMViewClass_Source;
@@ -220,6 +221,9 @@ private:
 	/** The event that are registered by the view to the sources. */
 	TArray<FBoundEvent> BoundEvents;
 
+	UPROPERTY(VisibleAnywhere, Transient, Category = "View")
+	TArray<TObjectPtr<UMVVMViewClassExtension>> Extensions;
+
 	/** Bitfield that represents the valid sources. */
 	UPROPERTY(VisibleAnywhere, Transient, Category = "View")
 	uint64 ValidSources = 0;
