Fixed array index going out of bounds in UMVVMViewPanelWidgetExtension::ReplaceAllSlots.

#rnx

diff --git a/Engine/Plugins/Runtime/ModelViewViewModel/Source/ModelViewViewModel/Private/Extensions/MVVMViewPanelWidgetExtension.cpp b/Engine/Plugins/Runtime/ModelViewViewModel/Source/ModelViewViewModel/Private/Extensions/MVVMViewPanelWidgetExtension.cpp
--- a/Engine/Plugins/Runtime/ModelViewViewModel/Source/ModelViewViewModel/Private/Extensions/MVVMViewPanelWidgetExtension.cpp
+++ b/Engine/Plugins/Runtime/ModelViewViewModel/Source/ModelViewViewModel/Private/Extensions/MVVMViewPanelWidgetExtension.cpp
@@ -188,7 +188,7 @@ void UMVVMViewPanelWidgetExtension::ReplaceAllSlots(TArrayView<TTuple<UPanelSlot
 		const int32 MinSlotNum = FMath::Min(OldSlotsNum, NewSlotsNum);
 
 		// as long as we're within the boundaries of both arrays, compare and replace elements
-		for (int32 SlotIndex = 0; SlotIndex < MinSlotNum; SlotIndex++)
+		for (int32 SlotIndex = MinSlotNum - 1; SlotIndex >= 0; --SlotIndex)
 		{
 			if (OldSlots[SlotIndex] != NewSlots[SlotIndex].Key)
 			{
