MVVM: Modification to support extension on the view instance. The class/shared version has the option to create isntance data. The class extension receives that instance and can ack on it. Modified the Panel to create an property of the instance.

#virtualized

diff --git a/Engine/Plugins/Runtime/ModelViewViewModel/Source/ModelViewViewModel/Private/Extensions/MVVMViewListViewBaseExtension.cpp b/Engine/Plugins/Runtime/ModelViewViewModel/Source/ModelViewViewModel/Private/Extensions/MVVMViewListViewBaseExtension.cpp
--- a/Engine/Plugins/Runtime/ModelViewViewModel/Source/ModelViewViewModel/Private/Extensions/MVVMViewListViewBaseExtension.cpp
+++ b/Engine/Plugins/Runtime/ModelViewViewModel/Source/ModelViewViewModel/Private/Extensions/MVVMViewListViewBaseExtension.cpp
@@ -16,7 +16,7 @@
 #define LOCTEXT_NAMESPACE "MVVMViewListViewBaseExtension"
 
 #if WITH_EDITOR
-void UMVVMViewListViewBaseExtension::Initialize(UMVVMViewListViewBaseExtension::FInitListViewBaseExtensionArgs InArgs)
+void UMVVMViewListViewBaseClassExtension::Initialize(UMVVMViewListViewBaseClassExtension::FInitListViewBaseExtensionArgs InArgs)
 {
 	WidgetName = InArgs.WidgetName;
 	WidgetPath = InArgs.WidgetPath;
@@ -24,7 +24,7 @@ void UMVVMViewListViewBaseExtension::Initialize(UMVVMViewListViewBaseExtension::
 }
 #endif
 
-void UMVVMViewListViewBaseExtension::OnSourcesInitialized(UUserWidget* UserWidget, UMVVMView* View)
+void UMVVMViewListViewBaseClassExtension::OnSourcesInitialized(UUserWidget* UserWidget, UMVVMView* View, UMVVMViewExtension* Extension)
 {
 	check(View->GetViewClass());
 	TValueOrError<UE::MVVM::FFieldContext, void> FieldPathResult = View->GetViewClass()->GetBindingLibrary().EvaluateFieldPath(UserWidget, WidgetPath);
@@ -36,8 +36,8 @@ void UMVVMViewListViewBaseExtension::OnSourcesInitialized(UUserWidget* UserWidge
 		{
 			if (UListViewBase* ListWidget = Cast<UListViewBase>(ObjectResult.GetValue()))
 			{
-				CachedListViewWidget = ListWidget;
-				ListWidget->OnEntryWidgetGenerated().AddUObject(this, &UMVVMViewListViewBaseExtension::HandleSetViewModelOnEntryWidget, UserWidget, ListWidget, EntryViewModelName);
+				CachedListViewWidgets.Add(FObjectKey(View), ListWidget);
+				ListWidget->OnEntryWidgetGenerated().AddUObject(this, &UMVVMViewListViewBaseClassExtension::HandleSetViewModelOnEntryWidget, UserWidget, ListWidget, EntryViewModelName);
 			}
 			else
 			{
@@ -64,15 +64,19 @@ void UMVVMViewListViewBaseExtension::OnSourcesInitialized(UUserWidget* UserWidge
 	}
 }
 
-void UMVVMViewListViewBaseExtension::OnSourcesUninitialized(UUserWidget* UserWidget, UMVVMView* View) 
+void UMVVMViewListViewBaseClassExtension::OnSourcesUninitialized(UUserWidget* UserWidget, UMVVMView* View, UMVVMViewExtension* Extension)
 {
-	if (UListViewBase* ListWidget = CachedListViewWidget.Get())
+	TWeakObjectPtr<UListViewBase> ListWidget;
+	if (CachedListViewWidgets.RemoveAndCopyValue(FObjectKey(View), ListWidget))
 	{
-		ListWidget->OnEntryWidgetGenerated().RemoveAll(this);
+		if (UListViewBase* ListWidgetPtr = ListWidget.Get())
+		{
+			ListWidgetPtr->OnEntryWidgetGenerated().RemoveAll(this);
+		}
 	}
 }
 
-void UMVVMViewListViewBaseExtension::HandleSetViewModelOnEntryWidget(UUserWidget& EntryWidget, UUserWidget* OwningUserWidget, UListViewBase* ListWidget, FName EntryViewmodelName)
+void UMVVMViewListViewBaseClassExtension::HandleSetViewModelOnEntryWidget(UUserWidget& EntryWidget, UUserWidget* OwningUserWidget, UListViewBase* ListWidget, FName EntryViewmodelName)
 {
 	if (UMVVMView* View = UMVVMSubsystem::GetViewFromUserWidget(&EntryWidget))
 	{
diff --git a/Engine/Plugins/Runtime/ModelViewViewModel/Source/ModelViewViewModel/Private/Extensions/MVVMViewPanelWidgetExtension.cpp b/Engine/Plugins/Runtime/ModelViewViewModel/Source/ModelViewViewModel/Private/Extensions/MVVMViewPanelWidgetExtension.cpp
--- a/Engine/Plugins/Runtime/ModelViewViewModel/Source/ModelViewViewModel/Private/Extensions/MVVMViewPanelWidgetExtension.cpp
+++ b/Engine/Plugins/Runtime/ModelViewViewModel/Source/ModelViewViewModel/Private/Extensions/MVVMViewPanelWidgetExtension.cpp
@@ -15,146 +15,84 @@
 
 #define LOCTEXT_NAMESPACE "MVVMViewPanelWidgetExtension"
 
-#if WITH_EDITOR
-void UMVVMViewPanelWidgetExtension::Initialize(UMVVMViewPanelWidgetExtension::FInitPanelWidgetExtensionArgs InArgs)
+void UMVVMViewPanelWidgetExtension::Initialize(UMVVMViewPanelWidgetClassExtension* InClassExtension, UPanelWidget* InPanelWidget)
 {
-	WidgetName = InArgs.WidgetName;
-	WidgetPath = InArgs.WidgetPath;
-	EntryViewModelName = InArgs.EntryViewModelName;
-	EntryWidgetClass = InArgs.EntryWidgetClass;
-	SlotTemplate = InArgs.SlotTemplate;
-	PanelPropertyName = InArgs.PanelPropertyName;
-	EntryViewModelClass = InArgs.EntryViewModelClass;
-}
-#endif
-
-void UMVVMViewPanelWidgetExtension::OnViewConstructed(UUserWidget* UserWidget, UMVVMView* View)
-{
-	check(View->GetViewClass());
-	CachedOwningUserWidget = UserWidget;
-	// Set the extension object on the runtime user widget.
-	FObjectPropertyBase* FoundObjectProperty = FindFProperty<FObjectPropertyBase>(UserWidget->GetClass(), PanelPropertyName);
-	if (ensureAlwaysMsgf(FoundObjectProperty, TEXT("The compiler should have added the property")))
-	{
-		FoundObjectProperty->SetObjectPropertyValue_InContainer(UserWidget, this);
-	}
-
-	// Fetch and cache the panel widget
-	TValueOrError<UE::MVVM::FFieldContext, void> FieldPathResult = View->GetViewClass()->GetBindingLibrary().EvaluateFieldPath(UserWidget, WidgetPath);
-
-	if (FieldPathResult.HasValue())
-	{
-		TValueOrError<UObject*, void> ObjectResult = UE::MVVM::FieldPathHelper::EvaluateObjectProperty(FieldPathResult.GetValue());
-		if (ObjectResult.HasValue() && ObjectResult.GetValue() != nullptr)
-		{
-			if (UPanelWidget* PanelWidget = Cast<UPanelWidget>(ObjectResult.GetValue()))
-			{
-				CachedPanelWidget = PanelWidget;
-			}
-			else
-			{
-				UE::MVVM::FMessageLog Log(UserWidget);
-				Log.Error(FText::Format(LOCTEXT("BindToEntryGenerationFailWidgetNotPanel", "The object property '{0}' is not of type panel widget, but has an Viewmodel extension meant for panel widgets. The extension won't have any effects.")
-					, FText::FromName(ObjectResult.GetValue()->GetFName())
-				));
-			}
-		}
-		else
-		{
-			UE::MVVM::FMessageLog Log(UserWidget);
-			Log.Error(FText::Format(LOCTEXT("BindToEntryGenerationFailInvalidObjectPropertyWidget", "The property object for panel-type widget '{0}' is not found, so viewmodels won't be bound to its entries.")
-				, FText::FromName(WidgetName)
-			));
-		}
-	}
-	else
-	{
-		UE::MVVM::FMessageLog Log(UserWidget);
-		Log.Error(FText::Format(LOCTEXT("BindToEntryGenerationFailInvalidFieldPathWidget", "The field path for panel-type widget '{0}' is invalid, so viewmodels won't be bound to its entries.")
-			, FText::FromName(WidgetName)
-		));
-	}
+	ClassExtension = InClassExtension;
+	PanelWidget = InPanelWidget;
 }
 
-void UMVVMViewPanelWidgetExtension::OnViewDestructed(UUserWidget* UserWidget, UMVVMView* View)
+void UMVVMViewPanelWidgetExtension::BP_SetItems(const TArray<UObject*>& InItems)
 {
-	FObjectPropertyBase* FoundObjectProperty = FindFProperty<FObjectPropertyBase>(UserWidget->GetClass(), PanelPropertyName);
-	if (ensureAlwaysMsgf(FoundObjectProperty, TEXT("The compiler should have added the property")))
+	if (!PanelWidget || !ClassExtension)
 	{
-		FoundObjectProperty->SetObjectPropertyValue_InContainer(UserWidget, nullptr);
+		return;
 	}
-}
 
-void UMVVMViewPanelWidgetExtension::BP_SetItems(const TArray<UObject*>& InItems)
-{
-	if (UPanelWidget* PanelWidgetPtr = CachedPanelWidget.Get())
+	// Store all the reusable slots in a temporary array so that we don't re-create them.
+	TArray<TTuple<UPanelSlot*, TScriptInterface<INotifyFieldValueChanged>>> PreviousSlots;
+	for (UPanelSlot* Slot : PanelWidget->GetSlots())
 	{
-		// Store all the reusable slots in a temporary array so that we don't re-create them.
-		TArray<TTuple<UPanelSlot*, TScriptInterface<INotifyFieldValueChanged>>> PreviousSlots;
-		for (UPanelSlot* Slot : PanelWidgetPtr->GetSlots())
+		if (UUserWidget* Content = Cast<UUserWidget>(Slot->Content))
 		{
-			if (UUserWidget* Content = Cast<UUserWidget>(Slot->Content))
+			// The class of the content should strictly match the EntryWidgetClass.
+			if (Content->GetClass() == ClassExtension->GetEntryWidgetClass().Get())
 			{
-				// The class of the content should strictly match the EntryWidgetClass.
-				if (Content->GetClass() == EntryWidgetClass.Get())
+				if (UMVVMView* View = UMVVMSubsystem::GetViewFromUserWidget(Content))
 				{
-					if (UMVVMView* View = UMVVMSubsystem::GetViewFromUserWidget(Content))
+					TScriptInterface<INotifyFieldValueChanged> Interface = View->GetViewModel(ClassExtension->GetEntryViewModelName());
+					if (Interface.GetObject() && Interface.GetObject()->GetClass() == ClassExtension->GetEntryViewModelClass())
 					{
-						TScriptInterface<INotifyFieldValueChanged> Interface = View->GetViewModel(EntryViewModelName);
-						if (Interface.GetObject() && Interface.GetObject()->GetClass() == EntryViewModelClass)
-						{
-							PreviousSlots.Emplace(Slot, Interface);
-						}
+						PreviousSlots.Emplace(Slot, Interface);
 					}
 				}
 			}
 		}
+	}
 
-		UClass* SelectedVMClass = EntryViewModelClass.Get();
-		UUserWidget* OwningUserWidget = CachedOwningUserWidget.Get();
+	UClass* SelectedVMClass = ClassExtension->GetEntryViewModelClass();
+	UUserWidget* OwningUserWidget = GetUserWidget();
 
-		TArray<TTuple<UPanelSlot*, UWidget*>> NewSlots;
-		for (int32 ItemIndex = 0; ItemIndex < InItems.Num(); ++ItemIndex)
+	TArray<TTuple<UPanelSlot*, UWidget*>> NewSlots;
+	for (int32 ItemIndex = 0; ItemIndex < InItems.Num(); ++ItemIndex)
+	{
+		if (UObject* Item = InItems[ItemIndex])
 		{
-			if (UObject* Item = InItems[ItemIndex])
-			{
-				const TTuple<UPanelSlot*, TScriptInterface<INotifyFieldValueChanged>>* FoundObject = PreviousSlots.FindByPredicate([Item](const auto& Other)
-					{
-						return Other.Value.GetObject() == Item;
-					});
-
-				if (Item && OwningUserWidget && Item->GetClass() != SelectedVMClass)
+			const TTuple<UPanelSlot*, TScriptInterface<INotifyFieldValueChanged>>* FoundObject = PreviousSlots.FindByPredicate([Item](const auto& Other)
 				{
-					UE::MVVM::FMessageLog Log(OwningUserWidget);
-					Log.Warning(FText::Format(LOCTEXT("SetPanelWidgetItemsViewmodelTypeMismatch", "The item {0} passed as an entry of widget {1} is not a viewmodel of the selected type {2}.")
-						, FText::FromString(Item->GetName()), FText::FromString(PanelWidgetPtr->GetName()), FText::FromString(SelectedVMClass->GetName())
-					));
-				}
+					return Other.Value.GetObject() == Item;
+				});
 
-				if (!FoundObject)
-				{
-					UUserWidget* EntryWidget = UUserWidget::CreateWidgetInstance(*PanelWidgetPtr, EntryWidgetClass, NAME_None);
-					ensure(EntryWidget);
+			if (Item && OwningUserWidget && Item->GetClass() != SelectedVMClass)
+			{
+				UE::MVVM::FMessageLog Log(OwningUserWidget);
+				Log.Warning(FText::Format(LOCTEXT("SetPanelWidgetItemsViewmodelTypeMismatch", "The item {0} passed as an entry of widget {1} is not a viewmodel of the selected type {2}.")
+					, FText::FromString(Item->GetName()), FText::FromString(PanelWidget->GetName()), FText::FromString(SelectedVMClass->GetName())
+				));
+			}
 
-					SetViewModelOnEntryWidget(EntryWidget, Item, OwningUserWidget);
-					NewSlots.Add(TTuple<UPanelSlot*, UWidget*>(SlotTemplate, EntryWidget));
-				}
-				else
-				{
-					NewSlots.Add(TTuple<UPanelSlot*, UWidget*>(FoundObject->Key, FoundObject->Key->Content));
-				}
+			if (!FoundObject)
+			{
+				UUserWidget* EntryWidget = UUserWidget::CreateWidgetInstance(*PanelWidget, ClassExtension->GetEntryWidgetClass(), NAME_None);
+				ensure(EntryWidget);
+
+				SetViewModelOnEntryWidget(EntryWidget, Item, OwningUserWidget);
+				NewSlots.Add(TTuple<UPanelSlot*, UWidget*>(ClassExtension->GetSlotTemplate(), EntryWidget));
 			}
 			else
 			{
-				UE::MVVM::FMessageLog Log(OwningUserWidget);
-				Log.Warning(FText::Format(LOCTEXT("SetPanelWidgetItemsViewmodelNullObject", "The item at index {0} passed as an entry of widget {1} is null. An entry widget won't be generated for this item.")
-					, ItemIndex, FText::FromString(PanelWidgetPtr->GetName())
-				));
+				NewSlots.Add(TTuple<UPanelSlot*, UWidget*>(FoundObject->Key, FoundObject->Key->Content));
 			}
 		}
-
-		ReplaceAllSlots(NewSlots);
+		else
+		{
+			UE::MVVM::FMessageLog Log(OwningUserWidget);
+			Log.Warning(FText::Format(LOCTEXT("SetPanelWidgetItemsViewmodelNullObject", "The item at index {0} passed as an entry of widget {1} is null. An entry widget won't be generated for this item.")
+				, ItemIndex, FText::FromString(PanelWidget->GetName())
+			));
+		}
 	}
+
+	ReplaceAllSlots(NewSlots);
 }
 
 void UMVVMViewPanelWidgetExtension::SetViewModelOnEntryWidget(UUserWidget* EntryWidget, UObject* ViewModelObject, UUserWidget* OwningUserWidget)
@@ -163,7 +101,7 @@ void UMVVMViewPanelWidgetExtension::SetViewModelOnEntryWidget(UUserWidget* Entry
 	{
 		if (ViewModelObject->Implements<UNotifyFieldValueChanged>())
 		{
-			View->SetViewModel(EntryViewModelName, ViewModelObject);
+			View->SetViewModel(ClassExtension->GetEntryViewModelName(), ViewModelObject);
 		}
 		else
 		{
@@ -171,7 +109,7 @@ void UMVVMViewPanelWidgetExtension::SetViewModelOnEntryWidget(UUserWidget* Entry
 			{
 				UE::MVVM::FMessageLog Log(OwningUserWidget);
 				Log.Error(FText::Format(LOCTEXT("SetViewModelOnEntryWidgetFailNotViewModel", "Trying to set an object that is not a viewmodel on entries of panel-type widget '{0}'. If you do not wish to set viewmodels on the entries of this widget, please remove the corresonding Viewmodel extension from it.")
-					, FText::FromName(WidgetName)
+					, FText::FromName(ClassExtension->GetWidgetName())
 				));
 			}
 		}
@@ -180,39 +118,122 @@ void UMVVMViewPanelWidgetExtension::SetViewModelOnEntryWidget(UUserWidget* Entry
 
 void UMVVMViewPanelWidgetExtension::ReplaceAllSlots(TArrayView<TTuple<UPanelSlot*, UWidget*>> NewSlots)
 {
-	if (UPanelWidget* PanelWidgetPtr = CachedPanelWidget.Get())
+	if (!PanelWidget || !ClassExtension)
 	{
-		const TArray<UPanelSlot*>& OldSlots = PanelWidgetPtr->GetSlots();
-		const int32 OldSlotsNum = OldSlots.Num();
-		const int32 NewSlotsNum = NewSlots.Num();
-		const int32 MinSlotNum = FMath::Min(OldSlotsNum, NewSlotsNum);
+		return;
+	}
 
-		// as long as we're within the boundaries of both arrays, compare and replace elements
-		for (int32 SlotIndex = MinSlotNum - 1; SlotIndex >= 0; --SlotIndex)
+	const TArray<UPanelSlot*>& OldSlots = PanelWidget->GetSlots();
+	const int32 OldSlotsNum = OldSlots.Num();
+	const int32 NewSlotsNum = NewSlots.Num();
+	const int32 MinSlotNum = FMath::Min(OldSlotsNum, NewSlotsNum);
+
+	// as long as we're within the boundaries of both arrays, compare and replace elements
+	for (int32 SlotIndex = MinSlotNum-1; SlotIndex >= 0; --SlotIndex)
+	{
+		if (OldSlots[SlotIndex] != NewSlots[SlotIndex].Key)
 		{
-			if (OldSlots[SlotIndex] != NewSlots[SlotIndex].Key)
-			{
-				PanelWidgetPtr->RemoveChildAt(SlotIndex);
-				UPanelSlot* NewSlot = PanelWidgetPtr->InsertChildAt(SlotIndex, NewSlots[SlotIndex].Value, NewSlots[SlotIndex].Key);
-			}
+			PanelWidget->RemoveChildAt(SlotIndex);
+			UPanelSlot* NewSlot = PanelWidget->InsertChildAt(SlotIndex, NewSlots[SlotIndex].Value, NewSlots[SlotIndex].Key);
 		}
+	}
 
-		// If we have more old slots than new ones, remove all the extra ones.
-		if (OldSlotsNum > NewSlotsNum)
+	// If we have more old slots than new ones, remove all the extra ones.
+	if (OldSlotsNum > NewSlotsNum)
+	{
+		for (int32 SlotIndex = OldSlotsNum - 1; SlotIndex >= NewSlotsNum; SlotIndex--)
 		{
-			for (int32 SlotIndex = OldSlotsNum - 1; SlotIndex >= NewSlotsNum; SlotIndex--)
-			{
-				PanelWidgetPtr->RemoveChildAt(SlotIndex);
-			}
+			PanelWidget->RemoveChildAt(SlotIndex);
 		}
-		// If we have more new slots than old ones, simply append the new ones.
-		else if (NewSlotsNum > OldSlotsNum)
+	}
+	// If we have more new slots than old ones, simply append the new ones.
+	else if (NewSlotsNum > OldSlotsNum)
+	{
+		for (int32 SlotIndex = OldSlotsNum; SlotIndex < NewSlotsNum; SlotIndex++)
 		{
-			for (int32 SlotIndex = OldSlotsNum; SlotIndex < NewSlotsNum; SlotIndex++)
+			UPanelSlot* NewSlot = PanelWidget->AddChild(NewSlots[SlotIndex].Value, NewSlots[SlotIndex].Key);
+		}
+	}
+}
+
+UUserWidget* UMVVMViewPanelWidgetExtension::GetUserWidget() const
+{
+	return GetView()->GetOuterUUserWidget();
+}
+
+#if WITH_EDITOR
+void UMVVMViewPanelWidgetClassExtension::Initialize(UMVVMViewPanelWidgetClassExtension::FInitPanelWidgetExtensionArgs InArgs)
+{
+	WidgetName = InArgs.WidgetName;
+	WidgetPath = InArgs.WidgetPath;
+	EntryViewModelName = InArgs.EntryViewModelName;
+	EntryWidgetClass = InArgs.EntryWidgetClass;
+	SlotTemplate = InArgs.SlotTemplate;
+	PanelPropertyName = InArgs.PanelPropertyName;
+	EntryViewModelClass = InArgs.EntryViewModelClass;
+}
+#endif
+
+UMVVMViewExtension* UMVVMViewPanelWidgetClassExtension::ViewConstructed(UUserWidget* UserWidget, UMVVMView* View)
+{
+	check(View->GetViewClass());
+
+	UMVVMViewPanelWidgetExtension* ResultExtension = nullptr;
+
+	// Fetch and cache the panel widget
+	TValueOrError<UE::MVVM::FFieldContext, void> FieldPathResult = View->GetViewClass()->GetBindingLibrary().EvaluateFieldPath(UserWidget, WidgetPath);
+
+	if (FieldPathResult.HasValue())
+	{
+		TValueOrError<UObject*, void> ObjectResult = UE::MVVM::FieldPathHelper::EvaluateObjectProperty(FieldPathResult.GetValue());
+		if (ObjectResult.HasValue() && ObjectResult.GetValue() != nullptr)
+		{
+			if (UPanelWidget* PanelWidget = Cast<UPanelWidget>(ObjectResult.GetValue()))
 			{
-				UPanelSlot* NewSlot = PanelWidgetPtr->AddChild(NewSlots[SlotIndex].Value, NewSlots[SlotIndex].Key);
+				ResultExtension = NewObject<UMVVMViewPanelWidgetExtension>(View);
+				ResultExtension->Initialize(this, PanelWidget);
+
+				// Set the extension object on the runtime user widget.
+				FObjectPropertyBase* FoundPanelObjectProperty = FindFProperty<FObjectPropertyBase>(UserWidget->GetClass(), PanelPropertyName);
+				bool bValidPanelObject = FoundPanelObjectProperty && FoundPanelObjectProperty->PropertyClass->IsChildOf(UMVVMViewPanelWidgetExtension::StaticClass());
+				if (ensureAlwaysMsgf(bValidPanelObject, TEXT("The compiler should have added the property")))
+				{
+					FoundPanelObjectProperty->SetObjectPropertyValue_InContainer(UserWidget, ResultExtension);
+				}
+			}
+			else
+			{
+				UE::MVVM::FMessageLog Log(UserWidget);
+				Log.Error(FText::Format(LOCTEXT("BindToEntryGenerationFailWidgetNotPanel", "The object property '{0}' is not of type panel widget, but has an Viewmodel extension meant for panel widgets. The extension won't have any effects.")
+					, FText::FromName(ObjectResult.GetValue()->GetFName())
+				));
 			}
 		}
+		else
+		{
+			UE::MVVM::FMessageLog Log(UserWidget);
+			Log.Error(FText::Format(LOCTEXT("BindToEntryGenerationFailInvalidObjectPropertyWidget", "The property object for panel-type widget '{0}' is not found, so viewmodels won't be bound to its entries.")
+				, FText::FromName(WidgetName)
+			));
+		}
+	}
+	else
+	{
+		UE::MVVM::FMessageLog Log(UserWidget);
+		Log.Error(FText::Format(LOCTEXT("BindToEntryGenerationFailInvalidFieldPathWidget", "The field path for panel-type widget '{0}' is invalid, so viewmodels won't be bound to its entries.")
+			, FText::FromName(WidgetName)
+		));
+	}
+
+	return ResultExtension;
+}
+
+void UMVVMViewPanelWidgetClassExtension::OnViewDestructed(UUserWidget* UserWidget, UMVVMView* View, UMVVMViewExtension* Extension)
+{
+	FObjectPropertyBase* FoundPanelObjectProperty = FindFProperty<FObjectPropertyBase>(UserWidget->GetClass(), PanelPropertyName);
+	if (ensureAlwaysMsgf(FoundPanelObjectProperty, TEXT("The compiler should have added the property")))
+	{
+		FoundPanelObjectProperty->SetObjectPropertyValue_InContainer(UserWidget, nullptr);
 	}
 }
 
diff --git a/Engine/Plugins/Runtime/ModelViewViewModel/Source/ModelViewViewModel/Private/View/MVVMView.cpp b/Engine/Plugins/Runtime/ModelViewViewModel/Source/ModelViewViewModel/Private/View/MVVMView.cpp
--- a/Engine/Plugins/Runtime/ModelViewViewModel/Source/ModelViewViewModel/Private/View/MVVMView.cpp
+++ b/Engine/Plugins/Runtime/ModelViewViewModel/Source/ModelViewViewModel/Private/View/MVVMView.cpp
@@ -76,20 +76,11 @@ void UMVVMView::Construct()
 	check(bConstructed == false);
 
 	// Create an independent copy of extensions per view.
-	ensure(Extensions.IsEmpty());
+	check(Extensions.IsEmpty());
 	Extensions.Reset(GeneratedViewClass->GetViewClassExtensions().Num());
 	for (UMVVMViewClassExtension* Extension : GeneratedViewClass->GetViewClassExtensions())
 	{
-		UMVVMViewClassExtension* ViewExtension = DuplicateObject<UMVVMViewClassExtension>(Extension, this);
-		if (ensure(ViewExtension))
-		{
-			Extensions.Add(ViewExtension);
-		}
-	}
-
-	for (UMVVMViewClassExtension* Extension : Extensions)
-	{
-		Extension->OnViewConstructed(GetUserWidget(), this);
+		Extensions.Add(Extension->ViewConstructed(GetUserWidget(), this));
 	}
 
 	if (GeneratedViewClass->DoesInitializeSourcesOnConstruct())
@@ -122,9 +113,11 @@ void UMVVMView::Destruct()
 	UninitializeEvents();
 	UninitializeSources(); // and bindings
 
-	for (UMVVMViewClassExtension* Extension : Extensions)
+	const TArrayView<const TObjectPtr<UMVVMViewClassExtension>> ClassExtensions = GeneratedViewClass->GetViewClassExtensions();
+	check(ClassExtensions.Num() == Extensions.Num());
+	for (int32 Index = 0; Index < ClassExtensions.Num(); ++Index)
 	{
-		Extension->OnViewDestructed(GetUserWidget(), this);
+		ClassExtensions[Index]->OnViewDestructed(GetUserWidget(), this, Extensions[Index]);
 	}
 }
 
@@ -170,16 +163,18 @@ void UMVVMView::InitializeSources()
 		InitializeSourceBindingsCommon();
 	}
 
-	for (UMVVMViewClassExtension* Extension : Extensions)
+	const TArrayView<const TObjectPtr<UMVVMViewClassExtension>> ClassExtensions = GeneratedViewClass->GetViewClassExtensions();
+	check(ClassExtensions.Num() == Extensions.Num());
+	for (int32 Index = 0; Index < ClassExtensions.Num(); ++Index)
 	{
-		Extension->OnSourcesInitialized(GetUserWidget(), this);
+		ClassExtensions[Index]->OnSourcesInitialized(GetUserWidget(), this, Extensions[Index]);
 	}
 
 	if (GeneratedViewClass->DoesInitializeBindingsOnConstruct())
 	{
-		for (UMVVMViewClassExtension* Extension : Extensions)
+		for (int32 Index = 0; Index < ClassExtensions.Num(); ++Index)
 		{
-			Extension->OnBindingsInitialized(GetUserWidget(), this);
+			ClassExtensions[Index]->OnBindingsInitialized(GetUserWidget(), this, Extensions[Index]);
 		}
 	}
 
@@ -201,9 +196,11 @@ void UMVVMView::UninitializeSources()
 		UninitializeBindings();
 	}
 
-	for (UMVVMViewClassExtension* Extension : Extensions)
+	const TArrayView<const TObjectPtr<UMVVMViewClassExtension>> ClassExtensions = GeneratedViewClass->GetViewClassExtensions();
+	check(ClassExtensions.Num() == Extensions.Num());
+	for (int32 Index = 0; Index < ClassExtensions.Num(); ++Index)
 	{
-		Extension->OnSourcesUninitialized(GetUserWidget(), this);
+		ClassExtensions[Index]->OnSourcesUninitialized(GetUserWidget(), this, Extensions[Index]);
 	}
 
 	if (GeneratedViewClass->DoesListenToViewModelCollectionChanged())
@@ -334,9 +331,11 @@ void UMVVMView::InitializeBindings()
 
 	InitializeSourceBindingsCommon();
 
-	for (UMVVMViewClassExtension* Extension : Extensions)
+	const TArrayView<const TObjectPtr<UMVVMViewClassExtension>> ClassExtensions = GeneratedViewClass->GetViewClassExtensions();
+	check(ClassExtensions.Num() == Extensions.Num());
+	for (int32 Index = 0; Index < ClassExtensions.Num(); ++Index)
 	{
-		Extension->OnBindingsInitialized(GetUserWidget(), this);
+		ClassExtensions[Index]->OnBindingsInitialized(GetUserWidget(), this, Extensions[Index]);
 	}
 }
 
@@ -371,9 +370,11 @@ void UMVVMView::UninitializeBindings()
 
 	bBindingsInitialized = false;
 
-	for (UMVVMViewClassExtension* Extension : Extensions)
+	const TArrayView<const TObjectPtr<UMVVMViewClassExtension>> ClassExtensions = GeneratedViewClass->GetViewClassExtensions();
+	check(ClassExtensions.Num() == Extensions.Num());
+	for (int32 Index = 0; Index < ClassExtensions.Num(); ++Index)
 	{
-		Extension->OnBindingsUninitialized(GetUserWidget(), this);
+		ClassExtensions[Index]->OnBindingsUninitialized(GetUserWidget(), this, Extensions[Index]);
 	}
 
 	const TArrayView<const FMVVMViewClass_Source> ClassSources = GeneratedViewClass->GetSources();
@@ -1027,9 +1028,11 @@ void UMVVMView::InitializeEvents()
 		BindEvent(ClassEvent, FMVVMViewClass_EventKey(Index));
 	}
 
-	for (UMVVMViewClassExtension* Extension : Extensions)
+	const TArrayView<const TObjectPtr<UMVVMViewClassExtension>> ClassExtensions = GeneratedViewClass->GetViewClassExtensions();
+	check(ClassExtensions.Num() == Extensions.Num());
+	for (int32 Index = 0; Index < ClassExtensions.Num(); ++Index)
 	{
-		Extension->OnEventsInitialized(GetUserWidget(), this);
+		ClassExtensions[Index]->OnEventsInitialized(GetUserWidget(), this, Extensions[Index]);
 	}
 
 	bEventsInitialized = true;
@@ -1047,9 +1050,11 @@ void UMVVMView::UninitializeEvents()
 	
 	bEventsInitialized = false;
 
-	for (UMVVMViewClassExtension* Extension : Extensions)
+	const TArrayView<const TObjectPtr<UMVVMViewClassExtension>> ClassExtensions = GeneratedViewClass->GetViewClassExtensions();
+	check(ClassExtensions.Num() == Extensions.Num());
+	for (int32 Index = 0; Index < ClassExtensions.Num(); ++Index)
 	{
-		Extension->OnEventsUninitialized(GetUserWidget(), this);
+		ClassExtensions[Index]->OnEventsUninitialized(GetUserWidget(), this, Extensions[Index]);
 	}
 
 	for (int32 Index = BoundEvents.Num() - 1; Index >= 0; --Index)
diff --git a/Engine/Plugins/Runtime/ModelViewViewModel/Source/ModelViewViewModel/Public/Extensions/MVVMViewClassExtension.h b/Engine/Plugins/Runtime/ModelViewViewModel/Source/ModelViewViewModel/Public/Extensions/MVVMViewClassExtension.h
--- a/Engine/Plugins/Runtime/ModelViewViewModel/Source/ModelViewViewModel/Public/Extensions/MVVMViewClassExtension.h
+++ b/Engine/Plugins/Runtime/ModelViewViewModel/Source/ModelViewViewModel/Public/Extensions/MVVMViewClassExtension.h
@@ -7,11 +7,28 @@
 #include "MVVMViewClassExtension.generated.h"
 
 class UMVVMView;
-class UMVVMBlueprintViewExtension;
 class UUserWidget;
 
+class UMVVMViewExtension;
+class UMVVMViewClassExtension;
+
+/**
+ * A runtime extension class instance. This information is per instance of a UMVVMViewClassExtension
+ */
+UCLASS(MinimalAPI, Within=MVVMView)
+class UMVVMViewExtension : public UObject
+{
+	GENERATED_BODY()
+
+protected:
+	UMVVMView* GetView() const
+	{
+		return GetOuterUMVVMView();
+	}
+};
+
 /**
- * A runtime extension class to define MVVM-related properties and behaviour. This information comes from the
+ * A runtime extension class to define MVVM-related properties and behavior. This information comes from the
  * corresponding UMVVMBlueprintViewExtension class. This class provides a hook into the MVVM runtime initializations.
  */
 UCLASS(MinimalAPI)
@@ -21,12 +38,17 @@ class UMVVMViewClassExtension : public UObject
 
 public:
 	//~ Functions to be overriden in a user-defined UMVVMViewMyWidgetExtension class
-	virtual void OnViewConstructed(UUserWidget* UserWidget, UMVVMView* View) {};
-	virtual void OnSourcesInitialized(UUserWidget* UserWidget, UMVVMView* View) {};
-	virtual void OnBindingsInitialized(UUserWidget* UserWidget, UMVVMView* View) {};
-	virtual void OnEventsInitialized(UUserWidget* UserWidget, UMVVMView* View) {};
-	virtual void OnEventsUninitialized(UUserWidget* UserWidget, UMVVMView* View) {};
-	virtual void OnBindingsUninitialized(UUserWidget* UserWidget, UMVVMView* View) {};
-	virtual void OnSourcesUninitialized(UUserWidget* UserWidget, UMVVMView* View) {};
-	virtual void OnViewDestructed(UUserWidget* UserWidget, UMVVMView* View) {};
-};
\ No newline at end of file
+
+	/** When the view is constructed. The class extension can create a view instance if needed or return nullptr. */
+	virtual UMVVMViewExtension* ViewConstructed(UUserWidget* UserWidget, UMVVMView* View)
+	{
+		return nullptr;
+	}
+	virtual void OnSourcesInitialized(UUserWidget* UserWidget, UMVVMView* View, UMVVMViewExtension* Extension) {};
+	virtual void OnBindingsInitialized(UUserWidget* UserWidget, UMVVMView* View, UMVVMViewExtension* Extension) {};
+	virtual void OnEventsInitialized(UUserWidget* UserWidget, UMVVMView* View, UMVVMViewExtension* Extension) {};
+	virtual void OnEventsUninitialized(UUserWidget* UserWidget, UMVVMView* View, UMVVMViewExtension* Extension) {};
+	virtual void OnBindingsUninitialized(UUserWidget* UserWidget, UMVVMView* View, UMVVMViewExtension* Extension) {};
+	virtual void OnSourcesUninitialized(UUserWidget* UserWidget, UMVVMView* View, UMVVMViewExtension* Extension) {};
+	virtual void OnViewDestructed(UUserWidget* UserWidget, UMVVMView* View, UMVVMViewExtension* Extension) {};
+};
diff --git a/Engine/Plugins/Runtime/ModelViewViewModel/Source/ModelViewViewModel/Public/Extensions/MVVMViewListViewBaseExtension.h b/Engine/Plugins/Runtime/ModelViewViewModel/Source/ModelViewViewModel/Public/Extensions/MVVMViewListViewBaseExtension.h
--- a/Engine/Plugins/Runtime/ModelViewViewModel/Source/ModelViewViewModel/Public/Extensions/MVVMViewListViewBaseExtension.h
+++ b/Engine/Plugins/Runtime/ModelViewViewModel/Source/ModelViewViewModel/Public/Extensions/MVVMViewListViewBaseExtension.h
@@ -4,6 +4,7 @@
 
 #include "Bindings/MVVMCompiledBindingLibrary.h"
 #include "MVVMViewClassExtension.h"
+#include "UObject/ObjectKey.h"
 #include "MVVMViewListViewBaseExtension.generated.h"
 
 class UListViewBase;
@@ -11,14 +12,14 @@ class UMVVMView;
 class UUserWidget;
 
 UCLASS()
-class MODELVIEWVIEWMODEL_API UMVVMViewListViewBaseExtension : public UMVVMViewClassExtension
+class MODELVIEWVIEWMODEL_API UMVVMViewListViewBaseClassExtension : public UMVVMViewClassExtension
 {
 	GENERATED_BODY()
 
 public:
 	//~ Begin UMVVMViewClassExtension overrides
-	virtual void OnSourcesInitialized(UUserWidget* UserWidget, UMVVMView* View) override;
-	virtual void OnSourcesUninitialized(UUserWidget* UserWidget, UMVVMView* View) override;
+	virtual void OnSourcesInitialized(UUserWidget* UserWidget, UMVVMView* View, UMVVMViewExtension* Extension) override;
+	virtual void OnSourcesUninitialized(UUserWidget* UserWidget, UMVVMView* View, UMVVMViewExtension* Extension) override;
 	//~ End UMVVMViewClassExtension overrides
 
 #if WITH_EDITOR
@@ -56,6 +57,5 @@ private:
 	UPROPERTY()
 	FMVVMVCompiledFieldPath WidgetPath;
 
-	UPROPERTY(Transient)
-	TWeakObjectPtr<UListViewBase> CachedListViewWidget;
+	TMap<FObjectKey, TWeakObjectPtr<UListViewBase>> CachedListViewWidgets;
 };
\ No newline at end of file
diff --git a/Engine/Plugins/Runtime/ModelViewViewModel/Source/ModelViewViewModel/Public/Extensions/MVVMViewPanelWidgetExtension.h b/Engine/Plugins/Runtime/ModelViewViewModel/Source/ModelViewViewModel/Public/Extensions/MVVMViewPanelWidgetExtension.h
--- a/Engine/Plugins/Runtime/ModelViewViewModel/Source/ModelViewViewModel/Public/Extensions/MVVMViewPanelWidgetExtension.h
+++ b/Engine/Plugins/Runtime/ModelViewViewModel/Source/ModelViewViewModel/Public/Extensions/MVVMViewPanelWidgetExtension.h
@@ -12,14 +12,38 @@ class UPanelWidget;
 class UUserWidget;
 
 UCLASS()
-class MODELVIEWVIEWMODEL_API UMVVMViewPanelWidgetExtension : public UMVVMViewClassExtension
+class MODELVIEWVIEWMODEL_API UMVVMViewPanelWidgetExtension : public UMVVMViewExtension
+{
+	GENERATED_BODY()
+
+public:
+	void Initialize(UMVVMViewPanelWidgetClassExtension* ClassExtension, UPanelWidget* PanelWidget);
+
+	UFUNCTION(BlueprintCallable, Category = PanelWidget, meta = (AllowPrivateAccess = true, DisplayName = "Set Items", ViewmodelBlueprintWidgetExtension = "EntryViewModel"))
+	void BP_SetItems(const TArray<UObject*>& InItems);
+
+private:
+	void SetViewModelOnEntryWidget(UUserWidget* EntryWidget, UObject* ViewModelObject, UUserWidget* OwningUserWidget);
+	void ReplaceAllSlots(TArrayView<TTuple<UPanelSlot*, UWidget*>> NewSlots);
+	UUserWidget* GetUserWidget() const;
+
+private:
+	UPROPERTY(Transient)
+	TObjectPtr<UPanelWidget> PanelWidget;
+
+	UPROPERTY(Transient)
+	TObjectPtr<UMVVMViewPanelWidgetClassExtension> ClassExtension;
+};
+
+UCLASS()
+class MODELVIEWVIEWMODEL_API UMVVMViewPanelWidgetClassExtension : public UMVVMViewClassExtension
 {
 	GENERATED_BODY()
 
 public:
 	//~ Begin UMVVMViewClassExtension overrides
-	virtual void OnViewConstructed(UUserWidget* UserWidget, UMVVMView* View) override;
-	virtual void OnViewDestructed(UUserWidget* UserWidget, UMVVMView* View) override;
+	virtual UMVVMViewExtension* ViewConstructed(UUserWidget* UserWidget, UMVVMView* View) override;
+	virtual void OnViewDestructed(UUserWidget* UserWidget, UMVVMView* View, UMVVMViewExtension* Extension) override;
 	//~ End UMVVMViewClassExtension overrides
 
 #if WITH_EDITOR
@@ -48,17 +72,31 @@ public:
 
 #endif
 
+public:
+	FName GetWidgetName() const
+	{ 
+		return WidgetName;
+	}
+	
 	FName GetEntryViewModelName() const
 	{ 
 		return EntryViewModelName; 
 	}
 	
-	UFUNCTION(BlueprintCallable, Category = PanelWidget, meta = (AllowPrivateAccess = true, DisplayName = "Set Items", ViewmodelBlueprintWidgetExtension = "EntryViewModel"))
-	virtual void BP_SetItems(const TArray<UObject*>& InItems);
+	UPanelSlot* GetSlotTemplate() const
+	{ 
+		return SlotTemplate;
+	}
 
-private:
-	void SetViewModelOnEntryWidget(UUserWidget* EntryWidget, UObject* ViewModelObject, UUserWidget* OwningUserWidget);
-	void ReplaceAllSlots(TArrayView<TTuple<UPanelSlot*, UWidget*>> NewSlots);
+	TSubclassOf<UUserWidget> GetEntryWidgetClass() const
+	{
+		return EntryWidgetClass;
+	}
+
+	UClass* GetEntryViewModelClass() const
+	{
+		return EntryViewModelClass;
+	}
 
 private:
 	UPROPERTY(VisibleAnywhere, Category = "MVVM Extension")
@@ -73,18 +111,12 @@ private:
 	UPROPERTY(VisibleAnywhere, Category = "MVVM Extension")
 	TObjectPtr<UPanelSlot> SlotTemplate;
 
-	UPROPERTY()
+	UPROPERTY(VisibleAnywhere, Category = "MVVM Extension")
 	FName PanelPropertyName;
 
-	UPROPERTY()
-	TObjectPtr<UClass> EntryViewModelClass = nullptr;
+	UPROPERTY(VisibleAnywhere, Category = "MVVM Extension")
+	TObjectPtr<UClass> EntryViewModelClass;
 
 	UPROPERTY()
 	FMVVMVCompiledFieldPath WidgetPath;
-
-	UPROPERTY(Transient)
-	TWeakObjectPtr<UPanelWidget> CachedPanelWidget;
-
-	UPROPERTY(Transient)
-	TWeakObjectPtr<UUserWidget> CachedOwningUserWidget;
 };
\ No newline at end of file
diff --git a/Engine/Plugins/Runtime/ModelViewViewModel/Source/ModelViewViewModel/Public/View/MVVMView.h b/Engine/Plugins/Runtime/ModelViewViewModel/Source/ModelViewViewModel/Public/View/MVVMView.h
--- a/Engine/Plugins/Runtime/ModelViewViewModel/Source/ModelViewViewModel/Public/View/MVVMView.h
+++ b/Engine/Plugins/Runtime/ModelViewViewModel/Source/ModelViewViewModel/Public/View/MVVMView.h
@@ -14,6 +14,7 @@ template <typename InterfaceType> class TScriptInterface;
 
 class UMVVMViewClass;
 class UMVVMViewClassExtension;
+class UMVVMViewExtension;
 struct FMVVMViewClass_Binding;
 struct FMVVMViewClass_Event;
 struct FMVVMViewClass_Source;
@@ -222,7 +223,7 @@ private:
 	TArray<FBoundEvent> BoundEvents;
 
 	UPROPERTY(VisibleAnywhere, Transient, Category = "View")
-	TArray<TObjectPtr<UMVVMViewClassExtension>> Extensions;
+	TArray<TObjectPtr<UMVVMViewExtension>> Extensions;
 
 	/** Bitfield that represents the valid sources. */
 	UPROPERTY(VisibleAnywhere, Transient, Category = "View")
diff --git a/Engine/Plugins/Runtime/ModelViewViewModel/Source/ModelViewViewModelBlueprint/Private/Extensions/MVVMViewBlueprintListViewBaseExtension.cpp b/Engine/Plugins/Runtime/ModelViewViewModel/Source/ModelViewViewModelBlueprint/Private/Extensions/MVVMViewBlueprintListViewBaseExtension.cpp
--- a/Engine/Plugins/Runtime/ModelViewViewModel/Source/ModelViewViewModelBlueprint/Private/Extensions/MVVMViewBlueprintListViewBaseExtension.cpp
+++ b/Engine/Plugins/Runtime/ModelViewViewModel/Source/ModelViewViewModelBlueprint/Private/Extensions/MVVMViewBlueprintListViewBaseExtension.cpp
@@ -178,11 +178,11 @@ void UMVVMViewBlueprintListViewBaseExtension::Compile(UE::MVVM::Compiler::IMVVMB
 						if (const FMVVMBlueprintViewModelContext* ViewModelContext = EntryBPView->FindViewModel(EntryViewModelId))
 						{
 							// Create the corresponding runtime extension
-							UMVVMViewClassExtension* NewExtensionObj = Compiler->CreateViewClassExtension(UMVVMViewListViewBaseExtension::StaticClass());
-							UMVVMViewListViewBaseExtension* NewExtension = CastChecked<UMVVMViewListViewBaseExtension>(NewExtensionObj);
+							UMVVMViewClassExtension* NewExtensionObj = Compiler->CreateViewClassExtension(UMVVMViewListViewBaseClassExtension::StaticClass());
+							UMVVMViewListViewBaseClassExtension* NewExtension = CastChecked<UMVVMViewListViewBaseClassExtension>(NewExtensionObj);
 
 							const FName EntryViewModelName = ViewModelContext->GetViewModelName();
-							NewExtension->Initialize(UMVVMViewListViewBaseExtension::FInitListViewBaseExtensionArgs(WidgetName, EntryViewModelName, CompiledFieldPath.GetValue()));
+							NewExtension->Initialize(UMVVMViewListViewBaseClassExtension::FInitListViewBaseExtensionArgs(WidgetName, EntryViewModelName, CompiledFieldPath.GetValue()));
 						
 						}
 					}
diff --git a/Engine/Plugins/Runtime/ModelViewViewModel/Source/ModelViewViewModelBlueprint/Private/Extensions/MVVMViewBlueprintPanelWidgetExtension.cpp b/Engine/Plugins/Runtime/ModelViewViewModel/Source/ModelViewViewModelBlueprint/Private/Extensions/MVVMViewBlueprintPanelWidgetExtension.cpp
--- a/Engine/Plugins/Runtime/ModelViewViewModel/Source/ModelViewViewModelBlueprint/Private/Extensions/MVVMViewBlueprintPanelWidgetExtension.cpp
+++ b/Engine/Plugins/Runtime/ModelViewViewModel/Source/ModelViewViewModelBlueprint/Private/Extensions/MVVMViewBlueprintPanelWidgetExtension.cpp
@@ -196,11 +196,11 @@ void UMVVMViewBlueprintPanelWidgetExtension::Compile(UE::MVVM::Compiler::IMVVMBl
 						if (const FMVVMBlueprintViewModelContext* ViewModelContext = EntryBPView->FindViewModel(EntryViewModelId))
 						{
 							// Create the corresponding runtime extension
-							UMVVMViewClassExtension* NewExtensionObj = Compiler->CreateViewClassExtension(UMVVMViewPanelWidgetExtension::StaticClass());
-							UMVVMViewPanelWidgetExtension* NewExtension = CastChecked<UMVVMViewPanelWidgetExtension>(NewExtensionObj);
+							UMVVMViewClassExtension* NewExtensionObj = Compiler->CreateViewClassExtension(UMVVMViewPanelWidgetClassExtension::StaticClass());
+							UMVVMViewPanelWidgetClassExtension* NewExtension = CastChecked<UMVVMViewPanelWidgetClassExtension>(NewExtensionObj);
 
 							const FName EntryViewModelName = ViewModelContext->GetViewModelName();
-							NewExtension->Initialize(UMVVMViewPanelWidgetExtension::FInitPanelWidgetExtensionArgs(WidgetName, EntryViewModelName, CompiledFieldPath.GetValue(), EntryWidgetClass, SlotObj, PanelPropertyName, ViewModelContext->NotifyFieldValueClass));
+							NewExtension->Initialize(UMVVMViewPanelWidgetClassExtension::FInitPanelWidgetExtensionArgs(WidgetName, EntryViewModelName, CompiledFieldPath.GetValue(), EntryWidgetClass, SlotObj, PanelPropertyName, ViewModelContext->NotifyFieldValueClass));
 						}
 					}
 				}
diff --git a/Engine/Plugins/Runtime/ModelViewViewModel/Source/ModelViewViewModelBlueprint/Public/Extensions/MVVMViewBlueprintPanelWidgetExtension.h b/Engine/Plugins/Runtime/ModelViewViewModel/Source/ModelViewViewModelBlueprint/Public/Extensions/MVVMViewBlueprintPanelWidgetExtension.h
--- a/Engine/Plugins/Runtime/ModelViewViewModel/Source/ModelViewViewModelBlueprint/Public/Extensions/MVVMViewBlueprintPanelWidgetExtension.h
+++ b/Engine/Plugins/Runtime/ModelViewViewModel/Source/ModelViewViewModelBlueprint/Public/Extensions/MVVMViewBlueprintPanelWidgetExtension.h
@@ -10,7 +10,7 @@
 
 class UMVVMBlueprintView;
 class UMVVMViewClass;
-class UMVVMViewPanelWidgetExtension;
+class UMVVMViewPanelWidgetClassExtension;
 class UUserWidget;
 class UWidgetBlueprintGeneratedClass;
 
@@ -64,7 +64,7 @@ private:
 	FName PanelPropertyName;
 
 	UPROPERTY()
-	TObjectPtr<UMVVMViewPanelWidgetExtension> ExtensionObj;
+	TObjectPtr<UMVVMViewPanelWidgetClassExtension> ExtensionObj;
 
 	UE::MVVM::FCompiledBindingLibraryCompiler::FFieldPathHandle WidgetPathHandle;
 
