MVVM: Add tag and name to handle added by extender.

diff --git a/Engine/Plugins/Runtime/ModelViewViewModel/Source/ModelViewViewModelEditor/Private/Customizations/MVVMPanelWidgetExtensionCustomizationExtender.cpp b/Engine/Plugins/Runtime/ModelViewViewModel/Source/ModelViewViewModelEditor/Private/Customizations/MVVMPanelWidgetExtensionCustomizationExtender.cpp
--- a/Engine/Plugins/Runtime/ModelViewViewModel/Source/ModelViewViewModelEditor/Private/Customizations/MVVMPanelWidgetExtensionCustomizationExtender.cpp
+++ b/Engine/Plugins/Runtime/ModelViewViewModel/Source/ModelViewViewModelEditor/Private/Customizations/MVVMPanelWidgetExtensionCustomizationExtender.cpp
@@ -43,6 +43,7 @@ void FMVVMPanelWidgetExtensionCustomizationExtender::CustomizeDetails(IDetailLay
 				if (GetDefault<UMVVMDeveloperProjectSettings>()->IsExtensionSupportedForPanelClass(Panel->GetClass()))
 				{
 					FName NAME_ViewmodelExtension = "ViewmodelExtension";
+					FName NAME_ViewmodelExtensionSlot = "ViewmodelExtensionSlot";
 					Widget = Panel;
 					WidgetBlueprintEditor = InWidgetBlueprintEditor;
 
@@ -144,12 +145,12 @@ void FMVVMPanelWidgetExtensionCustomizationExtender::CustomizeDetails(IDetailLay
 								.CreateCategoryNodes(false)
 								.AllowChildren(true)
 								.HideRootObjectNode(false)
+								.UniqueId(NAME_ViewmodelExtensionSlot)
 							);
 
 							TSharedPtr<IPropertyHandle> SlotPropertyHandle = SlotDetailRow->GetPropertyHandle();
 
 							SlotDetailRow->CustomWidget(true)
-								.RowTag(NAME_ViewmodelExtension)
 								.NameContent()
 								[
 									SNew(STextBlock)
