MVVM: Fix customization extender that didn't check if the extension was of the correct type.

diff --git a/Engine/Plugins/Runtime/ModelViewViewModel/Source/ModelViewViewModelBlueprint/Private/MVVMWidgetBlueprintExtension_View.cpp b/Engine/Plugins/Runtime/ModelViewViewModel/Source/ModelViewViewModelBlueprint/Private/MVVMWidgetBlueprintExtension_View.cpp
--- a/Engine/Plugins/Runtime/ModelViewViewModel/Source/ModelViewViewModelBlueprint/Private/MVVMWidgetBlueprintExtension_View.cpp
+++ b/Engine/Plugins/Runtime/ModelViewViewModel/Source/ModelViewViewModelBlueprint/Private/MVVMWidgetBlueprintExtension_View.cpp
@@ -93,7 +93,7 @@ void UMVVMWidgetBlueprintExtension_View::RemoveBlueprintWidgetExtension(UMVVMBlu
 
 	const FScopedTransaction Transaction(LOCTEXT("RemoveViewModelExtension", "Remove viewmodel extension"));
 	Modify();
-	BlueprintExtensions.Remove(Extension);
+	BlueprintExtensions.RemoveSingle(Extension);
 	FBlueprintEditorUtils::MarkBlueprintAsStructurallyModified(GetWidgetBlueprint());
 }
 
@@ -102,7 +102,7 @@ TArray<UMVVMBlueprintViewExtension*> UMVVMWidgetBlueprintExtension_View::GetBlue
 	TArray<UMVVMBlueprintViewExtension*> ThisWidgetExtensions;
 	for (const FMVVMExtensionItem& Extension : BlueprintExtensions)
 	{
-		if (Extension.WidgetName == WidgetName)
+		if (Extension.WidgetName == WidgetName && Extension.ExtensionObj)
 		{
 			ThisWidgetExtensions.Add(Extension.ExtensionObj);
 		}
@@ -116,38 +116,57 @@ void UMVVMWidgetBlueprintExtension_View::VerifyWidgetExtensions()
 	{
 		if (const UWidgetTree* WidgetTree = WidgetBlueprint->WidgetTree)
 		{
+			bool bModified = false;
+			auto UpdateModify = [&bModified, Self=this]()
+			{
+				if (!bModified)
+				{
+					bModified = true;
+					Self->Modify();
+				}
+			};
+
 			TArray <FName, TInlineAllocator<4>> WidgetNamesToRemove;
-			for (const FMVVMExtensionItem& Extension : BlueprintExtensions)
+			for (int32 Index = BlueprintExtensions.Num() - 1; Index >= 0; --Index)
 			{
-				if (!Extension.WidgetName.IsNone())
+				if (BlueprintExtensions[Index].ExtensionObj == nullptr)
 				{
-					WidgetNamesToRemove.Add(Extension.WidgetName);
+					UpdateModify();
+					BlueprintExtensions.RemoveAtSwap(Index);
+				}
+				else if (!BlueprintExtensions[Index].WidgetName.IsNone())
+				{
+					WidgetNamesToRemove.Add(BlueprintExtensions[Index].WidgetName);
 				}
 			}
+
 			// Find widgets that are no longer in the tree and delete their extensions.
-			WidgetTree->ForEachWidget([&WidgetNamesToRemove, this](TObjectPtr<UWidget> Widget) {
-				if (Widget)
-				{
-					for (int32 Index = WidgetNamesToRemove.Num() - 1; Index >= 0; Index--)
+			if (WidgetNamesToRemove.Num() > 0)
+			{
+				WidgetTree->ForEachWidget([&WidgetNamesToRemove, this](TObjectPtr<UWidget> Widget) {
+					if (Widget)
 					{
-						const FName WidgetName = WidgetNamesToRemove[Index];
-						if (WidgetName == Widget->GetFName())
+						for (int32 Index = WidgetNamesToRemove.Num() - 1; Index >= 0; Index--)
 						{
-							WidgetNamesToRemove.Remove(WidgetName);
+							const FName WidgetName = WidgetNamesToRemove[Index];
+							if (WidgetName == Widget->GetFName())
+							{
+								WidgetNamesToRemove.RemoveSingleSwap(WidgetName);
+							}
 						}
 					}
-				}
-			});
+					});
+			}
 
 			if (WidgetNamesToRemove.Num() > 0)
 			{
-				Modify();
-			}
-			for (int32 Index = BlueprintExtensions.Num() - 1; Index >= 0; Index--)
-			{
-				if (WidgetNamesToRemove.Contains(BlueprintExtensions[Index].WidgetName))
+				UpdateModify();
+				for (int32 Index = BlueprintExtensions.Num() - 1; Index >= 0; --Index)
 				{
-					BlueprintExtensions.RemoveAt(Index);
+					if (WidgetNamesToRemove.Contains(BlueprintExtensions[Index].WidgetName))
+					{
+						BlueprintExtensions.RemoveAtSwap(Index);
+					}
 				}
 			}
 		}
diff --git a/Engine/Plugins/Runtime/ModelViewViewModel/Source/ModelViewViewModelEditor/Private/Customizations/MVVMListViewBaseExtensionCustomizationExtender.cpp b/Engine/Plugins/Runtime/ModelViewViewModel/Source/ModelViewViewModelEditor/Private/Customizations/MVVMListViewBaseExtensionCustomizationExtender.cpp
--- a/Engine/Plugins/Runtime/ModelViewViewModel/Source/ModelViewViewModelEditor/Private/Customizations/MVVMListViewBaseExtensionCustomizationExtender.cpp
+++ b/Engine/Plugins/Runtime/ModelViewViewModel/Source/ModelViewViewModelEditor/Private/Customizations/MVVMListViewBaseExtensionCustomizationExtender.cpp
@@ -166,7 +166,9 @@ void FMVVMListViewBaseExtensionCustomizationExtender::CreateListBaseViewExtensio
 	{
 		if (UListViewBase* WidgetPtr = Widget.Get())
 		{
-			if (Extension->GetBlueprintExtensionsForWidget(WidgetPtr->GetFName()).IsEmpty())
+			TArray<UMVVMBlueprintViewExtension*> Extensions = Extension->GetBlueprintExtensionsForWidget(WidgetPtr->GetFName());
+			bool bExists = Extensions.ContainsByPredicate([](UMVVMBlueprintViewExtension* Extension) { return Cast<UMVVMBlueprintViewExtension_ListViewBase>(Extension) != nullptr; });
+			if (!bExists)
 			{
 				UMVVMBlueprintViewExtension* NewExtension = Extension->CreateBlueprintWidgetExtension(UMVVMBlueprintViewExtension_ListViewBase::StaticClass(), WidgetPtr->GetFName());
 				UMVVMBlueprintViewExtension_ListViewBase* NewListViewExtension = CastChecked<UMVVMBlueprintViewExtension_ListViewBase>(NewExtension);
diff --git a/Engine/Plugins/Runtime/ModelViewViewModel/Source/ModelViewViewModelEditor/Private/Customizations/MVVMPanelWidgetExtensionCustomizationExtender.cpp b/Engine/Plugins/Runtime/ModelViewViewModel/Source/ModelViewViewModelEditor/Private/Customizations/MVVMPanelWidgetExtensionCustomizationExtender.cpp
--- a/Engine/Plugins/Runtime/ModelViewViewModel/Source/ModelViewViewModelEditor/Private/Customizations/MVVMPanelWidgetExtensionCustomizationExtender.cpp
+++ b/Engine/Plugins/Runtime/ModelViewViewModel/Source/ModelViewViewModelEditor/Private/Customizations/MVVMPanelWidgetExtensionCustomizationExtender.cpp
@@ -207,7 +207,9 @@ void FMVVMPanelWidgetExtensionCustomizationExtender::CreatePanelWidgetViewExtens
 	{
 		if (UPanelWidget* WidgetPtr = Widget.Get())
 		{
-			if (Extension->GetBlueprintExtensionsForWidget(WidgetPtr->GetFName()).IsEmpty())
+			TArray<UMVVMBlueprintViewExtension*> Extensions = Extension->GetBlueprintExtensionsForWidget(WidgetPtr->GetFName());
+			bool bExists = Extensions.ContainsByPredicate([](UMVVMBlueprintViewExtension* Extension){ return Cast<UMVVMBlueprintViewExtension_PanelWidget>(Extension) != nullptr; });
+			if (!bExists)
 			{
 				UMVVMBlueprintViewExtension* NewExtension = Extension->CreateBlueprintWidgetExtension(UMVVMBlueprintViewExtension_PanelWidget::StaticClass(), WidgetPtr->GetFName());
 				UMVVMBlueprintViewExtension_PanelWidget* NewPanelWidgetExtension = CastChecked<UMVVMBlueprintViewExtension_PanelWidget>(NewExtension);
