MVVM: Hide autogenerated function from the picker. Rename old function before creating a new one. That could be reproduced by selecting multiple conversion function on the same binding.

diff --git a/Engine/Plugins/Runtime/ModelViewViewModel/Source/ModelViewViewModelBlueprint/Private/Bindings/MVVMConversionFunctionHelper.cpp b/Engine/Plugins/Runtime/ModelViewViewModel/Source/ModelViewViewModelBlueprint/Private/Bindings/MVVMConversionFunctionHelper.cpp
--- a/Engine/Plugins/Runtime/ModelViewViewModel/Source/ModelViewViewModelBlueprint/Private/Bindings/MVVMConversionFunctionHelper.cpp
+++ b/Engine/Plugins/Runtime/ModelViewViewModel/Source/ModelViewViewModelBlueprint/Private/Bindings/MVVMConversionFunctionHelper.cpp
@@ -75,6 +75,16 @@ namespace Private
 		}
 		return nullptr;
 	}
+	
+	void RenameObjectToTransientPackage(UObject* ObjectToRename)
+	{
+		const ERenameFlags RenFlags = REN_DoNotDirty | REN_ForceNoResetLoaders | REN_DontCreateRedirectors;
+
+		ObjectToRename->Rename(nullptr, GetTransientPackage(), RenFlags);
+		ObjectToRename->SetFlags(RF_Transient);
+		ObjectToRename->ClearFlags(RF_Public | RF_Standalone | RF_ArchetypeObject);
+		FLinkerLoad::InvalidateExport(ObjectToRename);
+	}
 
 	struct FCreateGraphResult
 	{
@@ -84,8 +94,16 @@ namespace Private
 	};
 	FCreateGraphResult CreateGraph(UBlueprint* Blueprint, FName GraphName, const UFunction* FunctionEntryDefinition, bool bIsConst, bool bIsEditable, bool bAddToBlueprint)
 	{
+		if (UObject* ExistingObject = StaticFindObject(nullptr, Blueprint, *GraphName.ToString(), true))
+		{
+			RenameObjectToTransientPackage(ExistingObject);
+		}
+
 		FName UniqueFunctionName = FBlueprintEditorUtils::FindUniqueKismetName(Blueprint, GraphName.ToString());
+		ensure(GraphName == UniqueFunctionName);
 		UEdGraph* FunctionGraph = FBlueprintEditorUtils::CreateNewGraph(Blueprint, UniqueFunctionName, UEdGraph::StaticClass(), UMVVMConversionFunctionGraphSchema::StaticClass());
+		ensure(FunctionGraph->GetFName() == GraphName);
+
 		FunctionGraph->bEditable = bIsEditable;
 		if (bAddToBlueprint)
 		{
@@ -1278,6 +1296,18 @@ UEdGraphPin* FindOutputPin(const UK2Node* Node)
 	return Node? Private::FindNewOutputPin(Node) : nullptr;
 }
 
+void SetMetaData(UEdGraph* NewGraph, FName MetaData, FStringView Value)
+{
+	if (NewGraph == nullptr || MetaData.IsNone())
+	{
+		return;
+	}
+	if (UK2Node_FunctionEntry* FunctionEntry = Private::FindFunctionEntry(NewGraph))
+	{
+		FunctionEntry->MetaData.SetMetaData(MetaData, Value);
+	}
+}
+
 } // UE::MVVM::ConversionFunctionHelper
 
 #undef LOCTEXT_NAMESPACE
diff --git a/Engine/Plugins/Runtime/ModelViewViewModel/Source/ModelViewViewModelBlueprint/Private/MVVMBlueprintViewConversionFunction.cpp b/Engine/Plugins/Runtime/ModelViewViewModel/Source/ModelViewViewModelBlueprint/Private/MVVMBlueprintViewConversionFunction.cpp
--- a/Engine/Plugins/Runtime/ModelViewViewModel/Source/ModelViewViewModelBlueprint/Private/MVVMBlueprintViewConversionFunction.cpp
+++ b/Engine/Plugins/Runtime/ModelViewViewModel/Source/ModelViewViewModelBlueprint/Private/MVVMBlueprintViewConversionFunction.cpp
@@ -24,6 +24,7 @@
 namespace UE::MVVM::Private
 {
 	static FLazyName DefaultConversionFunctionName = TEXT("__ConversionFunction");
+	static FLazyName HiddenName = "Hidden";
 }
 
 bool UMVVMBlueprintViewConversionFunction::IsValidConversionFunction(const UBlueprint* WidgetBlueprint, const UFunction* Function)
@@ -414,6 +415,8 @@ UEdGraph* UMVVMBlueprintViewConversionFunction::GetOrCreateWrapperGraphInternal(
 		check(Node.Get());
 		Result = UE::MVVM::ConversionFunctionHelper::CreateGraph(Blueprint, GraphName, nullptr, Node, bConst, bWrapperGraphTransient, [](UK2Node*){});
 	}
+	UE::MVVM::ConversionFunctionHelper::SetMetaData(Result.NewGraph, UE::MVVM::Private::HiddenName.Resolve(), FStringView());
+
 	const_cast<UMVVMBlueprintViewConversionFunction*>(this)->SetCachedWrapperGraph(Blueprint, Result.NewGraph, Result.WrappedNode);
 	LoadPinValuesInternal(Blueprint);
 	return CachedWrapperGraph;
diff --git a/Engine/Plugins/Runtime/ModelViewViewModel/Source/ModelViewViewModelBlueprint/Private/MVVMBlueprintViewEvent.cpp b/Engine/Plugins/Runtime/ModelViewViewModel/Source/ModelViewViewModelBlueprint/Private/MVVMBlueprintViewEvent.cpp
--- a/Engine/Plugins/Runtime/ModelViewViewModel/Source/ModelViewViewModelBlueprint/Private/MVVMBlueprintViewEvent.cpp
+++ b/Engine/Plugins/Runtime/ModelViewViewModel/Source/ModelViewViewModelBlueprint/Private/MVVMBlueprintViewEvent.cpp
@@ -27,6 +27,11 @@
 
 #define LOCTEXT_NAMESPACE "MVVMBlueprintViewEvent"
 
+namespace UE::MVVM::Private
+{
+	static FLazyName HiddenName = "Hidden";
+}
+
 void UMVVMBlueprintViewEvent::SetEventPath(FMVVMBlueprintPropertyPath InEventPath)
 {
 	if (InEventPath == EventPath)
@@ -275,6 +280,8 @@ UEdGraph* UMVVMBlueprintViewEvent::CreateWrapperGraphInternal()
 	}
 	else
 	{
+		UE::MVVM::ConversionFunctionHelper::SetMetaData(CreateSetterGraphResult.GetValue().NewGraph, UE::MVVM::Private::HiddenName.Resolve(), FStringView());
+
 		UMVVMK2Node_AreSourcesValidForEvent* BranchNode = Cast<UMVVMK2Node_AreSourcesValidForEvent>(UE::MVVM::ConversionFunctionHelper::InsertEarlyExitBranchNode(CreateSetterGraphResult.GetValue().NewGraph, UMVVMK2Node_AreSourcesValidForEvent::StaticClass()));
 		SetCachedWrapperGraphInternal(CreateSetterGraphResult.GetValue().NewGraph, CreateSetterGraphResult.GetValue().WrappedNode, BranchNode);
 		LoadPinValuesInternal();
diff --git a/Engine/Plugins/Runtime/ModelViewViewModel/Source/ModelViewViewModelBlueprint/Public/Bindings/MVVMConversionFunctionHelper.h b/Engine/Plugins/Runtime/ModelViewViewModel/Source/ModelViewViewModelBlueprint/Public/Bindings/MVVMConversionFunctionHelper.h
--- a/Engine/Plugins/Runtime/ModelViewViewModel/Source/ModelViewViewModelBlueprint/Public/Bindings/MVVMConversionFunctionHelper.h
+++ b/Engine/Plugins/Runtime/ModelViewViewModel/Source/ModelViewViewModelBlueprint/Public/Bindings/MVVMConversionFunctionHelper.h
@@ -83,6 +83,9 @@ namespace UE::MVVM::ConversionFunctionHelper
 	
 	/** Return the pin used as the return value. */
 	MODELVIEWVIEWMODELBLUEPRINT_API UEdGraphPin* FindOutputPin(const UK2Node* Node);
+
+	/** Add metadata to the Graph/Function. */
+	MODELVIEWVIEWMODELBLUEPRINT_API void SetMetaData(UEdGraph* NewGraph, FName MetaData, FStringView Value);
 } //namespace
 
 #if UE_ENABLE_INCLUDE_ORDER_DEPRECATED_IN_5_2
diff --git a/Engine/Plugins/Runtime/ModelViewViewModel/Source/ModelViewViewModelEditor/Private/Widgets/SMVVMViewModelBindingListWidget.cpp b/Engine/Plugins/Runtime/ModelViewViewModel/Source/ModelViewViewModelEditor/Private/Widgets/SMVVMViewModelBindingListWidget.cpp
--- a/Engine/Plugins/Runtime/ModelViewViewModel/Source/ModelViewViewModelEditor/Private/Widgets/SMVVMViewModelBindingListWidget.cpp
+++ b/Engine/Plugins/Runtime/ModelViewViewModel/Source/ModelViewViewModelEditor/Private/Widgets/SMVVMViewModelBindingListWidget.cpp
@@ -48,12 +48,13 @@ namespace Private
 
 	TOptional<FFieldVariant> PassFilter(const UBlueprint* Blueprint, const FMVVMAvailableBinding& Binding, const UStruct* Struct, const FMVVMFieldVariant& FieldVariant, EFieldVisibility FieldVisibilityFlags, const FProperty* AssignableTo, bool bDoObjectProperty)
 	{
+		static FName NAME_Hidden("Hidden");
 		if (ensure(!FieldVariant.IsEmpty()))
 		{
 			if (FieldVariant.IsFunction())
 			{
 				const UFunction* Function = FieldVariant.GetFunction();
-				if (Function == nullptr)
+				if (Function == nullptr || Function->HasMetaData(NAME_Hidden))
 				{
 					return TOptional<FFieldVariant>();
 				}
@@ -104,7 +105,7 @@ namespace Private
 			else if (FieldVariant.IsProperty())
 			{
 				const FProperty* Property = FieldVariant.GetProperty();
-				if (Property == nullptr)
+				if (Property == nullptr || Property->HasMetaData(NAME_Hidden))
 				{
 					return TOptional<FFieldVariant>();
 				}
