[Backout] - CL33402680
Original CL Desc
-----------------------------------------------------------------
MVVM: Hide autogenerated function from the picker. Rename old function before creating a new one. That could be reproduced by selecting multiple conversion function on the same binding.

diff --git a/Engine/Plugins/Runtime/ModelViewViewModel/Source/ModelViewViewModelBlueprint/Private/MVVMBlueprintViewConversionFunction.cpp b/Engine/Plugins/Runtime/ModelViewViewModel/Source/ModelViewViewModelBlueprint/Private/MVVMBlueprintViewConversionFunction.cpp
--- a/Engine/Plugins/Runtime/ModelViewViewModel/Source/ModelViewViewModelBlueprint/Private/MVVMBlueprintViewConversionFunction.cpp
+++ b/Engine/Plugins/Runtime/ModelViewViewModel/Source/ModelViewViewModelBlueprint/Private/MVVMBlueprintViewConversionFunction.cpp
@@ -24,7 +24,6 @@
 namespace UE::MVVM::Private
 {
 	static FLazyName DefaultConversionFunctionName = TEXT("__ConversionFunction");
-	static FLazyName HiddenName = "Hidden";
 }
 
 bool UMVVMBlueprintViewConversionFunction::IsValidConversionFunction(const UBlueprint* WidgetBlueprint, const UFunction* Function)
@@ -415,8 +414,6 @@ UEdGraph* UMVVMBlueprintViewConversionFunction::GetOrCreateWrapperGraphInternal(
 		check(Node.Get());
 		Result = UE::MVVM::ConversionFunctionHelper::CreateGraph(Blueprint, GraphName, nullptr, Node, bConst, bWrapperGraphTransient, [](UK2Node*){});
 	}
-	UE::MVVM::ConversionFunctionHelper::SetMetaData(Result.NewGraph, UE::MVVM::Private::HiddenName.Resolve(), FStringView());
-
 	const_cast<UMVVMBlueprintViewConversionFunction*>(this)->SetCachedWrapperGraph(Blueprint, Result.NewGraph, Result.WrappedNode);
 	LoadPinValuesInternal(Blueprint);
 	return CachedWrapperGraph;
diff --git a/Engine/Plugins/Runtime/ModelViewViewModel/Source/ModelViewViewModelBlueprint/Private/MVVMBlueprintViewEvent.cpp b/Engine/Plugins/Runtime/ModelViewViewModel/Source/ModelViewViewModelBlueprint/Private/MVVMBlueprintViewEvent.cpp
--- a/Engine/Plugins/Runtime/ModelViewViewModel/Source/ModelViewViewModelBlueprint/Private/MVVMBlueprintViewEvent.cpp
+++ b/Engine/Plugins/Runtime/ModelViewViewModel/Source/ModelViewViewModelBlueprint/Private/MVVMBlueprintViewEvent.cpp
@@ -27,11 +27,6 @@
 
 #define LOCTEXT_NAMESPACE "MVVMBlueprintViewEvent"
 
-namespace UE::MVVM::Private
-{
-	static FLazyName HiddenName = "Hidden";
-}
-
 void UMVVMBlueprintViewEvent::SetEventPath(FMVVMBlueprintPropertyPath InEventPath)
 {
 	if (InEventPath == EventPath)
@@ -280,8 +275,6 @@ UEdGraph* UMVVMBlueprintViewEvent::CreateWrapperGraphInternal()
 	}
 	else
 	{
-		UE::MVVM::ConversionFunctionHelper::SetMetaData(CreateSetterGraphResult.GetValue().NewGraph, UE::MVVM::Private::HiddenName.Resolve(), FStringView());
-
 		UMVVMK2Node_AreSourcesValidForEvent* BranchNode = Cast<UMVVMK2Node_AreSourcesValidForEvent>(UE::MVVM::ConversionFunctionHelper::InsertEarlyExitBranchNode(CreateSetterGraphResult.GetValue().NewGraph, UMVVMK2Node_AreSourcesValidForEvent::StaticClass()));
 		SetCachedWrapperGraphInternal(CreateSetterGraphResult.GetValue().NewGraph, CreateSetterGraphResult.GetValue().WrappedNode, BranchNode);
 		LoadPinValuesInternal();
diff --git a/Engine/Plugins/Runtime/ModelViewViewModel/Source/ModelViewViewModelEditor/Private/Widgets/SMVVMViewModelBindingListWidget.cpp b/Engine/Plugins/Runtime/ModelViewViewModel/Source/ModelViewViewModelEditor/Private/Widgets/SMVVMViewModelBindingListWidget.cpp
--- a/Engine/Plugins/Runtime/ModelViewViewModel/Source/ModelViewViewModelEditor/Private/Widgets/SMVVMViewModelBindingListWidget.cpp
+++ b/Engine/Plugins/Runtime/ModelViewViewModel/Source/ModelViewViewModelEditor/Private/Widgets/SMVVMViewModelBindingListWidget.cpp
@@ -48,13 +48,12 @@ namespace Private
 
 	TOptional<FFieldVariant> PassFilter(const UBlueprint* Blueprint, const FMVVMAvailableBinding& Binding, const UStruct* Struct, const FMVVMFieldVariant& FieldVariant, EFieldVisibility FieldVisibilityFlags, const FProperty* AssignableTo, bool bDoObjectProperty)
 	{
-		static FName NAME_Hidden("Hidden");
 		if (ensure(!FieldVariant.IsEmpty()))
 		{
 			if (FieldVariant.IsFunction())
 			{
 				const UFunction* Function = FieldVariant.GetFunction();
-				if (Function == nullptr || Function->HasMetaData(NAME_Hidden))
+				if (Function == nullptr)
 				{
 					return TOptional<FFieldVariant>();
 				}
@@ -105,7 +104,7 @@ namespace Private
 			else if (FieldVariant.IsProperty())
 			{
 				const FProperty* Property = FieldVariant.GetProperty();
-				if (Property == nullptr || Property->HasMetaData(NAME_Hidden))
+				if (Property == nullptr)
 				{
 					return TOptional<FFieldVariant>();
 				}
