MVVM Extension: allow extensions to add widgets as variables via MVVM compiler

diff --git a/Engine/Plugins/Runtime/ModelViewViewModel/Source/ModelViewViewModelBlueprint/Private/Extensions/MVVMViewBlueprintPanelWidgetExtension.cpp b/Engine/Plugins/Runtime/ModelViewViewModel/Source/ModelViewViewModelBlueprint/Private/Extensions/MVVMViewBlueprintPanelWidgetExtension.cpp
--- a/Engine/Plugins/Runtime/ModelViewViewModel/Source/ModelViewViewModelBlueprint/Private/Extensions/MVVMViewBlueprintPanelWidgetExtension.cpp
+++ b/Engine/Plugins/Runtime/ModelViewViewModel/Source/ModelViewViewModelBlueprint/Private/Extensions/MVVMViewBlueprintPanelWidgetExtension.cpp
@@ -22,19 +22,38 @@
 TArray<UE::MVVM::Compiler::FBlueprintViewUserWidgetProperty> UMVVMBlueprintViewExtension_PanelWidget::AddProperties()
 {
 	TArray<UE::MVVM::Compiler::FBlueprintViewUserWidgetProperty> PropertiesToAdd;
-	PanelPropertyName = FName(FText::Format(LOCTEXT("ExtensionVariableName", "{0}_{1}"), FText::FromName(WidgetName), FText::FromString(TEXT("Viewmodel_Extension"))).ToString());
 
-	UE::MVVM::Compiler::FBlueprintViewUserWidgetProperty Property;
-	Property.AuthoritativeClass = UMVVMPanelWidgetViewExtension::StaticClass();
-	Property.DisplayName = FText::FromName(PanelPropertyName);
-	Property.Name = PanelPropertyName;
-	Property.CategoryName = TEXT("PanelWidgetExtension");
-	Property.bReadOnly = true;
+	if (!WidgetName.IsNone())
+	{
+		// Add the runtime panel widget extension as a variable
+		PanelPropertyName = FName(FString::Printf(TEXT("%s_%s"), *WidgetName.ToString(), *FString(TEXT("Viewmodel_Extension"))));
+		UE::MVVM::Compiler::FBlueprintViewUserWidgetProperty Property;
+		Property.AuthoritativeClass = UMVVMPanelWidgetViewExtension::StaticClass();
+		Property.DisplayName = FText::FromName(PanelPropertyName);
+		Property.Name = PanelPropertyName;
+		Property.CategoryName = TEXT("PanelWidgetExtension");
+		Property.bReadOnly = true;
+		PropertiesToAdd.Add(Property);
+	}
 
-	PropertiesToAdd.Add(Property);
 	return PropertiesToAdd;
 }
 
+TArray<UE::MVVM::Compiler::FBlueprintViewUserWidgetWidgetProperty> UMVVMBlueprintViewExtension_PanelWidget::AddWidgetProperties()
+{
+	TArray<UE::MVVM::Compiler::FBlueprintViewUserWidgetWidgetProperty> WidgetPropertiesToAdd;
+
+	if (!WidgetName.IsNone())
+	{
+		// Add the panel widget as a variable
+		UE::MVVM::Compiler::FBlueprintViewUserWidgetWidgetProperty WidgetProperty;
+		WidgetProperty.WidgetName = WidgetName;
+		WidgetPropertiesToAdd.Add(WidgetProperty);
+	}
+
+	return WidgetPropertiesToAdd;
+}
+
 void UMVVMBlueprintViewExtension_PanelWidget::Precompile(UE::MVVM::Compiler::IMVVMBlueprintViewPrecompile* Compiler, UWidgetBlueprintGeneratedClass* Class)
 {
 	check (Compiler);
diff --git a/Engine/Plugins/Runtime/ModelViewViewModel/Source/ModelViewViewModelBlueprint/Private/MVVMViewBlueprintCompiler.cpp b/Engine/Plugins/Runtime/ModelViewViewModel/Source/ModelViewViewModelBlueprint/Private/MVVMViewBlueprintCompiler.cpp
--- a/Engine/Plugins/Runtime/ModelViewViewModel/Source/ModelViewViewModelBlueprint/Private/MVVMViewBlueprintCompiler.cpp
+++ b/Engine/Plugins/Runtime/ModelViewViewModel/Source/ModelViewViewModelBlueprint/Private/MVVMViewBlueprintCompiler.cpp
@@ -1075,7 +1075,28 @@ void FMVVMViewBlueprintCompiler::CreateRequiredProperties(const FWidgetBlueprint
 				CompilerUserWidgetProperty.bPrivate = Property.bPrivate;
 				CompilerUserWidgetProperty.bReadOnly = Property.bReadOnly;
 				CompilerUserWidgetProperty.CategoryName = Property.CategoryName;
-				CompilerUserWidgetProperty.CreationType = FCompilerUserWidgetProperty::ECreationType::CreateOnlyIfDoesntExist;
+				CompilerUserWidgetProperty.CreationType = FCompilerUserWidgetProperty::ECreationType::CreateIfDoesntExist;
+			}
+
+			TArray<Compiler::FBlueprintViewUserWidgetWidgetProperty> WidgetProperties = ExtensionPtr->AddWidgetProperties();
+			for (const Compiler::FBlueprintViewUserWidgetWidgetProperty& WidgetProperty : WidgetProperties)
+			{
+				if (WidgetNameToWidgetPointerMap.Find(WidgetProperty.WidgetName))
+				{				
+					FMVVMBlueprintPropertyPath WidgetPropertyPath;
+					WidgetPropertyPath.SetWidgetName(WidgetProperty.WidgetName);
+					TValueOrError<void, FText> SourceContextResult = GenerateCompilerContext(false, WidgetPropertyPath);
+					if (SourceContextResult.HasError())
+					{
+						AddMessage(SourceContextResult.StealError(), Compiler::EMessageType::Error);
+						bIsCreateVariableStepValid = false;
+					}
+				}
+				else
+				{
+					AddMessage(FText::Format(LOCTEXT("InvalidWidgetPropertyRequested", "Could not find the widget : {0} in the widget tree to make the property requested by viewmodel extension."), FText::FromName(WidgetProperty.WidgetName)), Compiler::EMessageType::Error);
+					bIsCreateVariableStepValid = false;
+				}
 			}
 		}
 	}
diff --git a/Engine/Plugins/Runtime/ModelViewViewModel/Source/ModelViewViewModelBlueprint/Public/Extensions/MVVMBlueprintViewExtension.h b/Engine/Plugins/Runtime/ModelViewViewModel/Source/ModelViewViewModelBlueprint/Public/Extensions/MVVMBlueprintViewExtension.h
--- a/Engine/Plugins/Runtime/ModelViewViewModel/Source/ModelViewViewModelBlueprint/Public/Extensions/MVVMBlueprintViewExtension.h
+++ b/Engine/Plugins/Runtime/ModelViewViewModel/Source/ModelViewViewModelBlueprint/Public/Extensions/MVVMBlueprintViewExtension.h
@@ -26,6 +26,12 @@ public:
 	{
 		return TArray<UE::MVVM::Compiler::FBlueprintViewUserWidgetProperty>();
 	};
+
+	virtual TArray<UE::MVVM::Compiler::FBlueprintViewUserWidgetWidgetProperty> AddWidgetProperties()
+	{
+		return TArray<UE::MVVM::Compiler::FBlueprintViewUserWidgetWidgetProperty>();
+	};
+
 	virtual void Precompile(UE::MVVM::Compiler::IMVVMBlueprintViewPrecompile* Compiler, UWidgetBlueprintGeneratedClass* Class) {}
 	virtual void Compile(UE::MVVM::Compiler::IMVVMBlueprintViewCompile* Compiler, UWidgetBlueprintGeneratedClass* Class, UMVVMViewClass* ViewExtension) {}
 	virtual bool WidgetRenamed(FName OldName, FName NewName) 
diff --git a/Engine/Plugins/Runtime/ModelViewViewModel/Source/ModelViewViewModelBlueprint/Public/Extensions/MVVMViewBlueprintPanelWidgetExtension.h b/Engine/Plugins/Runtime/ModelViewViewModel/Source/ModelViewViewModelBlueprint/Public/Extensions/MVVMViewBlueprintPanelWidgetExtension.h
--- a/Engine/Plugins/Runtime/ModelViewViewModel/Source/ModelViewViewModelBlueprint/Public/Extensions/MVVMViewBlueprintPanelWidgetExtension.h
+++ b/Engine/Plugins/Runtime/ModelViewViewModel/Source/ModelViewViewModelBlueprint/Public/Extensions/MVVMViewBlueprintPanelWidgetExtension.h
@@ -34,6 +34,7 @@ class MODELVIEWVIEWMODELBLUEPRINT_API UMVVMBlueprintViewExtension_PanelWidget :
 public:
 	//~ Begin UMVVMBlueprintViewExtension overrides
 	virtual TArray<UE::MVVM::Compiler::FBlueprintViewUserWidgetProperty> AddProperties() override;
+	virtual TArray<UE::MVVM::Compiler::FBlueprintViewUserWidgetWidgetProperty> AddWidgetProperties() override;
 	virtual void Precompile(UE::MVVM::Compiler::IMVVMBlueprintViewPrecompile* Compiler, UWidgetBlueprintGeneratedClass* Class) override;
 	virtual void Compile(UE::MVVM::Compiler::IMVVMBlueprintViewCompile* Compiler, UWidgetBlueprintGeneratedClass* Class, UMVVMViewClass* ViewExtension) override;
 	virtual bool WidgetRenamed(FName OldName, FName NewName) override;
diff --git a/Engine/Plugins/Runtime/ModelViewViewModel/Source/ModelViewViewModelBlueprint/Public/MVVMBlueprintViewCompilerInterface.h b/Engine/Plugins/Runtime/ModelViewViewModel/Source/ModelViewViewModelBlueprint/Public/MVVMBlueprintViewCompilerInterface.h
--- a/Engine/Plugins/Runtime/ModelViewViewModel/Source/ModelViewViewModelBlueprint/Public/MVVMBlueprintViewCompilerInterface.h
+++ b/Engine/Plugins/Runtime/ModelViewViewModel/Source/ModelViewViewModelBlueprint/Public/MVVMBlueprintViewCompilerInterface.h
@@ -61,6 +61,17 @@ enum class EMessageType
 	Error = 2
 };
 
+/**
+ * Data to be used by the compiler to create a property for a widget in the widget tree
+ */
+struct FBlueprintViewUserWidgetWidgetProperty
+{
+	FName WidgetName;
+};
+
+/**
+ * Data to be used by the compiler to create new user widget properties
+ */
 struct FBlueprintViewUserWidgetProperty
 {
 	UClass* AuthoritativeClass = nullptr;
