MVVM: Hide autogenerated function from the picker. Rename old function before creating a new one. That could be reproduced by selecting multiple conversion function on the same binding.

diff --git a/Engine/Plugins/Runtime/ModelViewViewModel/Source/ModelViewViewModelBlueprint/Private/MVVMBlueprintViewConversionFunction.cpp b/Engine/Plugins/Runtime/ModelViewViewModel/Source/ModelViewViewModelBlueprint/Private/MVVMBlueprintViewConversionFunction.cpp
--- a/Engine/Plugins/Runtime/ModelViewViewModel/Source/ModelViewViewModelBlueprint/Private/MVVMBlueprintViewConversionFunction.cpp
+++ b/Engine/Plugins/Runtime/ModelViewViewModel/Source/ModelViewViewModelBlueprint/Private/MVVMBlueprintViewConversionFunction.cpp
@@ -414,6 +414,10 @@ UEdGraph* UMVVMBlueprintViewConversionFunction::GetOrCreateWrapperGraphInternal(
 		check(Node.Get());
 		Result = UE::MVVM::ConversionFunctionHelper::CreateGraph(Blueprint, GraphName, nullptr, Node, bConst, bWrapperGraphTransient, [](UK2Node*){});
 	}
+
+	static FName NAME_Hidden("Hidden");
+	UE::MVVM::ConversionFunctionHelper::SetMetaData(Result.NewGraph, NAME_Hidden, FStringView());
+
 	const_cast<UMVVMBlueprintViewConversionFunction*>(this)->SetCachedWrapperGraph(Blueprint, Result.NewGraph, Result.WrappedNode);
 	LoadPinValuesInternal(Blueprint);
 	return CachedWrapperGraph;
diff --git a/Engine/Plugins/Runtime/ModelViewViewModel/Source/ModelViewViewModelBlueprint/Private/MVVMBlueprintViewEvent.cpp b/Engine/Plugins/Runtime/ModelViewViewModel/Source/ModelViewViewModelBlueprint/Private/MVVMBlueprintViewEvent.cpp
--- a/Engine/Plugins/Runtime/ModelViewViewModel/Source/ModelViewViewModelBlueprint/Private/MVVMBlueprintViewEvent.cpp
+++ b/Engine/Plugins/Runtime/ModelViewViewModel/Source/ModelViewViewModelBlueprint/Private/MVVMBlueprintViewEvent.cpp
@@ -275,6 +275,9 @@ UEdGraph* UMVVMBlueprintViewEvent::CreateWrapperGraphInternal()
 	}
 	else
 	{
+		static FName NAME_Hidden("Hidden");
+		UE::MVVM::ConversionFunctionHelper::SetMetaData(CreateSetterGraphResult.GetValue().NewGraph, NAME_Hidden, FStringView());
+
 		UMVVMK2Node_AreSourcesValidForEvent* BranchNode = Cast<UMVVMK2Node_AreSourcesValidForEvent>(UE::MVVM::ConversionFunctionHelper::InsertEarlyExitBranchNode(CreateSetterGraphResult.GetValue().NewGraph, UMVVMK2Node_AreSourcesValidForEvent::StaticClass()));
 		SetCachedWrapperGraphInternal(CreateSetterGraphResult.GetValue().NewGraph, CreateSetterGraphResult.GetValue().WrappedNode, BranchNode);
 		LoadPinValuesInternal();
diff --git a/Engine/Plugins/Runtime/ModelViewViewModel/Source/ModelViewViewModelEditor/Private/Widgets/SMVVMViewModelBindingListWidget.cpp b/Engine/Plugins/Runtime/ModelViewViewModel/Source/ModelViewViewModelEditor/Private/Widgets/SMVVMViewModelBindingListWidget.cpp
--- a/Engine/Plugins/Runtime/ModelViewViewModel/Source/ModelViewViewModelEditor/Private/Widgets/SMVVMViewModelBindingListWidget.cpp
+++ b/Engine/Plugins/Runtime/ModelViewViewModel/Source/ModelViewViewModelEditor/Private/Widgets/SMVVMViewModelBindingListWidget.cpp
@@ -48,12 +48,13 @@ namespace Private
 
 	TOptional<FFieldVariant> PassFilter(const UBlueprint* Blueprint, const FMVVMAvailableBinding& Binding, const UStruct* Struct, const FMVVMFieldVariant& FieldVariant, EFieldVisibility FieldVisibilityFlags, const FProperty* AssignableTo, bool bDoObjectProperty)
 	{
+		static FName NAME_Hidden("Hidden");
 		if (ensure(!FieldVariant.IsEmpty()))
 		{
 			if (FieldVariant.IsFunction())
 			{
 				const UFunction* Function = FieldVariant.GetFunction();
-				if (Function == nullptr)
+				if (Function == nullptr || Function->HasMetaData(NAME_Hidden))
 				{
 					return TOptional<FFieldVariant>();
 				}
@@ -104,7 +105,7 @@ namespace Private
 			else if (FieldVariant.IsProperty())
 			{
 				const FProperty* Property = FieldVariant.GetProperty();
-				if (Property == nullptr)
+				if (Property == nullptr || Property->HasMetaData(NAME_Hidden))
 				{
 					return TOptional<FFieldVariant>();
 				}
