MVVM Extension: fix slot template settings not being applied in cooked build.
The outer of the slot template must be an object which is available at runtime.

diff --git a/Engine/Plugins/Runtime/ModelViewViewModel/Source/ModelViewViewModel/Private/Extensions/MVVMViewPanelWidgetExtension.cpp b/Engine/Plugins/Runtime/ModelViewViewModel/Source/ModelViewViewModel/Private/Extensions/MVVMViewPanelWidgetExtension.cpp
--- a/Engine/Plugins/Runtime/ModelViewViewModel/Source/ModelViewViewModel/Private/Extensions/MVVMViewPanelWidgetExtension.cpp
+++ b/Engine/Plugins/Runtime/ModelViewViewModel/Source/ModelViewViewModel/Private/Extensions/MVVMViewPanelWidgetExtension.cpp
@@ -168,7 +168,7 @@ void UMVVMViewPanelWidgetClassExtension::Initialize(UMVVMViewPanelWidgetClassExt
 	WidgetPath = InArgs.WidgetPath;
 	EntryViewModelName = InArgs.EntryViewModelName;
 	EntryWidgetClass = InArgs.EntryWidgetClass;
-	SlotTemplate = InArgs.SlotTemplate;
+	SlotTemplate = DuplicateObject(InArgs.SlotTemplate, this);
 	PanelPropertyName = InArgs.PanelPropertyName;
 	EntryViewModelClass = InArgs.EntryViewModelClass;
 }
