Add conversion functions for MVVM Material Param updates.

#tests PIE

diff --git a/Engine/Plugins/Runtime/ModelViewViewModel/Source/ModelViewViewModel/Private/Bindings/ConversionLibraries/MVVMSlateBrushConversionLibrary.cpp b/Engine/Plugins/Runtime/ModelViewViewModel/Source/ModelViewViewModel/Private/Bindings/ConversionLibraries/MVVMSlateBrushConversionLibrary.cpp
new file mode 100644
--- /dev/null
+++ b/Engine/Plugins/Runtime/ModelViewViewModel/Source/ModelViewViewModel/Private/Bindings/ConversionLibraries/MVVMSlateBrushConversionLibrary.cpp
@@ -0,0 +1,96 @@
+// Copyright Epic Games, Inc. All Rights Reserved.
+
+#include "Bindings/ConversionLibraries/MVVMSlateBrushConversionLibrary.h"
+
+#include "Materials/MaterialInterface.h"
+#include "Materials/MaterialInstanceDynamic.h"
+
+FSlateBrush UMVVMSlateBrushConversionLibrary::Conv_SetScalarParameter(FSlateBrush Brush, FName ParameterName, float Value)
+{
+	if (UMaterialInstanceDynamic* Material = TryGetDynamicMaterial(Brush, nullptr))
+	{
+		Material->SetScalarParameterValue(ParameterName, Value);
+	}
+
+	return Brush;
+}
+
+FSlateBrush UMVVMSlateBrushConversionLibrary::Conv_SetVectorParameter(FSlateBrush Brush, FName ParameterName, FLinearColor Value)
+{
+	if (UMaterialInstanceDynamic* Material = TryGetDynamicMaterial(Brush, nullptr))
+	{
+		Material->SetVectorParameterValue(ParameterName, Value);
+	}
+
+	return Brush;
+}
+
+FSlateBrush UMVVMSlateBrushConversionLibrary::Conv_SetTextureParameter(FSlateBrush Brush, FName ParameterName, UTexture* Value)
+{
+	if (UMaterialInstanceDynamic* Material = TryGetDynamicMaterial(Brush, nullptr))
+	{
+		Material->SetTextureParameterValue(ParameterName, Value);
+	}
+
+	return Brush;
+}
+
+FSlateBrush UMVVMSlateBrushConversionLibrary::Conv_SetScalarParameterMID(FSlateBrush Brush, UMaterialInterface* Material, FName ParameterName, float Value)
+{
+	if (UMaterialInstanceDynamic* DynamicMaterial = TryGetDynamicMaterial(Brush, nullptr, Material))
+	{
+		DynamicMaterial->SetScalarParameterValue(ParameterName, Value);
+	}
+
+	return Brush;
+}
+
+FSlateBrush UMVVMSlateBrushConversionLibrary::Conv_SetVectorParameterMID(FSlateBrush Brush, UMaterialInterface* Material, FName ParameterName, FLinearColor Value)
+{
+	if (UMaterialInstanceDynamic* DynamicMaterial = TryGetDynamicMaterial(Brush, nullptr, Material))
+	{
+		DynamicMaterial->SetVectorParameterValue(ParameterName, Value);
+	}
+
+	return Brush;
+}
+
+FSlateBrush UMVVMSlateBrushConversionLibrary::Conv_SetTextureParameterMID(FSlateBrush Brush, UMaterialInterface* Material, FName ParameterName, UTexture* Value)
+{
+	if (UMaterialInstanceDynamic* DynamicMaterial = TryGetDynamicMaterial(Brush, nullptr, Material))
+	{
+		DynamicMaterial->SetTextureParameterValue(ParameterName, Value);
+	}
+
+	return Brush;
+}
+
+UMaterialInstanceDynamic* UMVVMSlateBrushConversionLibrary::TryGetDynamicMaterial(FSlateBrush& InBrush, UObject* InOuter, UMaterialInterface* InTargetMaterial)
+{
+	 UMaterialInterface* SourceMaterial = Cast<UMaterialInterface>(InBrush.GetResourceObject());
+
+	 // Want to either make the provided material if different, or make a MID of the existing material
+	 UMaterialInterface* DesiredMaterial = nullptr;
+	 if (InTargetMaterial && (!SourceMaterial || SourceMaterial->GetMaterial() != InTargetMaterial->GetMaterial()))
+	 {
+		 DesiredMaterial = InTargetMaterial;
+	 }
+	 else
+	 {
+		 DesiredMaterial = SourceMaterial;
+	 }
+
+	 if (DesiredMaterial)
+	 {
+	 	UMaterialInstanceDynamic* DynamicMaterial = Cast<UMaterialInstanceDynamic>(DesiredMaterial);
+
+	 	if ( !DynamicMaterial )
+	 	{
+	 		DynamicMaterial = UMaterialInstanceDynamic::Create(DesiredMaterial, InOuter);
+			InBrush.SetResourceObject(DynamicMaterial);
+	 	}
+	 	return DynamicMaterial;
+	 }
+
+	return nullptr;
+}
\ No newline at end of file
diff --git a/Engine/Plugins/Runtime/ModelViewViewModel/Source/ModelViewViewModel/Public/Bindings/ConversionLibraries/MVVMSlateBrushConversionLibrary.h b/Engine/Plugins/Runtime/ModelViewViewModel/Source/ModelViewViewModel/Public/Bindings/ConversionLibraries/MVVMSlateBrushConversionLibrary.h
new file mode 100644
--- /dev/null
+++ b/Engine/Plugins/Runtime/ModelViewViewModel/Source/ModelViewViewModel/Public/Bindings/ConversionLibraries/MVVMSlateBrushConversionLibrary.h
@@ -0,0 +1,97 @@
+// Copyright Epic Games, Inc. All Rights Reserved.
+
+#pragma once
+
+#include "CoreMinimal.h"
+#include "Components/SlateWrapperTypes.h"
+#include "Kismet/BlueprintFunctionLibrary.h"
+#include "Styling/SlateBrush.h"
+
+#include "MVVMSlateBrushConversionLibrary.generated.h"
+
+class UMaterialInstanceDynamic;
+class UMaterialInterface;
+class UTexture;
+
+/**
+ * Conversion library that contains methods for FSlateBrush
+ * 
+ * Primarily consists of methods to set material params on an existing FSlateBrush
+ */
+UCLASS()
+class MODELVIEWVIEWMODEL_API UMVVMSlateBrushConversionLibrary : public UBlueprintFunctionLibrary
+{
+	GENERATED_BODY()
+
+public:
+	
+	/**
+	 * Sets a scalar value on a brush material assuming it exists, handles MID existance appropriately
+	 * 
+	 * @param Brush			Brush with a material we want to set a parameter on 
+	 * @param ParameterName	Name of material parameter to set
+	 * @param Value			Value to set material parameter to
+	 */
+	UFUNCTION(BlueprintPure, Category = "Widget", meta = (DisplayName = "Set Scalar Parameter", MVVMBindToDestination = "Brush"))
+	static FSlateBrush Conv_SetScalarParameter(FSlateBrush Brush, FName ParameterName, float Value);
+	
+	/**
+	 * Sets a vector value on a brush material assuming it exists, handles MID existance appropriately
+	 * 
+	 * @param Brush			Brush with a material we want to set a parameter on 
+	 * @param ParameterName	Name of material parameter to set
+	 * @param Value			Value to set material parameter to
+	 */
+	UFUNCTION(BlueprintPure, Category = "Widget", meta = (DisplayName = "Set Vector Parameter", MVVMBindToDestination = "Brush"))
+	static FSlateBrush Conv_SetVectorParameter(FSlateBrush Brush, FName ParameterName, FLinearColor Value);
+	
+	/**
+	 * Sets a texture value on a brush material assuming it exists, handles MID existance appropriately
+	 * 
+	 * @param Brush			Brush with a material we want to set a parameter on 
+	 * @param ParameterName	Name of material parameter to set
+	 * @param Value			Value to set material parameter to
+	 */
+	UFUNCTION(BlueprintPure, Category = "Widget", meta = (DisplayName = "Set Texture Parameter", MVVMBindToDestination = "Brush"))
+	static FSlateBrush Conv_SetTextureParameter(FSlateBrush Brush, FName ParameterName, UTexture* Value);
+
+	/**
+	 * Sets a scalar value on a brush material assuming it exists, handles MID existance appropriately
+	 *
+	 * @param Brush			Brush with a material we want to set a parameter on
+	 * @param Material		Material to set on the brush
+	 * @param ParameterName	Name of material parameter to set
+	 * @param Value			Value to set material parameter to
+	 */
+	UFUNCTION(BlueprintPure, Category = "Widget", meta = (DisplayName = "Set Scalar Parameter MID", MVVMBindToDestination = "Brush"))
+	static FSlateBrush Conv_SetScalarParameterMID(FSlateBrush Brush, UMaterialInterface* Material, FName ParameterName, float Value);
+
+	/**
+	 * Sets a vector value on a brush material assuming it exists, handles MID existance appropriately
+	 *
+	 * @param Brush			Brush with a material we want to set a parameter on
+	 * @param Material		Material to set on the brush
+	 * @param ParameterName	Name of material parameter to set
+	 * @param Value			Value to set material parameter to
+	 */
+	UFUNCTION(BlueprintPure, Category = "Widget", meta = (DisplayName = "Set Vector Parameter MID", MVVMBindToDestination = "Brush"))
+	static FSlateBrush Conv_SetVectorParameterMID(FSlateBrush Brush, UMaterialInterface* Material, FName ParameterName, FLinearColor Value);
+
+	/**
+	 * Sets a texture value on a brush material assuming it exists, handles MID existance appropriately
+	 *
+	 * @param Brush			Brush with a material we want to set a parameter on
+	 * @param Material		Material to set on the brush
+	 * @param ParameterName	Name of material parameter to set
+	 * @param Value			Value to set material parameter to
+	 */
+	UFUNCTION(BlueprintPure, Category = "Widget", meta = (DisplayName = "Set Texture Parameter MID", MVVMBindToDestination = "Brush"))
+	static FSlateBrush Conv_SetTextureParameterMID(FSlateBrush Brush, UMaterialInterface* Material, FName ParameterName, UTexture* Value);
+
+private:
+
+	/**
+	 * Tries to get a Material Instance Dynamic from the brush, will create one if possible using the provided outer, else returns nullptr
+	 */
+	static UMaterialInstanceDynamic* TryGetDynamicMaterial(FSlateBrush& InBrush, UObject* InOuter, UMaterialInterface* InMaterial = nullptr);
+};
\ No newline at end of file
diff --git a/Engine/Plugins/Runtime/ModelViewViewModel/Source/ModelViewViewModelBlueprint/Public/MVVMBlueprintPin.h b/Engine/Plugins/Runtime/ModelViewViewModel/Source/ModelViewViewModelBlueprint/Public/MVVMBlueprintPin.h
--- a/Engine/Plugins/Runtime/ModelViewViewModel/Source/ModelViewViewModelBlueprint/Public/MVVMBlueprintPin.h
+++ b/Engine/Plugins/Runtime/ModelViewViewModel/Source/ModelViewViewModelBlueprint/Public/MVVMBlueprintPin.h
@@ -103,7 +103,7 @@ public:
 	explicit FMVVMBlueprintPin(FMVVMBlueprintPinId PinId);
 	explicit FMVVMBlueprintPin(const TArrayView<const FName> PinName);
 
-	UE_DEPRECATED(5.4, "GetName is deprecated. Use GetNames instead")
+	UE_DEPRECATED(5.4, "GetName is deprecated. Use GetId().GetNames instead")
 	FName GetName() const
 	{
 		return Id.GetNames().Num() > 0 ? Id.GetNames().Last() : FName();
diff --git a/Engine/Plugins/Runtime/ModelViewViewModel/Source/ModelViewViewModelEditor/Private/MVVMEditorSubsystem.cpp b/Engine/Plugins/Runtime/ModelViewViewModel/Source/ModelViewViewModelEditor/Private/MVVMEditorSubsystem.cpp
--- a/Engine/Plugins/Runtime/ModelViewViewModel/Source/ModelViewViewModelEditor/Private/MVVMEditorSubsystem.cpp
+++ b/Engine/Plugins/Runtime/ModelViewViewModel/Source/ModelViewViewModelEditor/Private/MVVMEditorSubsystem.cpp
@@ -8,6 +8,7 @@
 #include "BlueprintActionDatabase.h"
 #include "BlueprintActionFilter.h"
 #include "BlueprintNodeSpawner.h"
+#include "Editor.h"
 #include "Engine/Engine.h"
 #include "Framework/MVVMConversionFunctionLibrary.h"
 #include "Kismet/BlueprintFunctionLibrary.h"
@@ -407,6 +408,8 @@ void UMVVMEditorSubsystem::SetSourceToDestinationConversionFunction(UWidgetBluep
 				Binding.Conversion.SourceToDestinationConversion->Initialize(WidgetBlueprint, GraphName, NewConversionFunction);
 			}
 
+			GenerateBindToDestinationPathsForBinding(WidgetBlueprint, Binding);
+
 			UE::MVVM::Private::OnBindingPostEditChange(View, GET_MEMBER_NAME_CHECKED(FMVVMBlueprintViewBinding, Conversion));
 			FBlueprintEditorUtils::MarkBlueprintAsModified(WidgetBlueprint);
 		}
@@ -501,6 +504,8 @@ void UMVVMEditorSubsystem::SetDestinationPathForBinding(UWidgetBlueprint* Widget
 				Binding.DestinationPath = PropertyPath;
 			}
 
+			GenerateBindToDestinationPathsForBinding(WidgetBlueprint, Binding);
+
 			UE::MVVM::Private::OnBindingPostEditChange(View, GET_MEMBER_NAME_CHECKED(FMVVMBlueprintViewBinding, DestinationPath));
 			FBlueprintEditorUtils::MarkBlueprintAsModified(WidgetBlueprint);
 		}
@@ -622,6 +627,40 @@ void UMVVMEditorSubsystem::SetCompileForBinding(UWidgetBlueprint* WidgetBlueprin
 	}
 }
 
+void UMVVMEditorSubsystem::GenerateBindToDestinationPathsForBinding(UWidgetBlueprint* WidgetBlueprint, FMVVMBlueprintViewBinding& Binding)
+{
+	if (UMVVMBlueprintViewConversionFunction* ConversionFunction = Binding.Conversion.GetConversionFunction(UE::MVVM::IsForwardBinding(Binding.BindingType)))
+	{
+		if (const UFunction* Function = ConversionFunction->GetConversionFunction().GetFunction(WidgetBlueprint))
+		{
+			static const FName NAME_MVVMBindToDestination(TEXT("MVVMBindToDestination"));
+			const FString& MVVMBindToDestinationString = Function->GetMetaData(NAME_MVVMBindToDestination);
+			if (!MVVMBindToDestinationString.IsEmpty())
+			{
+				auto GetPinId = [&](const FString& InParamName)
+				{
+					for (const FMVVMBlueprintPin& Pin : Binding.Conversion.SourceToDestinationConversion->GetPins())
+					{
+						if (Pin.GetId().GetNames().Contains(InParamName))
+						{
+							return Pin.GetId();
+						}
+					}
+
+					return FMVVMBlueprintPinId();
+				};
+
+				FMVVMBlueprintPinId MVVMBindToDestinationPinId = GetPinId(MVVMBindToDestinationString);
+				if (MVVMBindToDestinationPinId.IsValid())
+				{
+					const UMVVMEditorSubsystem* Subsystem = GEditor->GetEditorSubsystem<UMVVMEditorSubsystem>();
+					Subsystem->SetPathForConversionFunctionArgument(WidgetBlueprint, Binding, MVVMBindToDestinationPinId, Binding.DestinationPath, true);
+				}
+			}
+		}
+	}
+}
+
 void UMVVMEditorSubsystem::SetEventPath(UMVVMBlueprintViewEvent* Event, FMVVMBlueprintPropertyPath PropertyPath)
 {
 	UMVVMBlueprintView* View = Event ? Event->GetOuterUMVVMBlueprintView() : nullptr;
diff --git a/Engine/Plugins/Runtime/ModelViewViewModel/Source/ModelViewViewModelEditor/Private/Widgets/BindingEntry/SMVVMFunctionParameterRow.cpp b/Engine/Plugins/Runtime/ModelViewViewModel/Source/ModelViewViewModelEditor/Private/Widgets/BindingEntry/SMVVMFunctionParameterRow.cpp
--- a/Engine/Plugins/Runtime/ModelViewViewModel/Source/ModelViewViewModelEditor/Private/Widgets/BindingEntry/SMVVMFunctionParameterRow.cpp
+++ b/Engine/Plugins/Runtime/ModelViewViewModel/Source/ModelViewViewModelEditor/Private/Widgets/BindingEntry/SMVVMFunctionParameterRow.cpp
@@ -49,6 +49,22 @@ TSharedRef<SWidget> SFunctionParameterRow::BuildRowWidget()
 			.ParameterId(GetEntry()->GetBindingParameterId())
 			.SourceToDestination(bSourceToDestination)
 			.AllowDefault(!bSimpleConversionFunction);
+
+		// Disable UI for autogenerated parameters
+		if (const UFunction* Function = ConversionFunction->GetConversionFunction().GetFunction(GetBlueprint()))
+		{
+			static const FName NAME_MVVMBindToDestination(TEXT("MVVMBindToDestination"));
+			const FString& MVVMBindToDestinationString = Function->GetMetaData(NAME_MVVMBindToDestination);
+			if (!MVVMBindToDestinationString.IsEmpty())
+			{
+				if (Pin->GetId().GetNames().Contains(MVVMBindToDestinationString))
+				{
+					ContentWidget->SetEnabled(false);
+					ContentWidget->SetToolTipText(LOCTEXT("MVVMBindToDestinationGeneratedParam", "Disabled since autogenerated param. As requested by 'MVVMBindToDestination' metadata in this UFunction."));
+				}
+			}
+		}
+
 	}
 	else if (GetEntry()->GetRowType() == FBindingEntry::ERowType::EventParameter)
 	{
diff --git a/Engine/Plugins/Runtime/ModelViewViewModel/Source/ModelViewViewModelEditor/Public/MVVMEditorSubsystem.h b/Engine/Plugins/Runtime/ModelViewViewModel/Source/ModelViewViewModelEditor/Public/MVVMEditorSubsystem.h
--- a/Engine/Plugins/Runtime/ModelViewViewModel/Source/ModelViewViewModelEditor/Public/MVVMEditorSubsystem.h
+++ b/Engine/Plugins/Runtime/ModelViewViewModel/Source/ModelViewViewModelEditor/Public/MVVMEditorSubsystem.h
@@ -85,6 +85,7 @@ public:
 	void SetBindingTypeForBinding(UWidgetBlueprint* WidgetBlueprint, FMVVMBlueprintViewBinding& Binding, EMVVMBindingMode Type);
 	void SetEnabledForBinding(UWidgetBlueprint* WidgetBlueprint, FMVVMBlueprintViewBinding& Binding, bool bEnabled);
 	void SetCompileForBinding(UWidgetBlueprint* WidgetBlueprint, FMVVMBlueprintViewBinding& Binding, bool bCompile);
+	void GenerateBindToDestinationPathsForBinding(UWidgetBlueprint* WidgetBlueprint, FMVVMBlueprintViewBinding& Binding);
 
 	void SetEventPath(UMVVMBlueprintViewEvent* Event, FMVVMBlueprintPropertyPath PropertyPath);
 	void SetEventDestinationPath(UMVVMBlueprintViewEvent* Event, FMVVMBlueprintPropertyPath PropertyPath);
