MVVM: Update FMVVMCompiledBindingLibrary::FieldPathToString to update the saved name of the property/function when the library is not loaded.

diff --git a/Engine/Plugins/Runtime/ModelViewViewModel/Source/ModelViewViewModel/Private/Bindings/MVVMCompiledBindingLibrary.cpp b/Engine/Plugins/Runtime/ModelViewViewModel/Source/ModelViewViewModel/Private/Bindings/MVVMCompiledBindingLibrary.cpp
--- a/Engine/Plugins/Runtime/ModelViewViewModel/Source/ModelViewViewModel/Private/Bindings/MVVMCompiledBindingLibrary.cpp
+++ b/Engine/Plugins/Runtime/ModelViewViewModel/Source/ModelViewViewModel/Private/Bindings/MVVMCompiledBindingLibrary.cpp
@@ -485,6 +485,8 @@ TValueOrError<FString, FString> FMVVMCompiledBindingLibrary::FieldPathToString(F
 	}
 #endif
 
+	const bool bIsLoaded = IsLoaded();
+
 	bool bHasError = false;
 	TStringBuilder<512> StringBuilder;
 	const int32 IterationMax = FieldPath.StartIndex + FieldPath.Num;
@@ -498,62 +500,104 @@ TValueOrError<FString, FString> FMVVMCompiledBindingLibrary::FieldPathToString(F
 
 		check(FieldPaths.IsValidIndex(Index));
 		FMVVMCompiledLoadedPropertyOrFunctionIndex PathIndex = FieldPaths[Index];
-		if (PathIndex.bIsProperty)
+		if (bIsLoaded)
 		{
-			if (LoadedProperties.IsValidIndex(PathIndex.Index))
+			if (PathIndex.bIsProperty)
 			{
-				const FProperty* Property = LoadedProperties[PathIndex.Index];
-				if (!Property)
+				if (LoadedProperties.IsValidIndex(PathIndex.Index))
+				{
+					const FProperty* Property = LoadedProperties[PathIndex.Index];
+					if (!Property)
+					{
+						StringBuilder << TEXT("<Invalid>");
+						bHasError = true;
+					}
+					else
+					{
+	#if WITH_EDITOR
+						if (bUseDisplayName)
+						{
+							StringBuilder << Property->GetDisplayNameText().ToString();
+						}
+						else
+	#endif
+						{
+							StringBuilder << Property->GetFName();
+						}
+					}
+				}
+				else
 				{
 					StringBuilder << TEXT("<Invalid>");
 					bHasError = true;
 				}
-				else
+			}
+			else
+			{
+				if (LoadedFunctions.IsValidIndex(PathIndex.Index))
 				{
-#if WITH_EDITOR
-					if (bUseDisplayName)
+					const FLoadedFunction& LoadedFunction = LoadedFunctions[PathIndex.Index];
+					const UFunction* Function = LoadedFunction.GetFunction();
+					if (!Function)
 					{
-						StringBuilder << Property->GetDisplayNameText().ToString();
+						StringBuilder << TEXT("<Invalid>");
+						bHasError = true;
 					}
 					else
-#endif
 					{
-						StringBuilder << Property->GetFName();
+	#if WITH_EDITOR
+						if (bUseDisplayName)
+						{
+							StringBuilder << Function->GetDisplayNameText().ToString();
+						}
+						else
+	#endif
+						{
+							StringBuilder << Function->GetFName();
+						}
 					}
 				}
-			}
-			else
-			{
-				StringBuilder << TEXT("<Invalid>");
-				bHasError = true;
+				else
+				{
+					StringBuilder << TEXT("<Invalid>");
+					bHasError = true;
+				}
 			}
 		}
 		else
 		{
-			if (LoadedFunctions.IsValidIndex(PathIndex.Index))
+			check(!bIsLoaded);
+			int32 CurrentIndex = PathIndex.Index;
+			bool bFound = false;
+			for (const FMVVMVCompiledFields& Field : CompiledFields)
 			{
-				const FLoadedFunction& LoadedFunction = LoadedFunctions[PathIndex.Index];
-				const UFunction* Function = LoadedFunction.GetFunction();
-				if (!Function)
+				if (PathIndex.bIsProperty)
 				{
-					StringBuilder << TEXT("<Invalid>");
-					bHasError = true;
+					const int32 NumberOfProperties = Field.GetPropertyNum();
+					if (CurrentIndex < NumberOfProperties)
+					{
+						FName FieldName = Field.GetPropertyName(CompiledFieldNames, CurrentIndex);
+						StringBuilder << FieldName;
+						bFound = true;
+						break;
+					}
+					CurrentIndex -= NumberOfProperties;
 				}
 				else
 				{
-#if WITH_EDITOR
-					if (bUseDisplayName)
+					const int32 NumberOfFunctions = Field.GetFunctionNum();
+					if (CurrentIndex < NumberOfFunctions)
 					{
-						StringBuilder << Function->GetDisplayNameText().ToString();
-					}
-					else
-#endif
-					{
-						StringBuilder << Function->GetFName();
+						FName FieldName = Field.GetFunctionName(CompiledFieldNames, CurrentIndex);
+						StringBuilder << FieldName;
+						bFound = true;
+						break;
 					}
+					CurrentIndex -= NumberOfFunctions;
 				}
 			}
-			else
+
+			if (!bFound)
 			{
 				StringBuilder << TEXT("<Invalid>");
 				bHasError = true;
diff --git a/Engine/Plugins/Runtime/ModelViewViewModel/Source/ModelViewViewModel/Public/Bindings/MVVMCompiledBindingLibrary.h b/Engine/Plugins/Runtime/ModelViewViewModel/Source/ModelViewViewModel/Public/Bindings/MVVMCompiledBindingLibrary.h
--- a/Engine/Plugins/Runtime/ModelViewViewModel/Source/ModelViewViewModel/Public/Bindings/MVVMCompiledBindingLibrary.h
+++ b/Engine/Plugins/Runtime/ModelViewViewModel/Source/ModelViewViewModel/Public/Bindings/MVVMCompiledBindingLibrary.h
@@ -37,14 +37,14 @@ public:
 		return ClassOrScriptStruct ? NumberOfFunctions : 0;
 	}
 
-	FName GetPropertyName(TArrayView<FName> Names, int32 Index) const
+	FName GetPropertyName(const TArrayView<const FName> Names, int32 Index) const
 	{
 		check(Index < NumberOfProperties && Index >= 0);
 		check(Names.IsValidIndex(LibraryStartIndex + Index));
 		return Names[LibraryStartIndex + Index];
 	}
 
-	FName GetFunctionName(TArrayView<FName> Names, int32 Index) const
+	FName GetFunctionName(const TArrayView<const FName> Names, int32 Index) const
 	{
 		check(Index < NumberOfFunctions && Index >= 0);
 		check(Names.IsValidIndex(LibraryStartIndex + NumberOfProperties + Index));
@@ -292,7 +292,7 @@ public:
 	 * Execute a binding, in one direction.
 	 * The ExecutionSource is the View (UserWidget) instance.
 	 * The Source of the binding (ViewModel, UserWidget, Widget, ...) can be provided if already known.
-	 * The Destination or the binding (ViewModel, UserWIdget, Widget, ...) can be provided if already known.
+	 * The Destination or the binding (ViewModel, UserWidget, Widget, ...) can be provided if already known.
 	 */
 	TValueOrError<void, EExecutionFailingReason> Execute(UObject* ExecutionSource, const FMVVMVCompiledBinding& Binding, EConversionFunctionType ConversionType) const;
 	TValueOrError<void, EExecutionFailingReason> ExecuteWithSource(UObject* ExecutionSource, const FMVVMVCompiledBinding& Binding, UObject* Source) const;
