MVVM: Add an option to select the event as a binding or as an event. Allow the user to bind a multicast delegate property in 2 different ways. UX could be improved. Tried a combobox but that removed the nice primary button. I tried a ContextMenu but was not clear what double click was going to do.

diff --git a/Engine/Plugins/Runtime/ModelViewViewModel/Source/ModelViewViewModelEditor/Private/Customizations/MVVMPropertyBindingExtension.cpp b/Engine/Plugins/Runtime/ModelViewViewModel/Source/ModelViewViewModelEditor/Private/Customizations/MVVMPropertyBindingExtension.cpp
--- a/Engine/Plugins/Runtime/ModelViewViewModel/Source/ModelViewViewModelEditor/Private/Customizations/MVVMPropertyBindingExtension.cpp
+++ b/Engine/Plugins/Runtime/ModelViewViewModel/Source/ModelViewViewModelEditor/Private/Customizations/MVVMPropertyBindingExtension.cpp
@@ -327,7 +327,7 @@ IPropertyBindingExtension::EDropResult FMVVMPropertyBindingExtension::OnDrop(con
 		{
 			DestinationPropertyPath.SetWidgetName(Widget->GetFName());
 		}
-		EditorSubsystem->SetDestinationPathForBinding(WidgetBlueprint, NewBinding, DestinationPropertyPath);
+		EditorSubsystem->SetDestinationPathForBinding(WidgetBlueprint, NewBinding, DestinationPropertyPath, false);
 
 		return EDropResult::HandledContinue;
 	}
diff --git a/Engine/Plugins/Runtime/ModelViewViewModel/Source/ModelViewViewModelEditor/Private/MVVMEditorSubsystem.cpp b/Engine/Plugins/Runtime/ModelViewViewModel/Source/ModelViewViewModelEditor/Private/MVVMEditorSubsystem.cpp
--- a/Engine/Plugins/Runtime/ModelViewViewModel/Source/ModelViewViewModelEditor/Private/MVVMEditorSubsystem.cpp
+++ b/Engine/Plugins/Runtime/ModelViewViewModel/Source/ModelViewViewModelEditor/Private/MVVMEditorSubsystem.cpp
@@ -470,12 +470,12 @@ void UMVVMEditorSubsystem::SetDestinationToSourceConversionFunction(UWidgetBluep
 	}
 }
 
-void UMVVMEditorSubsystem::SetDestinationPathForBinding(UWidgetBlueprint* WidgetBlueprint, FMVVMBlueprintViewBinding& Binding, FMVVMBlueprintPropertyPath PropertyPath)
+void UMVVMEditorSubsystem::SetDestinationPathForBinding(UWidgetBlueprint* WidgetBlueprint, FMVVMBlueprintViewBinding& Binding, FMVVMBlueprintPropertyPath PropertyPath, bool bAllowEventConversion)
 {
 	if (UMVVMBlueprintView* View = GetView(WidgetBlueprint))
 	{
 		bool bHasConversion = Binding.Conversion.DestinationToSourceConversion != nullptr;
-		bool bEventSupported = UMVVMBlueprintViewEvent::Supports(WidgetBlueprint, PropertyPath);
+		bool bEventSupported = bAllowEventConversion && UMVVMBlueprintViewEvent::Supports(WidgetBlueprint, PropertyPath);
 
 		if (bEventSupported || bHasConversion || Binding.DestinationPath != PropertyPath)
 		{
@@ -661,7 +661,7 @@ void UMVVMEditorSubsystem::GenerateBindToDestinationPathsForBinding(UWidgetBluep
 	}
 }
 
-void UMVVMEditorSubsystem::SetEventPath(UMVVMBlueprintViewEvent* Event, FMVVMBlueprintPropertyPath PropertyPath)
+void UMVVMEditorSubsystem::SetEventPath(UMVVMBlueprintViewEvent* Event, FMVVMBlueprintPropertyPath PropertyPath, bool bRequestBindingConversion)
 {
 	UMVVMBlueprintView* View = Event ? Event->GetOuterUMVVMBlueprintView() : nullptr;
 	if (View)
@@ -675,15 +675,19 @@ void UMVVMEditorSubsystem::SetEventPath(UMVVMBlueprintViewEvent* Event, FMVVMBlu
 
 		bool bSupports = UMVVMBlueprintViewEvent::Supports(WidgetBlueprint, PropertyPath);
 
+		if (bRequestBindingConversion || !bSupports)
+		{
+			FMVVMBlueprintViewBinding& Binding = AddBinding(WidgetBlueprint);
+			SetDestinationPathForBinding(WidgetBlueprint, Binding, PropertyPath, false);
+			View->RemoveEvent(Event);
+		}
 		if (bSupports)
 		{
 			Event->SetEventPath(PropertyPath);
 		}
 		else
 		{
-			FMVVMBlueprintViewBinding& Binding = AddBinding(WidgetBlueprint);
-			SetDestinationPathForBinding(WidgetBlueprint, Binding, PropertyPath);
-			View->RemoveEvent(Event);
+			Event->SetEventPath(FMVVMBlueprintPropertyPath());
 		}
 
 		UE::MVVM::Private::OnEventPostEditChange(Event, EventPath);
diff --git a/Engine/Plugins/Runtime/ModelViewViewModel/Source/ModelViewViewModelEditor/Private/Widgets/BindingEntry/SMVVMBindingRow.cpp b/Engine/Plugins/Runtime/ModelViewViewModel/Source/ModelViewViewModelEditor/Private/Widgets/BindingEntry/SMVVMBindingRow.cpp
--- a/Engine/Plugins/Runtime/ModelViewViewModel/Source/ModelViewViewModelEditor/Private/Widgets/BindingEntry/SMVVMBindingRow.cpp
+++ b/Engine/Plugins/Runtime/ModelViewViewModel/Source/ModelViewViewModelEditor/Private/Widgets/BindingEntry/SMVVMBindingRow.cpp
@@ -125,6 +125,7 @@ TSharedRef<SWidget> SBindingRow::BuildRowWidget()
 					.OnDrop(this, &SBindingRow::HandleFieldSelectorDrop, false)
 					.OnDragEnter(this, &SBindingRow::HandleFieldSelectorDragEnter, false)
 					.ShowContext(false)
+					.CanCreateEvent(true)
 					.ToolTipText(FText::Join(FText::FromString("\n"), LOCTEXT("MVVMDestTooltipHeader", "Destination"), LOCTEXT("MVVMDestTooltipBody", "The property to hook data into.")))
 				]
 			]
@@ -396,7 +397,7 @@ FMVVMLinkedPinValue SBindingRow::GetFieldSelectedValue(bool bSourceToDest) const
 	return FMVVMLinkedPinValue();
 }
 
-void SBindingRow::HandleFieldSelectionChanged(FMVVMLinkedPinValue Value, bool bSource)
+void SBindingRow::HandleFieldSelectionChanged(FMVVMLinkedPinValue Value, SFieldSelectorMenu::ESelectionType SelectionType, bool bSource)
 {
 	UWidgetBlueprint* WidgetBlueprint = GetBlueprint();
 	FMVVMBlueprintViewBinding* ViewBinding = GetThisViewBinding();
@@ -426,7 +427,8 @@ void SBindingRow::HandleFieldSelectionChanged(FMVVMLinkedPinValue Value, bool bS
 		{
 			if (Value.IsPropertyPath())
 			{
-				Subsystem->SetDestinationPathForBinding(WidgetBlueprint, *ViewBinding, Value.GetPropertyPath());
+				const bool bAllowEventConversion = SelectionType == SFieldSelectorMenu::ESelectionType::Event;
+				Subsystem->SetDestinationPathForBinding(WidgetBlueprint, *ViewBinding, Value.GetPropertyPath(), bAllowEventConversion);
 			}
 			else if (Value.IsConversionFunction())
 			{
@@ -438,7 +440,7 @@ void SBindingRow::HandleFieldSelectionChanged(FMVVMLinkedPinValue Value, bool bS
 			}
 			else
 			{
-				Subsystem->SetDestinationPathForBinding(WidgetBlueprint, *ViewBinding, FMVVMBlueprintPropertyPath());
+				Subsystem->SetDestinationPathForBinding(WidgetBlueprint, *ViewBinding, FMVVMBlueprintPropertyPath(), false);
 			}
 		}
 	}
@@ -520,7 +522,7 @@ FReply SBindingRow::HandleFieldSelectorDrop(const FGeometry& MyGeometry, const F
 	}
 	else
 	{
-		Subsystem->SetDestinationPathForBinding(WidgetBlueprint, *ViewBinding, PropertyPath.GetValue());
+		Subsystem->SetDestinationPathForBinding(WidgetBlueprint, *ViewBinding, PropertyPath.GetValue(), false);
 	}
 	return FReply::Handled();
 }
diff --git a/Engine/Plugins/Runtime/ModelViewViewModel/Source/ModelViewViewModelEditor/Private/Widgets/BindingEntry/SMVVMBindingRow.h b/Engine/Plugins/Runtime/ModelViewViewModel/Source/ModelViewViewModelEditor/Private/Widgets/BindingEntry/SMVVMBindingRow.h
--- a/Engine/Plugins/Runtime/ModelViewViewModel/Source/ModelViewViewModelEditor/Private/Widgets/BindingEntry/SMVVMBindingRow.h
+++ b/Engine/Plugins/Runtime/ModelViewViewModel/Source/ModelViewViewModelEditor/Private/Widgets/BindingEntry/SMVVMBindingRow.h
@@ -53,7 +53,7 @@ private:
 
 	FMVVMLinkedPinValue GetFieldSelectedValue(bool bSource) const;
 
-	void HandleFieldSelectionChanged(FMVVMLinkedPinValue Value, bool bSource);
+	void HandleFieldSelectionChanged(FMVVMLinkedPinValue Value, SFieldSelectorMenu::ESelectionType SelectionType, bool bSource);
 
 	FFieldSelectionContext GetSelectedSelectionContext(bool bSource) const;
 
diff --git a/Engine/Plugins/Runtime/ModelViewViewModel/Source/ModelViewViewModelEditor/Private/Widgets/BindingEntry/SMVVMEventRow.cpp b/Engine/Plugins/Runtime/ModelViewViewModel/Source/ModelViewViewModelEditor/Private/Widgets/BindingEntry/SMVVMEventRow.cpp
--- a/Engine/Plugins/Runtime/ModelViewViewModel/Source/ModelViewViewModelEditor/Private/Widgets/BindingEntry/SMVVMEventRow.cpp
+++ b/Engine/Plugins/Runtime/ModelViewViewModel/Source/ModelViewViewModelEditor/Private/Widgets/BindingEntry/SMVVMEventRow.cpp
@@ -89,6 +89,7 @@ TSharedRef<SWidget> SEventRow::BuildRowWidget()
 					.OnDrop(this, &SEventRow::HandleFieldSelectorDrop, true)
 					.OnDragEnter(this, &SEventRow::HandleFieldSelectorDragEnter, true)
 					.ShowContext(false)
+					.CanCreateEvent(true)
 				]
 			]
 
@@ -318,7 +319,7 @@ FMVVMLinkedPinValue SEventRow::GetFieldSelectedValue(bool bEvent) const
 	return FMVVMLinkedPinValue();
 }
 
-void SEventRow::HandleFieldSelectionChanged(FMVVMLinkedPinValue Value, bool bEvent)
+void SEventRow::HandleFieldSelectionChanged(FMVVMLinkedPinValue Value, SFieldSelectorMenu::ESelectionType SelectionType, bool bEvent)
 {
 	UWidgetBlueprint* WidgetBlueprint = GetBlueprint();
 	UMVVMBlueprintViewEvent* Event = GetEvent();
@@ -327,7 +328,8 @@ void SEventRow::HandleFieldSelectionChanged(FMVVMLinkedPinValue Value, bool bEve
 		UMVVMEditorSubsystem* Subsystem = GetEditorSubsystem();
 		if (bEvent)
 		{
-			Subsystem->SetEventPath(Event, Value.IsPropertyPath() ? Value.GetPropertyPath() : FMVVMBlueprintPropertyPath());
+			const bool bRequestBindingConversion = SelectionType == SFieldSelectorMenu::ESelectionType::Binding;
+			Subsystem->SetEventPath(Event, Value.IsPropertyPath() ? Value.GetPropertyPath() : FMVVMBlueprintPropertyPath(), bRequestBindingConversion);
 		}
 		else
 		{
@@ -377,7 +379,7 @@ FReply SEventRow::HandleFieldSelectorDrop(const FGeometry& MyGeometry, const FDr
 
 	if (bEvent)
 	{
-		GetEditorSubsystem()->SetEventPath(Event, PropertyPath.GetValue());
+		GetEditorSubsystem()->SetEventPath(Event, PropertyPath.GetValue(), true);
 	}
 	else
 	{
diff --git a/Engine/Plugins/Runtime/ModelViewViewModel/Source/ModelViewViewModelEditor/Private/Widgets/BindingEntry/SMVVMEventRow.h b/Engine/Plugins/Runtime/ModelViewViewModel/Source/ModelViewViewModelEditor/Private/Widgets/BindingEntry/SMVVMEventRow.h
--- a/Engine/Plugins/Runtime/ModelViewViewModel/Source/ModelViewViewModelEditor/Private/Widgets/BindingEntry/SMVVMEventRow.h
+++ b/Engine/Plugins/Runtime/ModelViewViewModel/Source/ModelViewViewModelEditor/Private/Widgets/BindingEntry/SMVVMEventRow.h
@@ -53,7 +53,7 @@ private:
 
 	FMVVMLinkedPinValue GetFieldSelectedValue(bool bSource) const;
 
-	void HandleFieldSelectionChanged(FMVVMLinkedPinValue Value, bool bSource);
+	void HandleFieldSelectionChanged(FMVVMLinkedPinValue Value, SFieldSelectorMenu::ESelectionType SelectionType, bool bSource);
 
 	FFieldSelectionContext GetSelectedSelectionContext(bool bSource) const;
 
diff --git a/Engine/Plugins/Runtime/ModelViewViewModel/Source/ModelViewViewModelEditor/Private/Widgets/BindingEntry/SMVVMGroupRow.cpp b/Engine/Plugins/Runtime/ModelViewViewModel/Source/ModelViewViewModelEditor/Private/Widgets/BindingEntry/SMVVMGroupRow.cpp
--- a/Engine/Plugins/Runtime/ModelViewViewModel/Source/ModelViewViewModelEditor/Private/Widgets/BindingEntry/SMVVMGroupRow.cpp
+++ b/Engine/Plugins/Runtime/ModelViewViewModel/Source/ModelViewViewModelEditor/Private/Widgets/BindingEntry/SMVVMGroupRow.cpp
@@ -95,7 +95,7 @@ void SGroupRow::SetSelectedWidget(FBindingSource Source)
 					FMVVMBlueprintPropertyPath CurrentPath = Binding->DestinationPath;
 					Source.SetSourceTo(CurrentPath);
 
-					EditorSubsystem->SetDestinationPathForBinding(WidgetBlueprint, *Binding, CurrentPath);
+					EditorSubsystem->SetDestinationPathForBinding(WidgetBlueprint, *Binding, CurrentPath, false);
 				}
 			}
 
@@ -106,7 +106,7 @@ void SGroupRow::SetSelectedWidget(FBindingSource Source)
 					FMVVMBlueprintPropertyPath CurrentPath = Event->GetEventPath();
 					Source.SetSourceTo(CurrentPath);
 
-					EditorSubsystem->SetEventPath(Event, CurrentPath);
+					EditorSubsystem->SetEventPath(Event, CurrentPath, true);
 				}
 			}
 		}
@@ -137,7 +137,7 @@ FReply SGroupRow::AddBinding() const
 			Path.SetViewModelId(Entry->GetGroupAsViewModel());
 		}
 
-		EditorSubsystem->SetDestinationPathForBinding(WidgetBlueprint, Binding, Path);
+		EditorSubsystem->SetDestinationPathForBinding(WidgetBlueprint, Binding, Path, false);
 	}
 
 	return FReply::Handled();
diff --git a/Engine/Plugins/Runtime/ModelViewViewModel/Source/ModelViewViewModelEditor/Private/Widgets/SMVVMEventParameter.cpp b/Engine/Plugins/Runtime/ModelViewViewModel/Source/ModelViewViewModelEditor/Private/Widgets/SMVVMEventParameter.cpp
--- a/Engine/Plugins/Runtime/ModelViewViewModel/Source/ModelViewViewModelEditor/Private/Widgets/SMVVMEventParameter.cpp
+++ b/Engine/Plugins/Runtime/ModelViewViewModel/Source/ModelViewViewModelEditor/Private/Widgets/SMVVMEventParameter.cpp
@@ -168,7 +168,7 @@ void SEventParameter::SetSelectedField(const FMVVMLinkedPinValue& Path)
 	}
 }
 
-void SEventParameter::HandleFieldSelectionChanged(FMVVMLinkedPinValue Value)
+void SEventParameter::HandleFieldSelectionChanged(FMVVMLinkedPinValue Value, SFieldSelectorMenu::ESelectionType SelectionType)
 {
 	SetSelectedField(Value);
 }
diff --git a/Engine/Plugins/Runtime/ModelViewViewModel/Source/ModelViewViewModelEditor/Private/Widgets/SMVVMEventParameter.h b/Engine/Plugins/Runtime/ModelViewViewModel/Source/ModelViewViewModelEditor/Private/Widgets/SMVVMEventParameter.h
--- a/Engine/Plugins/Runtime/ModelViewViewModel/Source/ModelViewViewModelEditor/Private/Widgets/SMVVMEventParameter.h
+++ b/Engine/Plugins/Runtime/ModelViewViewModel/Source/ModelViewViewModelEditor/Private/Widgets/SMVVMEventParameter.h
@@ -39,7 +39,7 @@ private:
 	FMVVMLinkedPinValue OnGetSelectedField() const;
 	void SetSelectedField(const FMVVMLinkedPinValue& Path);
 
-	void HandleFieldSelectionChanged(FMVVMLinkedPinValue Value);
+	void HandleFieldSelectionChanged(FMVVMLinkedPinValue Value, SFieldSelectorMenu::ESelectionType SelectionType);
 	FFieldSelectionContext GetSelectedSelectionContext() const;
 
 	int32 GetCurrentWidgetIndex() const;
diff --git a/Engine/Plugins/Runtime/ModelViewViewModel/Source/ModelViewViewModelEditor/Private/Widgets/SMVVMFieldSelector.cpp b/Engine/Plugins/Runtime/ModelViewViewModel/Source/ModelViewViewModelEditor/Private/Widgets/SMVVMFieldSelector.cpp
--- a/Engine/Plugins/Runtime/ModelViewViewModel/Source/ModelViewViewModelEditor/Private/Widgets/SMVVMFieldSelector.cpp
+++ b/Engine/Plugins/Runtime/ModelViewViewModel/Source/ModelViewViewModelEditor/Private/Widgets/SMVVMFieldSelector.cpp
@@ -35,6 +35,7 @@ void SFieldSelector::Construct(const FArguments& InArgs, const UWidgetBlueprint*
 	OnDragEnterEvent = InArgs._OnDragEnter;
 	OnDropEvent = InArgs._OnDrop;
 	bIsBindingToEvent = InArgs._IsBindingToEvent;
+	bCanCreateEvent = InArgs._CanCreateEvent;
 
 	ChildSlot
 	[
@@ -77,10 +78,11 @@ TSharedRef<SWidget> SFieldSelector::HandleGetMenuContent()
 
 	TSharedRef<SFieldSelectorMenu> Menu = SNew(SFieldSelectorMenu, WidgetBlueprintPtr)
 		.CurrentSelected(CurrentSelected)
-		.OnSelectionChanged(this, &SFieldSelector::HandleFieldSelectionChanged)
+		.OnSelected(this, &SFieldSelector::HandleFieldSelectionChanged)
 		.OnMenuCloseRequested(this, &SFieldSelector::HandleMenuClosed)
 		.SelectionContext(SelectionContext)
 		.IsBindingToEvent(bIsBindingToEvent)
+		.CanCreateEvent(bCanCreateEvent)
 		;
 
 	ComboButton->SetMenuContentWidgetToFocus(Menu->GetWidgetToFocus());
@@ -88,7 +90,7 @@ TSharedRef<SWidget> SFieldSelector::HandleGetMenuContent()
 	return Menu;
 }
 
-void SFieldSelector::HandleFieldSelectionChanged(FMVVMLinkedPinValue LinkedValue)
+void SFieldSelector::HandleFieldSelectionChanged(FMVVMLinkedPinValue LinkedValue, SFieldSelectorMenu::ESelectionType SelectionType)
 {
 	if (ComboButton.IsValid())
 	{
@@ -97,7 +99,7 @@ void SFieldSelector::HandleFieldSelectionChanged(FMVVMLinkedPinValue LinkedValue
 
 	if (OnSelectionChanged.IsBound())
 	{
-		OnSelectionChanged.Execute(LinkedValue);
+		OnSelectionChanged.Execute(LinkedValue, SelectionType);
 	}
 }
 
diff --git a/Engine/Plugins/Runtime/ModelViewViewModel/Source/ModelViewViewModelEditor/Private/Widgets/SMVVMFieldSelector.h b/Engine/Plugins/Runtime/ModelViewViewModel/Source/ModelViewViewModelEditor/Private/Widgets/SMVVMFieldSelector.h
--- a/Engine/Plugins/Runtime/ModelViewViewModel/Source/ModelViewViewModelEditor/Private/Widgets/SMVVMFieldSelector.h
+++ b/Engine/Plugins/Runtime/ModelViewViewModel/Source/ModelViewViewModelEditor/Private/Widgets/SMVVMFieldSelector.h
@@ -31,10 +31,11 @@ public:
 		SLATE_ARGUMENT_DEFAULT(bool, ShowContext) { true };
 		SLATE_ARGUMENT_DEFAULT(bool, IsBindingToEvent) { false };
 		SLATE_EVENT(SFieldDisplay::FOnGetLinkedPinValue, OnGetLinkedValue)
-		SLATE_EVENT(FOnLinkedValueSelectionChanged, OnSelectionChanged)
+		SLATE_EVENT(SFieldSelectorMenu::FOnLinkedValueSelected, OnSelectionChanged)
 		SLATE_EVENT(FOnGetSelectionContext, OnGetSelectionContext)
 		SLATE_EVENT(FOnDrop, OnDrop)
 		SLATE_EVENT(FOnDragEnter, OnDragEnter)
+		SLATE_ARGUMENT_DEFAULT(bool, CanCreateEvent) { false };
 	SLATE_END_ARGS()
 
 	void Construct(const FArguments& InArgs, const UWidgetBlueprint* InWidgetBlueprint);
@@ -46,7 +47,7 @@ private:
 	TSharedRef<SWidget> CreateSourcePanel();
 	TSharedRef<SWidget> HandleGetMenuContent();
 
-	void HandleFieldSelectionChanged(FMVVMLinkedPinValue NewValue);
+	void HandleFieldSelectionChanged(FMVVMLinkedPinValue NewValue, SFieldSelectorMenu::ESelectionType SelectionType);
 	void HandleMenuClosed();
 
 private:
@@ -54,11 +55,12 @@ private:
 	TSharedPtr<SComboButton> ComboButton;
 	TSharedPtr<SFieldDisplay> FieldDisplay;
 	TWeakObjectPtr<const UWidgetBlueprint> WidgetBlueprint;
-	FOnLinkedValueSelectionChanged OnSelectionChanged;
+	SFieldSelectorMenu::FOnLinkedValueSelected OnSelectionChanged;
 	FOnGetSelectionContext OnGetSelectionContext;
 	FOnDrop OnDropEvent;
 	FOnDragEnter OnDragEnterEvent;
 	bool bIsBindingToEvent = false;
+	bool bCanCreateEvent = false;
 }; 
 
 } // namespace UE::MVVM
diff --git a/Engine/Plugins/Runtime/ModelViewViewModel/Source/ModelViewViewModelEditor/Private/Widgets/SMVVMFieldSelectorMenu.cpp b/Engine/Plugins/Runtime/ModelViewViewModel/Source/ModelViewViewModelEditor/Private/Widgets/SMVVMFieldSelectorMenu.cpp
--- a/Engine/Plugins/Runtime/ModelViewViewModel/Source/ModelViewViewModelEditor/Private/Widgets/SMVVMFieldSelectorMenu.cpp
+++ b/Engine/Plugins/Runtime/ModelViewViewModel/Source/ModelViewViewModelEditor/Private/Widgets/SMVVMFieldSelectorMenu.cpp
@@ -8,6 +8,7 @@
 #include "Framework/Views/TableViewMetadata.h"
 #include "Hierarchy/SReadOnlyHierarchyView.h"
 #include "MVVMBlueprintView.h"
+#include "MVVMBlueprintViewEvent.h"
 #include "MVVMEditorSubsystem.h"
 #include "MVVMDeveloperProjectSettings.h"
 #include "MVVMWidgetBlueprintExtension_View.h"
@@ -54,19 +55,22 @@ EFieldVisibility GetFieldVisibilityFlags(EMVVMBindingMode Mode, bool bReadable,
 
 	return Flags;
 }
-}
+} // namespace
 
 void SFieldSelectorMenu::Construct(const FArguments& InArgs, const UWidgetBlueprint* InWidgetBlueprint)
 {
+	SetCanTick(true);
+
 	WidgetBlueprint = InWidgetBlueprint;
 	check(InWidgetBlueprint);
 
-	OnSelectionChanged = InArgs._OnSelectionChanged;
+	OnSelected = InArgs._OnSelected;
 	OnMenuCloseRequested = InArgs._OnMenuCloseRequested;
 	SelectionContext = InArgs._SelectionContext;
+	bCanCreateEvent = InArgs._CanCreateEvent;
 
 	bIsMenuInitialized = false;
-	bIsClearEnabled = InArgs._CurrentSelected.IsSet() && InArgs._CurrentSelected.GetValue().IsValid();
+	const bool bIsClearEnabled = InArgs._CurrentSelected.IsSet() && InArgs._CurrentSelected.GetValue().IsValid();
 
 	// If we're showing conversion functions, we don't want to set the AssignableTo property of SSourceBindingList, because then it will only show exact matches, 
 	// and since we're also showing conversion functions we know that's not what the user wants.
@@ -153,11 +157,20 @@ void SFieldSelectorMenu::Construct(const FArguments& InArgs, const UWidgetBluepr
 			SNew(SHorizontalBox)
 			+ SHorizontalBox::Slot()
 			.AutoWidth()
+			[
+				SAssignNew(EventSelectButton, SPrimaryButton)
+				.OnClicked(this, &SFieldSelectorMenu::HandleEventSelectClicked)
+				.Text(LOCTEXT("EventSelect", "Event"))
+				.ToolTipText(LOCTEXT("SelectorEventTooltip", "Create an event.") )
+			]
+			+ SHorizontalBox::Slot()
+			.AutoWidth()
 			[
 				SNew(SPrimaryButton)
 				.OnClicked(this, &SFieldSelectorMenu::HandleSelectClicked)
 				.IsEnabled(this, &SFieldSelectorMenu::IsSelectEnabled)
 				.Text(LOCTEXT("Select", "Select"))
+				.ToolTipText(LOCTEXT("SelectorBindingTooltip", "Create a binding with the field."))
 			]
 			+ SHorizontalBox::Slot()
 			.Padding(FAppStyle::Get().GetMargin("StandardDialog.SlotPadding"))
@@ -165,7 +178,7 @@ void SFieldSelectorMenu::Construct(const FArguments& InArgs, const UWidgetBluepr
 			[
 				SNew(SButton)
 				.OnClicked(this, &SFieldSelectorMenu::HandleClearClicked)
-				.IsEnabled(this, &SFieldSelectorMenu::IsClearEnabled)
+				.IsEnabled(bIsClearEnabled)
 				.HAlign(HAlign_Center)
 				[
 					SNew(STextBlock)
@@ -204,6 +217,24 @@ void SFieldSelectorMenu::Construct(const FArguments& InArgs, const UWidgetBluepr
 	];
 
 	bIsMenuInitialized = true;
+	UpdateSelection();
+}
+
+void SFieldSelectorMenu::Tick(const FGeometry& AllottedGeometry, const double InCurrentTime, const float InDeltaTime)
+{
+	SCompoundWidget::Tick(AllottedGeometry, InCurrentTime, InDeltaTime);
+	TOptional<FMVVMLinkedPinValue> NewSelection = GetCurrentSelection();
+	if (NewSelection != CurrentSelectedValue)
+	{
+		CurrentSelectedValue = NewSelection;
+		UpdateSelection();
+	}
+}
+
+void SFieldSelectorMenu::UpdateSelection()
+{
+	check(EventSelectButton);
+	EventSelectButton->SetVisibility(IsEventSelectEnabled()? EVisibility::Visible : EVisibility::Collapsed);
 }
 
 TSharedRef<SWidget> SFieldSelectorMenu::GetWidgetToFocus() const
@@ -211,68 +242,66 @@ TSharedRef<SWidget> SFieldSelectorMenu::GetWidgetToFocus() const
 	return SearchBox.ToSharedRef();
 }
 
-bool SFieldSelectorMenu::IsSelectEnabled() const
+TOptional<FMVVMLinkedPinValue> SFieldSelectorMenu::GetCurrentSelection() const
 {
-	if (BindingList.IsValid())
+	TOptional<FMVVMLinkedPinValue> Value;
+	if (ConversionFunctionList.IsValid())
 	{
-		FMVVMBlueprintPropertyPath Path = BindingList->GetSelectedProperty();
-		if (Path.IsValid())
+		TArray<FConversionFunctionValue> Selection;
+		if (ConversionFunctionList.IsValid())
 		{
-			return true;
+			Selection = ConversionFunctionList->GetSelectedItems();
+		}
+		if (Selection.Num() > 0)
+		{
+			Value = FMVVMLinkedPinValue(Selection[0]);
 		}
 	}
 
-	if (ConversionFunctionList.IsValid())
+	if (BindingList.IsValid())
 	{
-		if (ConversionFunctionList->GetSelectedItems().Num() > 0)
+		FMVVMBlueprintPropertyPath SelectedProperty = BindingList->GetSelectedProperty();
+		if (SelectedProperty.IsValid())
 		{
-			return true;
+			ensure(Value.IsSet() == false);
+			Value = FMVVMLinkedPinValue(MoveTemp(SelectedProperty));
 		}
 	}
+	return Value;
+}
 
+bool SFieldSelectorMenu::IsSelectEnabled() const
+{
+	return CurrentSelectedValue.IsSet();
+}
+
+bool SFieldSelectorMenu::IsEventSelectEnabled() const
+{
+	if (bCanCreateEvent && CurrentSelectedValue.IsSet() && CurrentSelectedValue.GetValue().IsPropertyPath())
+	{
+		return UMVVMBlueprintViewEvent::Supports(WidgetBlueprint.Get(), CurrentSelectedValue.GetValue().GetPropertyPath());
+	}
 	return false;
 }
 
 FReply SFieldSelectorMenu::HandleSelectClicked()
 {
-	if (OnSelectionChanged.IsBound())
+	if (OnSelected.IsBound())
 	{
-		bool bSet = false;
-		FMVVMLinkedPinValue Value;
-		if (ConversionFunctionList.IsValid())
-		{
-			TArray<FConversionFunctionValue> Selection;
-			if (ConversionFunctionList.IsValid())
-			{
-				Selection = ConversionFunctionList->GetSelectedItems();
-			}
-			if (Selection.Num() > 0)
-			{
-				Value = FMVVMLinkedPinValue(Selection[0]);
-				bSet = true;
-			}
-		}
-
-		if (BindingList.IsValid())
-		{
-			FMVVMBlueprintPropertyPath SelectedProperty = BindingList->GetSelectedProperty();
-			if (SelectedProperty.IsValid())
-			{
-				ensure(bSet == false);
-				Value = FMVVMLinkedPinValue(MoveTemp(SelectedProperty));
-				bSet = true;
-			}
-		}
-
-		OnSelectionChanged.Execute(MoveTemp(Value));
+		OnSelected.Execute(CurrentSelectedValue.Get(FMVVMLinkedPinValue()), ESelectionType::Binding);
 	}
 
 	return FReply::Handled();
 }
 
-bool SFieldSelectorMenu::IsClearEnabled() const
+FReply SFieldSelectorMenu::HandleEventSelectClicked()
 {
-	return bIsClearEnabled;
+	if (OnSelected.IsBound())
+	{
+		OnSelected.Execute(CurrentSelectedValue.Get(FMVVMLinkedPinValue()), ESelectionType::Event);
+	}
+
+	return FReply::Handled();
 }
 
 FReply SFieldSelectorMenu::HandleClearClicked()
@@ -285,7 +314,7 @@ FReply SFieldSelectorMenu::HandleClearClicked()
 		Source.SetSourceTo(NewProperty);
 	}
 
-	OnSelectionChanged.ExecuteIfBound(FMVVMLinkedPinValue(MoveTemp(NewProperty)));
+	OnSelected.ExecuteIfBound(FMVVMLinkedPinValue(MoveTemp(NewProperty)), ESelectionType::None);
 
 	return FReply::Handled();
 }
@@ -299,12 +328,28 @@ FReply SFieldSelectorMenu::HandleCancelClicked()
 
 void SFieldSelectorMenu::SetPropertyPathSelection(const FMVVMBlueprintPropertyPath& SelectedPath)
 {
-	OnSelectionChanged.ExecuteIfBound(FMVVMLinkedPinValue(SelectedPath));
+	CurrentSelectedValue = FMVVMLinkedPinValue(SelectedPath);
+	if (IsEventSelectEnabled())
+	{
+		HandleEventSelectClicked();
+	}
+	else
+	{
+		HandleSelectClicked();
+	}
 }
 
 void SFieldSelectorMenu::SetConversionFunctionSelection(const FConversionFunctionValue SelectedFunction)
 {
-	OnSelectionChanged.ExecuteIfBound(FMVVMLinkedPinValue(SelectedFunction));
+	CurrentSelectedValue = FMVVMLinkedPinValue(SelectedFunction);
+	if (IsEventSelectEnabled())
+	{
+		HandleEventSelectClicked();
+	}
+	else
+	{
+		HandleSelectClicked();
+	}
 }
 
 TSharedPtr<SFieldSelectorMenu::FConversionFunctionItem> SFieldSelectorMenu::FindConversionFunctionCategory(const TArray<TSharedPtr<FConversionFunctionItem>>& Items, TArrayView<FString> CategoryNameParts) const
diff --git a/Engine/Plugins/Runtime/ModelViewViewModel/Source/ModelViewViewModelEditor/Private/Widgets/SMVVMFieldSelectorMenu.h b/Engine/Plugins/Runtime/ModelViewViewModel/Source/ModelViewViewModelEditor/Private/Widgets/SMVVMFieldSelectorMenu.h
--- a/Engine/Plugins/Runtime/ModelViewViewModel/Source/ModelViewViewModelEditor/Private/Widgets/SMVVMFieldSelectorMenu.h
+++ b/Engine/Plugins/Runtime/ModelViewViewModel/Source/ModelViewViewModelEditor/Private/Widgets/SMVVMFieldSelectorMenu.h
@@ -10,6 +10,8 @@
 #include "Styling/SlateTypes.h"
 #include "Widgets/SCompoundWidget.h"
 
+class SPrimaryButton;
+
 namespace ESelectInfo { enum Type : int; }
 class ITableRow;
 class SSearchBox;
@@ -87,22 +89,30 @@ struct FFieldSelectionContext
 	bool bWritable = true;
 };
 
-DECLARE_DELEGATE_OneParam(FOnLinkedValueSelectionChanged, FMVVMLinkedPinValue);
-
 class SFieldSelectorMenu : public SCompoundWidget
 {
 public:
+	enum class ESelectionType
+	{
+		None,
+		Binding,
+		Event,
+	};
+
 	DECLARE_DELEGATE_RetVal(FFieldSelectionContext, FOnGetFieldSelectionContext);
+	DECLARE_DELEGATE_TwoParams(FOnLinkedValueSelected, FMVVMLinkedPinValue, ESelectionType);
 
 	SLATE_BEGIN_ARGS(SFieldSelectorMenu){}
 		SLATE_ARGUMENT(TOptional<FMVVMLinkedPinValue>, CurrentSelected)
-		SLATE_EVENT(FOnLinkedValueSelectionChanged, OnSelectionChanged)
+		SLATE_EVENT(FOnLinkedValueSelected, OnSelected)
 		SLATE_EVENT(FSimpleDelegate, OnMenuCloseRequested)
 		SLATE_ARGUMENT(FFieldSelectionContext, SelectionContext)
 		SLATE_ARGUMENT_DEFAULT(bool, IsBindingToEvent) { false };
+		SLATE_ARGUMENT_DEFAULT(bool, CanCreateEvent) { false };
 	SLATE_END_ARGS()
 
 	void Construct(const FArguments& InArgs, const UWidgetBlueprint* InWidgetBlueprint);
+	virtual void Tick(const FGeometry& AllottedGeometry, const double InCurrentTime, const float InDeltaTime) override;
 
 	TSharedRef<SWidget> GetWidgetToFocus() const;
 
@@ -152,19 +162,25 @@ private:
 	void HandleEnabledContextToggleChanged(ECheckBoxState CheckState);
 	ECheckBoxState ToggleEnabledContext() const;
 
-	bool IsClearEnabled() const;
+	TOptional<FMVVMLinkedPinValue> GetCurrentSelection() const;
+	void UpdateSelection();
+
 	bool IsSelectEnabled() const;
-	FReply HandleClearClicked();
+	bool IsEventSelectEnabled() const;
+
 	FReply HandleSelectClicked();
+	FReply HandleEventSelectClicked();
+	FReply HandleClearClicked();
 	FReply HandleCancelClicked();
 
 private:
 	TWeakObjectPtr<const UWidgetBlueprint> WidgetBlueprint;
-	FOnLinkedValueSelectionChanged OnSelectionChanged;
+	FOnLinkedValueSelected OnSelected;
 	FSimpleDelegate OnMenuCloseRequested;
 	FFieldSelectionContext SelectionContext;
 
 	TSharedPtr<SSearchBox> SearchBox;
+	TSharedPtr<SPrimaryButton> EventSelectButton;
 
 	//~ viewmodels (binding context panel)
 	TSharedPtr<SListView<FBindingSource>> ViewModelList;
@@ -187,8 +203,10 @@ private:
 	TArray<UE::MVVM::FConversionFunctionValue> ConversionFunctions;
 	TArray<UE::MVVM::FConversionFunctionValue> FilteredConversionFunctions;
 
+	TOptional<FMVVMLinkedPinValue> CurrentSelectedValue;
+
 	bool bIsMenuInitialized = false;
-	bool bIsClearEnabled = false;
+	bool bCanCreateEvent = false;
 }; 
 
 } // namespace UE::MVVM
diff --git a/Engine/Plugins/Runtime/ModelViewViewModel/Source/ModelViewViewModelEditor/Private/Widgets/SMVVMFunctionParameter.cpp b/Engine/Plugins/Runtime/ModelViewViewModel/Source/ModelViewViewModelEditor/Private/Widgets/SMVVMFunctionParameter.cpp
--- a/Engine/Plugins/Runtime/ModelViewViewModel/Source/ModelViewViewModelEditor/Private/Widgets/SMVVMFunctionParameter.cpp
+++ b/Engine/Plugins/Runtime/ModelViewViewModel/Source/ModelViewViewModelEditor/Private/Widgets/SMVVMFunctionParameter.cpp
@@ -196,7 +196,7 @@ void SFunctionParameter::SetSelectedField(const FMVVMLinkedPinValue& Value)
 	}
 }
 
-void SFunctionParameter::HandleFieldSelectionChanged(FMVVMLinkedPinValue Value)
+void SFunctionParameter::HandleFieldSelectionChanged(FMVVMLinkedPinValue Value, SFieldSelectorMenu::ESelectionType SelectionType)
 {
 	SetSelectedField(Value);
 }
diff --git a/Engine/Plugins/Runtime/ModelViewViewModel/Source/ModelViewViewModelEditor/Private/Widgets/SMVVMFunctionParameter.h b/Engine/Plugins/Runtime/ModelViewViewModel/Source/ModelViewViewModelEditor/Private/Widgets/SMVVMFunctionParameter.h
--- a/Engine/Plugins/Runtime/ModelViewViewModel/Source/ModelViewViewModelEditor/Private/Widgets/SMVVMFunctionParameter.h
+++ b/Engine/Plugins/Runtime/ModelViewViewModel/Source/ModelViewViewModelEditor/Private/Widgets/SMVVMFunctionParameter.h
@@ -41,7 +41,7 @@ private:
 	FMVVMLinkedPinValue OnGetSelectedField() const;
 	void SetSelectedField(const FMVVMLinkedPinValue& Path);
 
-	void HandleFieldSelectionChanged(FMVVMLinkedPinValue Value);
+	void HandleFieldSelectionChanged(FMVVMLinkedPinValue Value, SFieldSelectorMenu::ESelectionType SelectionType);
 	FFieldSelectionContext GetSelectedSelectionContext() const;
 
 	int32 GetCurrentWidgetIndex() const;
diff --git a/Engine/Plugins/Runtime/ModelViewViewModel/Source/ModelViewViewModelEditor/Private/Widgets/SMVVMViewBindingPanel.cpp b/Engine/Plugins/Runtime/ModelViewViewModel/Source/ModelViewViewModelEditor/Private/Widgets/SMVVMViewBindingPanel.cpp
--- a/Engine/Plugins/Runtime/ModelViewViewModel/Source/ModelViewViewModelEditor/Private/Widgets/SMVVMViewBindingPanel.cpp
+++ b/Engine/Plugins/Runtime/ModelViewViewModel/Source/ModelViewViewModelEditor/Private/Widgets/SMVVMViewBindingPanel.cpp
@@ -275,7 +275,7 @@ void SBindingsPanel::AddBindingToWidgetList(const TSet<FWidgetReference>& Widget
 						{
 							Path.SetWidgetName(WidgetReference.GetTemplate()->GetFName());
 						}
-						EditorSubsystem->SetDestinationPathForBinding(MVVMExtensionPtr->GetWidgetBlueprint(), Binding, Path);
+						EditorSubsystem->SetDestinationPathForBinding(MVVMExtensionPtr->GetWidgetBlueprint(), Binding, Path, false);
 						AddedBindingId = Binding.BindingId;
 
 						bBindingAdded = true;
diff --git a/Engine/Plugins/Runtime/ModelViewViewModel/Source/ModelViewViewModelEditor/Public/MVVMEditorSubsystem.h b/Engine/Plugins/Runtime/ModelViewViewModel/Source/ModelViewViewModelEditor/Public/MVVMEditorSubsystem.h
--- a/Engine/Plugins/Runtime/ModelViewViewModel/Source/ModelViewViewModelEditor/Public/MVVMEditorSubsystem.h
+++ b/Engine/Plugins/Runtime/ModelViewViewModel/Source/ModelViewViewModelEditor/Public/MVVMEditorSubsystem.h
@@ -78,7 +78,7 @@ public:
 	UE_DEPRECATED(5.4, "SetDestinationToSourceConversionFunction  with a UFunction is deprecated.")
 	void SetDestinationToSourceConversionFunction(UWidgetBlueprint* WidgetBlueprint, FMVVMBlueprintViewBinding& Binding, const UFunction* ConversionFunction);
 	void SetDestinationToSourceConversionFunction(UWidgetBlueprint* WidgetBlueprint, FMVVMBlueprintViewBinding& Binding, FMVVMBlueprintFunctionReference ConversionFunction);
-	void SetDestinationPathForBinding(UWidgetBlueprint* WidgetBlueprint, FMVVMBlueprintViewBinding& Binding, FMVVMBlueprintPropertyPath Field);
+	void SetDestinationPathForBinding(UWidgetBlueprint* WidgetBlueprint, FMVVMBlueprintViewBinding& Binding, FMVVMBlueprintPropertyPath Field, bool bAllowEventConversion);
 	void SetSourcePathForBinding(UWidgetBlueprint* WidgetBlueprint, FMVVMBlueprintViewBinding& Binding, FMVVMBlueprintPropertyPath Field);
 	void OverrideExecutionModeForBinding(UWidgetBlueprint* WidgetBlueprint, FMVVMBlueprintViewBinding& Binding, EMVVMExecutionMode Mode);
 	void ResetExecutionModeForBinding(UWidgetBlueprint* WidgetBlueprint, FMVVMBlueprintViewBinding& Binding);
@@ -87,7 +87,7 @@ public:
 	void SetCompileForBinding(UWidgetBlueprint* WidgetBlueprint, FMVVMBlueprintViewBinding& Binding, bool bCompile);
 	void GenerateBindToDestinationPathsForBinding(UWidgetBlueprint* WidgetBlueprint, FMVVMBlueprintViewBinding& Binding);
 
-	void SetEventPath(UMVVMBlueprintViewEvent* Event, FMVVMBlueprintPropertyPath PropertyPath);
+	void SetEventPath(UMVVMBlueprintViewEvent* Event, FMVVMBlueprintPropertyPath PropertyPath, bool bRequestBindingConversion);
 	void SetEventDestinationPath(UMVVMBlueprintViewEvent* Event, FMVVMBlueprintPropertyPath PropertyPath);
 	void SetEventArgumentPath(UMVVMBlueprintViewEvent* Event, const FMVVMBlueprintPinId& PinId, const FMVVMBlueprintPropertyPath& PropertyPath) const;
 	void SetEnabledForEvent(UMVVMBlueprintViewEvent* Event, bool bEnabled);
