MVVM: Can only select property of loaded UObject ptr.
#rnx

diff --git a/Engine/Plugins/Runtime/ModelViewViewModel/Source/ModelViewViewModelEditor/Private/Widgets/SMVVMViewModelBindingListWidget.cpp b/Engine/Plugins/Runtime/ModelViewViewModel/Source/ModelViewViewModelEditor/Private/Widgets/SMVVMViewModelBindingListWidget.cpp
--- a/Engine/Plugins/Runtime/ModelViewViewModel/Source/ModelViewViewModelEditor/Private/Widgets/SMVVMViewModelBindingListWidget.cpp
+++ b/Engine/Plugins/Runtime/ModelViewViewModel/Source/ModelViewViewModelEditor/Private/Widgets/SMVVMViewModelBindingListWidget.cpp
@@ -261,12 +261,15 @@ FFieldExpander_Bindable::FFieldExpander_Bindable()
 
 TOptional<const UClass*> FFieldExpander_Bindable::CanExpandObject(const FObjectPropertyBase* Property, const UObject* Instance) const
 {
-	TOptional<const UClass*> Result = UE::PropertyViewer::FFieldExpander_Default::CanExpandObject(Property, Instance);
-	if (Result.IsSet() && Result.GetValue())
+	if (CastField<FObjectProperty>(Property))
 	{
-		if (GetDefault<UBlueprintEditorSettings>()->IsClassAllowedOnPin(Result.GetValue()))
+		TOptional<const UClass*> Result = UE::PropertyViewer::FFieldExpander_Default::CanExpandObject(Property, Instance);
+		if (Result.IsSet() && Result.GetValue())
 		{
-			return Result;
+			if (GetDefault<UBlueprintEditorSettings>()->IsClassAllowedOnPin(Result.GetValue()))
+			{
+				return Result;
+			}
 		}
 	}
 	return TOptional<const UClass*>();
