MVVM: fix crash when deleting binding that previously errored.
Make sure the struct details view is cleared when a binding is invalid.

diff --git a/Engine/Plugins/Runtime/ModelViewViewModel/Source/ModelViewViewModelEditor/Private/Customizations/MVVMPropertyPathCustomization.cpp b/Engine/Plugins/Runtime/ModelViewViewModel/Source/ModelViewViewModelEditor/Private/Customizations/MVVMPropertyPathCustomization.cpp
--- a/Engine/Plugins/Runtime/ModelViewViewModel/Source/ModelViewViewModelEditor/Private/Customizations/MVVMPropertyPathCustomization.cpp
+++ b/Engine/Plugins/Runtime/ModelViewViewModel/Source/ModelViewViewModelEditor/Private/Customizations/MVVMPropertyPathCustomization.cpp
@@ -79,7 +79,7 @@ FMVVMBlueprintPropertyPath FPropertyPathCustomization::GetPropertyPathValue() co
 		{
 			continue;
 		}
-		FMVVMBlueprintPropertyPath* PropertyPath = (FMVVMBlueprintPropertyPath*)RawPtr;
+		FMVVMBlueprintPropertyPath* PropertyPath = reinterpret_cast<FMVVMBlueprintPropertyPath*>(RawPtr);
 		if (PropertyPath)
 		{
 			return *PropertyPath;
@@ -121,11 +121,15 @@ FMVVMLinkedPinValue FPropertyPathCustomization::GetFieldValue() const
 					{
 						return FMVVMLinkedPinValue(WidgetBP, ConversionFunction->GetConversionFunction());
 					}
+					else
+					{
+						return FMVVMLinkedPinValue(GetPropertyPathValue());
+					}
 				}
 			}
 		}
 	}
-	return FMVVMLinkedPinValue(GetPropertyPathValue());
+	return FMVVMLinkedPinValue(FMVVMBlueprintPropertyPath());
 }
 
 } // namespace UE::MVVM
diff --git a/Engine/Plugins/Runtime/ModelViewViewModel/Source/ModelViewViewModelEditor/Private/Widgets/SMVVMViewBindingPanel.cpp b/Engine/Plugins/Runtime/ModelViewViewModel/Source/ModelViewViewModelEditor/Private/Widgets/SMVVMViewBindingPanel.cpp
--- a/Engine/Plugins/Runtime/ModelViewViewModel/Source/ModelViewViewModelEditor/Private/Widgets/SMVVMViewBindingPanel.cpp
+++ b/Engine/Plugins/Runtime/ModelViewViewModel/Source/ModelViewViewModelEditor/Private/Widgets/SMVVMViewBindingPanel.cpp
@@ -221,6 +221,7 @@ void SBindingsPanel::OnBindingListSelectionChanged(TConstArrayView<FMVVMBlueprin
 	{
 		// show empty details view
 		DetailsView->SetObject(nullptr);
+		StructDetailsView->SetStructureData(nullptr);
 		DetailContainer->SetContent(DetailsView.ToSharedRef());
 	}
 }
@@ -326,6 +327,7 @@ void SBindingsPanel::RefreshDetailsView()
 		}
 	}
 	DetailsView->SetObject(nullptr);
+	StructDetailsView->SetStructureData(nullptr);
 	DetailContainer->SetContent(DetailsView.ToSharedRef());
 }
 
