Disable adding children to panel hierarchy when it has an MVVM extension.

IClipboardExtension:
- Added virtual CanWidgetAcceptPaste function.

FClipboardExtension:
- Overrides CanWidgetAcceptPaste. Checks if specified widget has any UMVVMBlueprintViewExtensions.

FWidgetBlueprintEditorUtils:
- Added static CanPasteWidgetsExtension function. Checks CanWidgetAcceptPaste on all IClipboardExtensions.

FWidgetBlueprintEditor:
- CanPasteWidgets now also calls FWidgetBlueprintEditorUtils::CanPasteWidgetsExtension.

diff --git a/Engine/Plugins/Runtime/ModelViewViewModel/Source/ModelViewViewModelEditor/Private/Customizations/MVVMClipboardExtension.cpp b/Engine/Plugins/Runtime/ModelViewViewModel/Source/ModelViewViewModelEditor/Private/Customizations/MVVMClipboardExtension.cpp
--- a/Engine/Plugins/Runtime/ModelViewViewModel/Source/ModelViewViewModelEditor/Private/Customizations/MVVMClipboardExtension.cpp
+++ b/Engine/Plugins/Runtime/ModelViewViewModel/Source/ModelViewViewModelEditor/Private/Customizations/MVVMClipboardExtension.cpp
@@ -4,6 +4,7 @@
 
 #include "Exporters/Exporter.h"
 #include "Extensions/MVVMBlueprintViewExtension.h"
+#include "Extensions/MVVMViewBlueprintPanelWidgetExtension.h"
 #include "MVVMDeveloperProjectSettings.h"
 #include "MVVMWidgetBlueprintExtension_View.h"
 #include "WidgetBlueprintEditorUtils.h"
@@ -80,6 +81,25 @@ void FClipboardExtension::ImportDataToWidget(const UWidget* Widget, FName OldWid
 	}
 }
 
+bool FClipboardExtension::CanWidgetAcceptPaste(const UWidget* Widget) const
+{
+	bool bCanWidgetAcceptPaste = true;
+
+	if (const UWidgetBlueprint* WidgetBlueprint = FWidgetBlueprintEditorUtils::GetWidgetBlueprintFromWidget(Widget))
+	{
+		if (const UMVVMWidgetBlueprintExtension_View* ExtensionView = UMVVMWidgetBlueprintExtension_View::GetExtension<UMVVMWidgetBlueprintExtension_View>(WidgetBlueprint))
+		{
+			const TArray<UMVVMBlueprintViewExtension*> BlueprintExtensions = ExtensionView->GetBlueprintExtensionsForWidget(Widget->GetFName());
+			bCanWidgetAcceptPaste = !BlueprintExtensions.ContainsByPredicate([](const UMVVMBlueprintViewExtension* BlueprintExtension) -> bool
+			{
+				return BlueprintExtension->IsA<UMVVMBlueprintViewExtension_PanelWidget>();
+			});
+		}
+	}
+
+	return bCanWidgetAcceptPaste;
+}
+
 FClipboardExtension::FExtensionTextFactory::FExtensionTextFactory()
 	: FCustomizableTextObjectFactory(GWarn)
 {
diff --git a/Engine/Plugins/Runtime/ModelViewViewModel/Source/ModelViewViewModelEditor/Private/Customizations/MVVMClipboardExtension.h b/Engine/Plugins/Runtime/ModelViewViewModel/Source/ModelViewViewModelEditor/Private/Customizations/MVVMClipboardExtension.h
--- a/Engine/Plugins/Runtime/ModelViewViewModel/Source/ModelViewViewModelEditor/Private/Customizations/MVVMClipboardExtension.h
+++ b/Engine/Plugins/Runtime/ModelViewViewModel/Source/ModelViewViewModelEditor/Private/Customizations/MVVMClipboardExtension.h
@@ -21,6 +21,7 @@ class FClipboardExtension
 	virtual void ProcessImportedText(const UWidgetBlueprint* WidgetBlueprint, const FString& TextToImport, UPackage*& TempPackage) override;
 	virtual bool CanImportFromClipboard(const UWidget* Widget) const override;
 	virtual void ImportDataToWidget(const UWidget* Widget, FName OldWidgetName) override;
+	virtual bool CanWidgetAcceptPaste(const UWidget* Widget) const override;
 	//~ End IClipboardExtension overrides
 
 private:
diff --git a/Engine/Source/Editor/UMGEditor/Private/WidgetBlueprintEditor.cpp b/Engine/Source/Editor/UMGEditor/Private/WidgetBlueprintEditor.cpp
--- a/Engine/Source/Editor/UMGEditor/Private/WidgetBlueprintEditor.cpp
+++ b/Engine/Source/Editor/UMGEditor/Private/WidgetBlueprintEditor.cpp
@@ -841,6 +841,11 @@ bool FWidgetBlueprintEditor::CanPasteWidgets()
 		return false;
 	}
 
+	if (!FWidgetBlueprintEditorUtils::CanPasteWidgetsExtension(Widgets))
+	{
+		return false;
+	}
+
 	if ( Widgets.Num() == 1 )
 	{
 		// Always return true here now since we want to support pasting widgets as siblings
diff --git a/Engine/Source/Editor/UMGEditor/Private/WidgetBlueprintEditorUtils.cpp b/Engine/Source/Editor/UMGEditor/Private/WidgetBlueprintEditorUtils.cpp
--- a/Engine/Source/Editor/UMGEditor/Private/WidgetBlueprintEditorUtils.cpp
+++ b/Engine/Source/Editor/UMGEditor/Private/WidgetBlueprintEditorUtils.cpp
@@ -1529,6 +1529,34 @@ bool FWidgetBlueprintEditorUtils::IsAnySelectedWidgetLocked(TSet<FWidgetReferenc
 	return false;
 }
 
+bool FWidgetBlueprintEditorUtils::CanPasteWidgetsExtension(TSet<FWidgetReference> SelectedWidgets)
+{
+	if (!SelectedWidgets.IsEmpty())
+	{
+		IUMGEditorModule& EditorModule = FModuleManager::LoadModuleChecked<IUMGEditorModule>("UMGEditor");
+		const TArrayView<const TSharedPtr<IClipboardExtension>> ClipboardExtensions = EditorModule.GetClipboardExtensibilityManager()->GetExtensions();
+
+		for (const TSharedPtr<IClipboardExtension>& ClipboardExtension : ClipboardExtensions)
+		{
+			if (ensure(ClipboardExtension.IsValid()))
+			{
+				for (const FWidgetReference& SelectedWidget : SelectedWidgets)
+				{
+					if (UWidget* TemplateWidget = SelectedWidget.GetTemplate())
+					{
+						if (!ClipboardExtension->CanWidgetAcceptPaste(TemplateWidget))
+						{
+							return false;
+						}
+					}
+				}
+			}
+		}
+	}
+
+	return true;
+}
+
 UWidget* FWidgetBlueprintEditorUtils::GetWidgetTemplateFromDragDrop(UWidgetBlueprint* Blueprint, UWidgetTree* RootWidgetTree, TSharedPtr<FDragDropOperation>& DragDropOp)
 {
 	UWidget* Widget = nullptr;
diff --git a/Engine/Source/Editor/UMGEditor/Public/IHasClipboardExtensibility.h b/Engine/Source/Editor/UMGEditor/Public/IHasClipboardExtensibility.h
--- a/Engine/Source/Editor/UMGEditor/Public/IHasClipboardExtensibility.h
+++ b/Engine/Source/Editor/UMGEditor/Public/IHasClipboardExtensibility.h
@@ -32,6 +32,7 @@ public:
 	virtual void ProcessImportedText(const UWidgetBlueprint* WidgetBlueprint, const FString& TextToImport, UPackage*& TempPackage) = 0;
 	virtual bool CanImportFromClipboard(const UWidget* Widget) const = 0;
 	virtual void ImportDataToWidget(const UWidget* Widget, FName OldWidgetName) = 0;
+	virtual bool CanWidgetAcceptPaste(const UWidget* Widget) const = 0;
 
 	virtual ~IClipboardExtension() { }
 };
diff --git a/Engine/Source/Editor/UMGEditor/Public/WidgetBlueprintEditorUtils.h b/Engine/Source/Editor/UMGEditor/Public/WidgetBlueprintEditorUtils.h
--- a/Engine/Source/Editor/UMGEditor/Public/WidgetBlueprintEditorUtils.h
+++ b/Engine/Source/Editor/UMGEditor/Public/WidgetBlueprintEditorUtils.h
@@ -52,6 +52,8 @@ public:
 
 	static bool IsAnySelectedWidgetLocked(TSet<FWidgetReference> SelectedWidgets);
 
+	static bool CanPasteWidgetsExtension(TSet<FWidgetReference> SelectedWidgets);
+
 	static UWidget* GetWidgetTemplateFromDragDrop(UWidgetBlueprint* Blueprint, UWidgetTree* RootWidgetTree, TSharedPtr<FDragDropOperation>& DragDropOp);
 
 	static bool IsBindWidgetProperty(const FProperty* InProperty);
