Add designer preview entries to MVVM panel extension.

UMVVMBlueprintViewExtension:
- Added virtual OnPreviewContentChanged function.

UMVVMWidgetBlueprintExtension_View:
- Added GetAllBlueprintExtensions function.

FBlueprintViewDesignerExtension:
- New designer extension which overrides PreviewContentChanged and calls OnPreviewContentChanged on each UMVVMBlueprintViewExtension.

FBlueprintViewDesignerExtensionFactory:
- New designer extension factory which creates a FBlueprintViewDesignerExtension.

UMVVMBlueprintViewExtension_PanelWidget:
- Added NumDesignerPreviewEntries property.
- Added static CreateDesignerPreviewEntries function which clears panels children then adds specified number of entries.
- Overrides OnPreviewContentChanged which calls CreateDesignerPreviewEntries using the preview panel widget.

FModelViewViewModelEditorModule:
- Adds PanelWidgetExtensionFactory on module startup and removes on shutdown.

FMVVMPanelWidgetExtensionCustomizationExtender:
- Added "Num Designer Preview Entries" property row.
- Calls UMVVMBlueprintViewExtension_PanelWidget::CreateDesignerPreviewEntries when entry class, slot template or "Num Designer Preview Entries" properties change.

diff --git a/Engine/Plugins/Runtime/ModelViewViewModel/Source/ModelViewViewModelBlueprint/Private/Extensions/MVVMViewBlueprintPanelWidgetExtension.cpp b/Engine/Plugins/Runtime/ModelViewViewModel/Source/ModelViewViewModelBlueprint/Private/Extensions/MVVMViewBlueprintPanelWidgetExtension.cpp
--- a/Engine/Plugins/Runtime/ModelViewViewModel/Source/ModelViewViewModelBlueprint/Private/Extensions/MVVMViewBlueprintPanelWidgetExtension.cpp
+++ b/Engine/Plugins/Runtime/ModelViewViewModel/Source/ModelViewViewModelBlueprint/Private/Extensions/MVVMViewBlueprintPanelWidgetExtension.cpp
@@ -11,6 +11,7 @@
 #include "MVVMBlueprintViewModelContext.h"
 #include "MVVMPropertyPath.h"
 #include "MVVMWidgetBlueprintExtension_View.h"
+#include "Slate/SObjectWidget.h"
 #include "Templates/ValueOrError.h"
 #include "UObject/UnrealType.h"
 #include "View/MVVMViewClass.h"
@@ -243,6 +244,42 @@ const UMVVMBlueprintView* UMVVMBlueprintViewExtension_PanelWidget::GetEntryWidge
 	return nullptr;
 }
 
+void UMVVMBlueprintViewExtension_PanelWidget::RefreshDesignerPreviewEntries(UPanelWidget* PanelWidget, TSubclassOf<UUserWidget> EntryWidgetClass, UPanelSlot* SlotTemplate, int32 NumDesignerPreviewEntries, bool bFullRebuild)
+{
+	if (ensure(PanelWidget))
+	{
+		if (bFullRebuild || !EntryWidgetClass)
+		{
+			PanelWidget->ClearChildren();
+
+			if (EntryWidgetClass)
+			{
+				for (int32 EntryIndex = 0; EntryIndex < NumDesignerPreviewEntries; ++EntryIndex)
+				{
+					if (UUserWidget* EntryWidget = UUserWidget::CreateWidgetInstance(*PanelWidget, EntryWidgetClass, NAME_None))
+					{
+						PanelWidget->AddChild(EntryWidget, SlotTemplate);
+					}
+				}
+			}
+		}
+		else if (NumDesignerPreviewEntries > PanelWidget->GetChildrenCount())
+		{
+			for (int32 NumToAdd = NumDesignerPreviewEntries - PanelWidget->GetChildrenCount(); NumToAdd > 0; --NumToAdd)
+			{
+				PanelWidget->AddChild(UUserWidget::CreateWidgetInstance(*PanelWidget, EntryWidgetClass, NAME_None), SlotTemplate);
+			}
+		}
+		else if (ensure(NumDesignerPreviewEntries >= 0) && NumDesignerPreviewEntries < PanelWidget->GetChildrenCount())
+		{
+			for (int32 NumToRemove = PanelWidget->GetChildrenCount() - NumDesignerPreviewEntries; NumToRemove > 0; --NumToRemove)
+			{
+				PanelWidget->RemoveChildAt(PanelWidget->GetChildrenCount() - 1);
+			}
+		}
+	}
+}
+
 bool UMVVMBlueprintViewExtension_PanelWidget::WidgetRenamed(FName OldName, FName NewName)
 {
 	if (WidgetName == OldName)
@@ -254,4 +291,16 @@ bool UMVVMBlueprintViewExtension_PanelWidget::WidgetRenamed(FName OldName, FName
 	return false;
 }
 
+void UMVVMBlueprintViewExtension_PanelWidget::OnPreviewContentChanged(TSharedRef<SWidget> NewContent)
+{
+	const SObjectWidget* ObjectWidget = StaticCastSharedPtr<SObjectWidget>(NewContent.ToSharedPtr()).Get();
+	const UUserWidget* PreviewRoot = ObjectWidget ? ObjectWidget->GetWidgetObject() : nullptr;
+
+	if (UPanelWidget* PreviewWidget = PreviewRoot ? Cast<UPanelWidget>(PreviewRoot->GetWidgetFromName(WidgetName)) : nullptr)
+	{
+		constexpr bool bFullRebuild = false;
+		RefreshDesignerPreviewEntries(PreviewWidget, EntryWidgetClass, SlotObj, NumDesignerPreviewEntries, bFullRebuild);
+	}
+}
+
 #undef LOCTEXT_NAMESPACE
\ No newline at end of file
diff --git a/Engine/Plugins/Runtime/ModelViewViewModel/Source/ModelViewViewModelBlueprint/Private/MVVMWidgetBlueprintExtension_View.cpp b/Engine/Plugins/Runtime/ModelViewViewModel/Source/ModelViewViewModelBlueprint/Private/MVVMWidgetBlueprintExtension_View.cpp
--- a/Engine/Plugins/Runtime/ModelViewViewModel/Source/ModelViewViewModelBlueprint/Private/MVVMWidgetBlueprintExtension_View.cpp
+++ b/Engine/Plugins/Runtime/ModelViewViewModel/Source/ModelViewViewModelBlueprint/Private/MVVMWidgetBlueprintExtension_View.cpp
@@ -110,6 +110,23 @@ TArray<UMVVMBlueprintViewExtension*> UMVVMWidgetBlueprintExtension_View::GetBlue
 	return ThisWidgetExtensions;
 }
 
+TArray<UMVVMBlueprintViewExtension*> UMVVMWidgetBlueprintExtension_View::GetAllBlueprintExtensions() const
+{
+	TArray<UMVVMBlueprintViewExtension*> AllExtensions;
+	AllExtensions.Reset(BlueprintExtensions.Num());
+	
+	for (const FMVVMExtensionItem& Extension : BlueprintExtensions)
+	{
+		if (Extension.ExtensionObj)
+		{
+			AllExtensions.Add(Extension.ExtensionObj);
+		}
+	}
+
+	AllExtensions.Shrink();
+	return AllExtensions;
+}
+
 void UMVVMWidgetBlueprintExtension_View::VerifyWidgetExtensions()
 {
 	if (const UWidgetBlueprint* WidgetBlueprint = GetWidgetBlueprint())
diff --git a/Engine/Plugins/Runtime/ModelViewViewModel/Source/ModelViewViewModelBlueprint/Public/Extensions/MVVMBlueprintViewExtension.h b/Engine/Plugins/Runtime/ModelViewViewModel/Source/ModelViewViewModelBlueprint/Public/Extensions/MVVMBlueprintViewExtension.h
--- a/Engine/Plugins/Runtime/ModelViewViewModel/Source/ModelViewViewModelBlueprint/Public/Extensions/MVVMBlueprintViewExtension.h
+++ b/Engine/Plugins/Runtime/ModelViewViewModel/Source/ModelViewViewModelBlueprint/Public/Extensions/MVVMBlueprintViewExtension.h
@@ -38,4 +38,6 @@ public:
 	{ 
 		return false; 
 	}
+
+	virtual void OnPreviewContentChanged(TSharedRef<SWidget> NewContent) {}
 };
\ No newline at end of file
diff --git a/Engine/Plugins/Runtime/ModelViewViewModel/Source/ModelViewViewModelBlueprint/Public/Extensions/MVVMViewBlueprintPanelWidgetExtension.h b/Engine/Plugins/Runtime/ModelViewViewModel/Source/ModelViewViewModelBlueprint/Public/Extensions/MVVMViewBlueprintPanelWidgetExtension.h
--- a/Engine/Plugins/Runtime/ModelViewViewModel/Source/ModelViewViewModelBlueprint/Public/Extensions/MVVMViewBlueprintPanelWidgetExtension.h
+++ b/Engine/Plugins/Runtime/ModelViewViewModel/Source/ModelViewViewModelBlueprint/Public/Extensions/MVVMViewBlueprintPanelWidgetExtension.h
@@ -38,6 +38,7 @@ public:
 	virtual void Precompile(UE::MVVM::Compiler::IMVVMBlueprintViewPrecompile* Compiler, UWidgetBlueprintGeneratedClass* Class) override;
 	virtual void Compile(UE::MVVM::Compiler::IMVVMBlueprintViewCompile* Compiler, UWidgetBlueprintGeneratedClass* Class, UMVVMViewClass* ViewExtension) override;
 	virtual bool WidgetRenamed(FName OldName, FName NewName) override;
+	virtual void OnPreviewContentChanged(TSharedRef<SWidget> NewContent) override;
 	//~ End UMVVMBlueprintViewExtension overrides
 
 	FGuid GetEntryViewModelId() const
@@ -48,6 +49,8 @@ public:
 private:
 	const UMVVMBlueprintView* GetEntryWidgetBlueprintView(const UUserWidget* EntryUserWidget) const;
 
+	static void RefreshDesignerPreviewEntries(UPanelWidget* PanelWidget, TSubclassOf<UUserWidget> EntryWidgetClass, UPanelSlot* SlotTemplate, int32 NumDesignerPreviewEntries, bool bFullRebuild);
+
 private:
 	UPROPERTY()
 	FName WidgetName;
@@ -58,8 +61,11 @@ private:
 	UPROPERTY(EditAnywhere, BlueprintReadWrite, Category = "EntryClass", meta = (AllowPrivateAccess = "true"))
 	TSubclassOf<UUserWidget> EntryWidgetClass;
 
-	UPROPERTY(EditAnywhere, BlueprintReadWrite, Instanced, Category="Slot", meta = (AllowPrivateAccess = "true"))
+	UPROPERTY(EditAnywhere, BlueprintReadWrite, Instanced, Category = "Slot", meta = (AllowPrivateAccess = "true"))
 	TObjectPtr<UPanelSlot> SlotObj;
+	
+	UPROPERTY(EditAnywhere, BlueprintReadWrite, Category = "NumDesignerPreviewEntries", meta = (AllowPrivateAccess = "true", ClampMin = 0, ClampMax = 20))
+	int32 NumDesignerPreviewEntries = 3;
 
 	UPROPERTY()
 	FName PanelPropertyName;
diff --git a/Engine/Plugins/Runtime/ModelViewViewModel/Source/ModelViewViewModelBlueprint/Public/MVVMWidgetBlueprintExtension_View.h b/Engine/Plugins/Runtime/ModelViewViewModel/Source/ModelViewViewModelBlueprint/Public/MVVMWidgetBlueprintExtension_View.h
--- a/Engine/Plugins/Runtime/ModelViewViewModel/Source/ModelViewViewModelBlueprint/Public/MVVMWidgetBlueprintExtension_View.h
+++ b/Engine/Plugins/Runtime/ModelViewViewModel/Source/ModelViewViewModelBlueprint/Public/MVVMWidgetBlueprintExtension_View.h
@@ -95,6 +95,7 @@ public:
 	UMVVMBlueprintViewExtension* CreateBlueprintWidgetExtension(TSubclassOf<UMVVMBlueprintViewExtension> ExtensionClass, FName WidgetName);
 	void RemoveBlueprintWidgetExtension(UMVVMBlueprintViewExtension* ExtensionToRemove, FName WidgetName);
 	TArray<UMVVMBlueprintViewExtension*> GetBlueprintExtensionsForWidget(FName WidgetName) const;
+	TArray<UMVVMBlueprintViewExtension*> GetAllBlueprintExtensions() const;
 
 	void SetFilterSettings(FMVVMViewBindingFilterSettings InFilterSettings);
 	FMVVMViewBindingFilterSettings GetFilterSettings() const
diff --git a/Engine/Plugins/Runtime/ModelViewViewModel/Source/ModelViewViewModelEditor/Private/Customizations/MVVMBlueprintViewDesignerExtension.cpp b/Engine/Plugins/Runtime/ModelViewViewModel/Source/ModelViewViewModelEditor/Private/Customizations/MVVMBlueprintViewDesignerExtension.cpp
new file mode 100644
--- /dev/null
+++ b/Engine/Plugins/Runtime/ModelViewViewModel/Source/ModelViewViewModelEditor/Private/Customizations/MVVMBlueprintViewDesignerExtension.cpp
@@ -0,0 +1,31 @@
+// Copyright Epic Games, Inc. All Rights Reserved.
+
+#include "Customizations/MVVMBlueprintViewDesignerExtension.h"
+
+#include "Extensions/MVVMBlueprintViewExtension.h"
+#include "MVVMWidgetBlueprintExtension_View.h"
+#include "WidgetBlueprintExtension.h"
+
+namespace UE::MVVM
+{
+TSharedRef<FDesignerExtension> FBlueprintViewDesignerExtensionFactory::CreateDesignerExtension() const
+{
+	return StaticCastSharedRef<FBlueprintViewDesignerExtension>(MakeShared<FBlueprintViewDesignerExtension>());
+}
+
+void FBlueprintViewDesignerExtension::PreviewContentChanged(TSharedRef<SWidget> NewContent)
+{
+	const UWidgetBlueprint* WidgetBlueprint = Blueprint.Get();
+
+	if (UMVVMWidgetBlueprintExtension_View* ExtensionView = WidgetBlueprint ? UWidgetBlueprintExtension::GetExtension<UMVVMWidgetBlueprintExtension_View>(WidgetBlueprint) : nullptr)
+	{
+		for (UMVVMBlueprintViewExtension* BlueprintViewExtension : ExtensionView->GetAllBlueprintExtensions())
+		{
+			if (ensure(BlueprintViewExtension))
+			{
+				BlueprintViewExtension->OnPreviewContentChanged(NewContent);
+			}
+		}
+	}
+}
+}
diff --git a/Engine/Plugins/Runtime/ModelViewViewModel/Source/ModelViewViewModelEditor/Private/Customizations/MVVMPanelWidgetExtensionCustomizationExtender.cpp b/Engine/Plugins/Runtime/ModelViewViewModel/Source/ModelViewViewModelEditor/Private/Customizations/MVVMPanelWidgetExtensionCustomizationExtender.cpp
--- a/Engine/Plugins/Runtime/ModelViewViewModel/Source/ModelViewViewModelEditor/Private/Customizations/MVVMPanelWidgetExtensionCustomizationExtender.cpp
+++ b/Engine/Plugins/Runtime/ModelViewViewModel/Source/ModelViewViewModelEditor/Private/Customizations/MVVMPanelWidgetExtensionCustomizationExtender.cpp
@@ -6,6 +6,7 @@
 #include "Components/PanelWidget.h"
 #include "DetailCategoryBuilder.h"
 #include "DetailLayoutBuilder.h"
+#include "Dialogs/Dialogs.h"
 #include "Extensions/MVVMViewBlueprintPanelWidgetExtension.h"
 #include "Kismet2/BlueprintEditorUtils.h"
 #include "MVVMWidgetBlueprintExtension_View.h"
@@ -17,6 +18,7 @@
 #include "ScopedTransaction.h"
 #include "WidgetBlueprint.h"
 #include "WidgetBlueprintEditor.h"
+#include "WidgetBlueprintEditorUtils.h"
 #include "Widgets/Input/SButton.h"
 #include "Widgets/Input/SComboButton.h"
 #include "Widgets/Text/STextBlock.h"
@@ -149,6 +151,7 @@ void FMVVMPanelWidgetExtensionCustomizationExtender::CustomizeDetails(IDetailLay
 							);
 
 							TSharedPtr<IPropertyHandle> SlotPropertyHandle = SlotDetailRow->GetPropertyHandle();
+							SlotPropertyHandle->SetOnChildPropertyValueChanged(FSimpleDelegate::CreateSP(this, &FMVVMPanelWidgetExtensionCustomizationExtender::HandleSlotChildPropertyChanged));
 
 							SlotDetailRow->CustomWidget(true)
 								.NameContent()
@@ -175,6 +178,11 @@ void FMVVMPanelWidgetExtensionCustomizationExtender::CustomizeDetails(IDetailLay
 										})
 									)
 							);
+
+							// "Num Designer Preview Entries" property row
+							TSharedPtr<IPropertyHandle> NumDesignerPreviewEntriesHandle = PanelExtensionObjectHandle->GetChildHandle("NumDesignerPreviewEntries");
+							NumDesignerPreviewEntriesHandle->SetOnPropertyValueChanged(FSimpleDelegate::CreateSP(this, &FMVVMPanelWidgetExtensionCustomizationExtender::HandleNumDesignerPreviewEntriesChanged));
+							MVVMCategory.AddProperty(NumDesignerPreviewEntriesHandle);
 						}
 					}
 				}
@@ -185,19 +193,70 @@ void FMVVMPanelWidgetExtensionCustomizationExtender::CustomizeDetails(IDetailLay
 
 FReply FMVVMPanelWidgetExtensionCustomizationExtender::ModifyExtension()
 {
+	UPanelWidget* WidgetPtr = Widget.Get();
+
 	if (UMVVMBlueprintViewExtension_PanelWidget* PanelExtension = GetPanelWidgetExtension())
 	{
-		if (UPanelWidget* WidgetPtr = Widget.Get())
+		if (WidgetPtr)
 		{
 			GetExtensionViewForSelectedWidgetBlueprint()->RemoveBlueprintWidgetExtension(PanelExtension, WidgetPtr->GetFName());
 			bIsExtensionAdded = false;
+			WidgetPtr->ClearChildren();
 		}
 	}
 	else
 	{
+		// Warn the user that this may result in data loss
+		if (WidgetPtr && WidgetPtr->GetChildrenCount() > 0)
+		{
+			FString ChildNames;
+
+			for (UWidget* ChildWidget : WidgetPtr->GetAllChildren())
+			{
+				if (ensure(ChildWidget))
+				{
+					ChildNames += FString::Printf(TEXT("%s"), *ChildWidget->GetFName().ToString());
+					ChildNames += ChildWidget == WidgetPtr->GetAllChildren().Last() ? TEXT(".") : TEXT(", ");
+				}
+			}
+
+			FFormatNamedArguments Args;
+			Args.Add(TEXT("WidgetName"), FText::FromName(WidgetPtr->GetFName()));
+			Args.Add(TEXT("NumChildren"), FText::AsNumber(WidgetPtr->GetChildrenCount()));
+			Args.Add(TEXT("ChildNames"), FText::FromString(ChildNames));
+			const FText ConfirmDelete = FText::Format(LOCTEXT("ConfirmReplaceWidgetWithVariableInUse", "Adding a viewmodel extension to {WidgetName} will erase its {NumChildren}|plural(one=child, other=children):\n\n{ChildNames}\n\nDo you wish to continue?"), Args);
+
+			FSuppressableWarningDialog::FSetupInfo Info(ConfirmDelete, LOCTEXT("DeletePanelWidgetChildren", "Delete children"), "DeletePanelWidgetChildren_Warning");
+			Info.ConfirmText = LOCTEXT("DeleteChildren_Continue", "Continue");
+			Info.CancelText = LOCTEXT("DeleteChildren_Cancel", "Cancel");
+
+			const FSuppressableWarningDialog DeletePanelWidgetChildren(Info);
+
+			if (DeletePanelWidgetChildren.ShowModal() == FSuppressableWarningDialog::Cancel)
+			{
+				return FReply::Handled();
+			}
+
+			if (const TSharedPtr<FWidgetBlueprintEditor> BPEditor = WidgetBlueprintEditor.Pin())
+			{
+				if (UWidgetBlueprint* Blueprint = BPEditor->GetWidgetBlueprintObj())
+				{
+					TSet<FWidgetReference> ChildWidgets;
+
+					for (UWidget* Child : WidgetPtr->GetAllChildren())
+					{
+						ChildWidgets.Add(BPEditor->GetReferenceFromPreview(Child));
+					}
+
+					FWidgetBlueprintEditorUtils::DeleteWidgets(BPEditor.ToSharedRef(), Blueprint, ChildWidgets);
+				}
+			}
+		}
+
 		CreatePanelWidgetViewExtensionIfNotExisting();
 		bIsExtensionAdded = true;
 	}
+
 	return FReply::Handled();
 }
 
@@ -222,6 +281,17 @@ void FMVVMPanelWidgetExtensionCustomizationExtender::CreatePanelWidgetViewExtens
 	}
 }
 
+void FMVVMPanelWidgetExtensionCustomizationExtender::RefreshDesignerPreviewEntries(bool bFullRebuild)
+{
+	if (UPanelWidget* PanelWidget = Widget.Get())
+	{
+		UMVVMBlueprintViewExtension_PanelWidget* PanelWidgetExtension = GetPanelWidgetExtension();
+		UPanelSlot* SlotTemplate = PanelWidgetExtension ? PanelWidgetExtension->SlotObj : nullptr;
+		const int32 NumDesignerPreviewEntries = PanelWidgetExtension ? PanelWidgetExtension->NumDesignerPreviewEntries : 0;
+		UMVVMBlueprintViewExtension_PanelWidget::RefreshDesignerPreviewEntries(PanelWidget, EntryClass, SlotTemplate, NumDesignerPreviewEntries, bFullRebuild);
+	}
+}
+
 UMVVMBlueprintViewExtension_PanelWidget* FMVVMPanelWidgetExtensionCustomizationExtender::GetPanelWidgetExtension() const
 {
 	if (UMVVMWidgetBlueprintExtension_View* ViewClass = GetExtensionViewForSelectedWidgetBlueprint())
@@ -278,9 +348,9 @@ void FMVVMPanelWidgetExtensionCustomizationExtender::HandleEntryClassChanged(boo
 	EntryClass = EntryClassValue ? *EntryClassValue : nullptr;
 
 	// Update other values that depend on the entry class (only if the cached value actually changed)
-	if (bEntryClassChanged && EntryClass)
+	if (bEntryClassChanged)
 	{
-		if (UUserWidget* EntryCDO = Cast<UUserWidget>(EntryClass->ClassDefaultObject))
+		if (UUserWidget* EntryCDO = EntryClass ? Cast<UUserWidget>(EntryClass->ClassDefaultObject) : nullptr)
 		{
 			EntryWidgetBlueprint = Cast<UWidgetBlueprint>(EntryCDO->GetClass()->ClassGeneratedBy);
 		}
@@ -289,10 +359,24 @@ void FMVVMPanelWidgetExtensionCustomizationExtender::HandleEntryClassChanged(boo
 		if (!bIsInit)
 		{
 			SetEntryViewModel(FGuid(), false);
+			constexpr bool bFullRebuild = true;
+			RefreshDesignerPreviewEntries(bFullRebuild);
 		}
 	}
 }
 
+void FMVVMPanelWidgetExtensionCustomizationExtender::HandleSlotChildPropertyChanged()
+{
+	constexpr bool bFullRebuild = true;
+	RefreshDesignerPreviewEntries(bFullRebuild);
+}
+
+void FMVVMPanelWidgetExtensionCustomizationExtender::HandleNumDesignerPreviewEntriesChanged()
+{
+	constexpr bool bFullRebuild = false;
+	RefreshDesignerPreviewEntries(bFullRebuild);
+}
+
 FText FMVVMPanelWidgetExtensionCustomizationExtender::OnGetSelectedViewModel() const
 {
 	if (const UPanelWidget* WidgetPtr = Widget.Get())
diff --git a/Engine/Plugins/Runtime/ModelViewViewModel/Source/ModelViewViewModelEditor/Private/Customizations/MVVMPanelWidgetExtensionCustomizationExtender.h b/Engine/Plugins/Runtime/ModelViewViewModel/Source/ModelViewViewModelEditor/Private/Customizations/MVVMPanelWidgetExtensionCustomizationExtender.h
--- a/Engine/Plugins/Runtime/ModelViewViewModel/Source/ModelViewViewModelEditor/Private/Customizations/MVVMPanelWidgetExtensionCustomizationExtender.h
+++ b/Engine/Plugins/Runtime/ModelViewViewModel/Source/ModelViewViewModelEditor/Private/Customizations/MVVMPanelWidgetExtensionCustomizationExtender.h
@@ -40,12 +40,21 @@ private:
 	/** Update the cached variables when the entry class property changes. */
 	void HandleEntryClassChanged(bool bIsInit);
 
+	/** Called when a child property on the slot property changes. */
+	void HandleSlotChildPropertyChanged();
+
+	/** Called when the "Num Designer Preview Entries" property changes. */
+	void HandleNumDesignerPreviewEntriesChanged();
+
 	/** Update the cached variables when the number of grid columns changes. */
 	void HandleNumGridColumns();
 
 	/** Create a new MVVMViewBlueprintPanelWidgetExtension for this widget in the blueprint view class. */
 	void CreatePanelWidgetViewExtensionIfNotExisting();
 
+	/** Create preview entries for the selected panel widget. */
+	void RefreshDesignerPreviewEntries(bool bFullRebuild);
+
 	/** Get the MVVMViewBlueprintPanelWidgetExtension for this widget in the blueprint view class. */
 	UMVVMBlueprintViewExtension_PanelWidget* GetPanelWidgetExtension() const;
 
diff --git a/Engine/Plugins/Runtime/ModelViewViewModel/Source/ModelViewViewModelEditor/Private/ModelViewViewModelEditorModule.cpp b/Engine/Plugins/Runtime/ModelViewViewModel/Source/ModelViewViewModelEditor/Private/ModelViewViewModelEditorModule.cpp
--- a/Engine/Plugins/Runtime/ModelViewViewModel/Source/ModelViewViewModelEditor/Private/ModelViewViewModelEditorModule.cpp
+++ b/Engine/Plugins/Runtime/ModelViewViewModel/Source/ModelViewViewModelEditor/Private/ModelViewViewModelEditorModule.cpp
@@ -5,9 +5,10 @@
 #include "Blueprint/WidgetBlueprintGeneratedClass.h"
 #include "BlueprintModes/WidgetBlueprintApplicationMode.h"
 #include "BlueprintModes/WidgetBlueprintApplicationModes.h"
+#include "Customizations/MVVMBlueprintViewDesignerExtension.h"
 #include "Customizations/MVVMBlueprintViewModelContextCustomization.h"
-#include "Customizations/MVVMListViewBaseExtensionCustomizationExtender.h"
 #include "Customizations/MVVMClipboardExtension.h"
+#include "Customizations/MVVMListViewBaseExtensionCustomizationExtender.h"
 #include "Customizations/MVVMPanelWidgetExtensionCustomizationExtender.h"
 #include "Customizations/MVVMPropertyBindingExtension.h"
 #include "Extensions/MVVMBlueprintViewExtension.h"
@@ -26,8 +27,8 @@
 #include "Tabs/MVVMPreviewSourceSummoner.h"
 #include "Tabs/MVVMViewModelSummoner.h"
 #include "ToolMenus.h"
-#include "UObject/AssetRegistryTagsContext.h"
 #include "UMGEditorModule.h"
+#include "UObject/AssetRegistryTagsContext.h"
 #include "WidgetBlueprintEditor.h"
 #include "WidgetDrawerConfig.h"
 #include "Widgets/SMVVMViewBindingPanel.h"
@@ -57,6 +58,9 @@ void FModelViewViewModelEditorModule::StartupModule()
 	PanelWidgetCustomizationExtender = UE::MVVM::FMVVMPanelWidgetExtensionCustomizationExtender::MakeInstance();
 	UMGEditorModule.AddWidgetCustomizationExtender(PanelWidgetCustomizationExtender.ToSharedRef());
 
+	BlueprintViewDesignerExtensionFactory = MakeShared<UE::MVVM::FBlueprintViewDesignerExtensionFactory>();
+	UMGEditorModule.GetDesignerExtensibilityManager()->AddDesignerExtensionFactory(BlueprintViewDesignerExtensionFactory.ToSharedRef());
+
 	UMGEditorModule.RegisterInstancedCustomPropertyTypeLayout(
 		FMVVMBlueprintViewModelContext::StaticStruct()->GetStructPathName()
 		, IUMGEditorModule::FOnGetInstancePropertyTypeCustomizationInstance::CreateStatic(UE::MVVM::FBlueprintViewModelContextDetailCustomization::MakeInstance)
@@ -96,9 +100,11 @@ void FModelViewViewModelEditorModule::ShutdownModule()
 
 	if (IUMGEditorModule* UMGEditorModule = FModuleManager::GetModulePtr<IUMGEditorModule>("UMGEditor"))
 	{
+		UMGEditorModule->GetDesignerExtensibilityManager()->RemoveDesignerExtensionFactory(BlueprintViewDesignerExtensionFactory.ToSharedRef());
 		UMGEditorModule->OnRegisterTabsForEditor().RemoveAll(this);
 		UMGEditorModule->GetClipboardExtensibilityManager()->RemoveExtension(ClipboardExtension.ToSharedRef());
 		UMGEditorModule->GetPropertyBindingExtensibilityManager()->RemoveExtension(PropertyBindingExtension.ToSharedRef());
+
 		if (UObjectInitialized())
 		{
 			UMGEditorModule->UnregisterInstancedCustomPropertyTypeLayout(FMVVMBlueprintViewModelContext::StaticStruct()->GetStructPathName());
diff --git a/Engine/Plugins/Runtime/ModelViewViewModel/Source/ModelViewViewModelEditor/Private/ModelViewViewModelEditorModule.h b/Engine/Plugins/Runtime/ModelViewViewModel/Source/ModelViewViewModelEditor/Private/ModelViewViewModelEditorModule.h
--- a/Engine/Plugins/Runtime/ModelViewViewModel/Source/ModelViewViewModelEditor/Private/ModelViewViewModelEditorModule.h
+++ b/Engine/Plugins/Runtime/ModelViewViewModel/Source/ModelViewViewModelEditor/Private/ModelViewViewModelEditorModule.h
@@ -8,6 +8,7 @@
 namespace UE::MVVM
 {
 	class FAssetTypeActions_ViewModelBlueprint;
+	class FBlueprintViewDesignerExtensionFactory;
 	class FClipboardExtension;
 	class FMVVMListViewBaseExtensionCustomizationExtender;
 	class FMVVMPanelWidgetExtensionCustomizationExtender;
@@ -50,4 +51,5 @@ private:
 	TSharedPtr<UE::MVVM::FAssetTypeActions_ViewModelBlueprint> ViewModelBlueprintActions;
 	TSharedPtr<UE::MVVM::FMVVMListViewBaseExtensionCustomizationExtender> ListViewBaseCustomizationExtender;
 	TSharedPtr<UE::MVVM::FMVVMPanelWidgetExtensionCustomizationExtender> PanelWidgetCustomizationExtender;
+	TSharedPtr<UE::MVVM::FBlueprintViewDesignerExtensionFactory> BlueprintViewDesignerExtensionFactory;
 };
diff --git a/Engine/Plugins/Runtime/ModelViewViewModel/Source/ModelViewViewModelEditor/Public/Customizations/MVVMBlueprintViewDesignerExtension.h b/Engine/Plugins/Runtime/ModelViewViewModel/Source/ModelViewViewModelEditor/Public/Customizations/MVVMBlueprintViewDesignerExtension.h
new file mode 100644
--- /dev/null
+++ b/Engine/Plugins/Runtime/ModelViewViewModel/Source/ModelViewViewModelEditor/Public/Customizations/MVVMBlueprintViewDesignerExtension.h
@@ -0,0 +1,22 @@
+// Copyright Epic Games, Inc. All Rights Reserved.
+
+#include "DesignerExtension.h"
+#include "IHasDesignerExtensibility.h"
+
+namespace UE::MVVM
+{
+class FBlueprintViewDesignerExtensionFactory : public IDesignerExtensionFactory, public TSharedFromThis<FBlueprintViewDesignerExtensionFactory>
+{
+public:
+	FBlueprintViewDesignerExtensionFactory() = default;
+	virtual ~FBlueprintViewDesignerExtensionFactory() = default;
+
+	virtual TSharedRef<FDesignerExtension> CreateDesignerExtension() const override;
+};
+
+class FBlueprintViewDesignerExtension : public FDesignerExtension
+{
+public:
+	virtual void PreviewContentChanged(TSharedRef<SWidget> NewContent) override;
+};
+}
