MVVM: Autocast is supported for conversion function.

diff --git a/Engine/Plugins/Runtime/ModelViewViewModel/Source/ModelViewViewModelBlueprint/Private/Bindings/MVVMConversionFunctionHelper.cpp b/Engine/Plugins/Runtime/ModelViewViewModel/Source/ModelViewViewModelBlueprint/Private/Bindings/MVVMConversionFunctionHelper.cpp
--- a/Engine/Plugins/Runtime/ModelViewViewModel/Source/ModelViewViewModelBlueprint/Private/Bindings/MVVMConversionFunctionHelper.cpp
+++ b/Engine/Plugins/Runtime/ModelViewViewModel/Source/ModelViewViewModelBlueprint/Private/Bindings/MVVMConversionFunctionHelper.cpp
@@ -29,6 +29,7 @@ namespace UE::MVVM::ConversionFunctionHelper
 namespace Private
 {
 	static const FLazyName ConversionFunctionMetadataKey = "ConversionFunction";
+	static const FStringView AutoPromoteFunctionMetadataKey = TEXT("AutoPromoteFunction");
 	static const FLazyName ConversionFunctionCategory = "Conversion Functions";
 
 	UMVVMBlueprintView* GetView(const UBlueprint* Blueprint)
@@ -170,7 +171,10 @@ namespace Private
 				if (Private::IsSystemInputPin(Pin) && Pin->LinkedTo.Num() == 1)
 				{
 					Result = Pin->LinkedTo[0]->GetOwningNode();
-					NodesInPath.Emplace(Result, Pin->LinkedTo[0]);
+					if (Result && !IsAutoPromoteNode(Result))
+					{
+						NodesInPath.Emplace(Result, Pin->LinkedTo[0]);
+					}
 				}
 				return Result;
 			};
@@ -1308,6 +1312,19 @@ void SetMetaData(UEdGraph* NewGraph, FName MetaData, FStringView Value)
 	}
 }
 
+void MarkNodeAsAutoPromote(UEdGraphNode* Node)
+{
+	if (Node && !Node->NodeComment.Contains(Private::AutoPromoteFunctionMetadataKey))
+	{
+		Node->NodeComment.Append(Private::AutoPromoteFunctionMetadataKey);
+	}
+}
+
+bool IsAutoPromoteNode(const UEdGraphNode* Node)
+{
+	return Node && Node->NodeComment.Contains(Private::AutoPromoteFunctionMetadataKey);
+}
+
 } // UE::MVVM::ConversionFunctionHelper
 
 #undef LOCTEXT_NAMESPACE
diff --git a/Engine/Plugins/Runtime/ModelViewViewModel/Source/ModelViewViewModelBlueprint/Private/MVVMBlueprintViewConversionFunction.cpp b/Engine/Plugins/Runtime/ModelViewViewModel/Source/ModelViewViewModelBlueprint/Private/MVVMBlueprintViewConversionFunction.cpp
--- a/Engine/Plugins/Runtime/ModelViewViewModel/Source/ModelViewViewModelBlueprint/Private/MVVMBlueprintViewConversionFunction.cpp
+++ b/Engine/Plugins/Runtime/ModelViewViewModel/Source/ModelViewViewModelBlueprint/Private/MVVMBlueprintViewConversionFunction.cpp
@@ -115,7 +115,24 @@ bool UMVVMBlueprintViewConversionFunction::NeedsWrapperGraphInternal(const UClas
 		const UFunction* Function = ConversionFunction.GetFunction(SkeletalSelfContext);
 		if (ensure(Function))
 		{
-			return SavedPins.Num() > 1 || !UE::MVVM::BindingHelper::IsValidForSimpleRuntimeConversion(Function);
+			bool bNeedsWrapper = SavedPins.Num() > 1 || !UE::MVVM::BindingHelper::IsValidForSimpleRuntimeConversion(Function);
+			if (!bNeedsWrapper)
+			{
+				// Confirms there are no autocast/autopromote node
+				if (ensure(CachedWrapperGraph))
+				{
+					for (UEdGraphNode* Node : CachedWrapperGraph->Nodes)
+					{
+						if (UE::MVVM::ConversionFunctionHelper::IsAutoPromoteNode(Node))
+						{
+							bNeedsWrapper = true;
+							break;
+						}
+					}
+				}
+			}
+
+			return bNeedsWrapper;
 		}
 	}
 	else if (ConversionFunction.GetType() == EMVVMBlueprintFunctionReferenceType::Node)
diff --git a/Engine/Plugins/Runtime/ModelViewViewModel/Source/ModelViewViewModelBlueprint/Private/MVVMConversionFunctionGraphSchema.cpp b/Engine/Plugins/Runtime/ModelViewViewModel/Source/ModelViewViewModelBlueprint/Private/MVVMConversionFunctionGraphSchema.cpp
new file mode 100644
--- /dev/null
+++ b/Engine/Plugins/Runtime/ModelViewViewModel/Source/ModelViewViewModelBlueprint/Private/MVVMConversionFunctionGraphSchema.cpp
@@ -0,0 +1,33 @@
+// Copyright Epic Games, Inc. All Rights Reserved.
+
+
+#include "MVVMConversionFunctionGraphSchema.h"
+#include "Bindings/MVVMConversionFunctionHelper.h"
+#include "EdGraph/EdGraphNode.h"
+#include "EdGraph/EdGraphPin.h"
+
+
+bool UMVVMConversionFunctionGraphSchema::TryCreateConnection(UEdGraphPin* A, UEdGraphPin* B) const
+{
+	bool bResult = Super::TryCreateConnection(A, B);
+
+	if (bResult && !A->LinkedTo.Contains(B))
+	{
+		// ConversionNode or Promotion node was created between the 2 graph pins
+		if (A->LinkedTo.Num() == 1)
+		{
+			if (UEdGraphNode* AutogeneratedNode = A->LinkedTo[0]->GetOwningNode())
+			{
+				for (UEdGraphPin* AutogeneratedPin : AutogeneratedNode->Pins)
+				{
+					if (AutogeneratedPin->LinkedTo.Contains(B))
+					{
+						UE::MVVM::ConversionFunctionHelper::MarkNodeAsAutoPromote(AutogeneratedNode);
+						break;
+					}
+				}
+			}
+		}
+	}
+	return bResult;
+}
diff --git a/Engine/Plugins/Runtime/ModelViewViewModel/Source/ModelViewViewModelBlueprint/Private/MVVMConversionFunctionGraphSchema.h b/Engine/Plugins/Runtime/ModelViewViewModel/Source/ModelViewViewModelBlueprint/Private/MVVMConversionFunctionGraphSchema.h
--- a/Engine/Plugins/Runtime/ModelViewViewModel/Source/ModelViewViewModelBlueprint/Private/MVVMConversionFunctionGraphSchema.h
+++ b/Engine/Plugins/Runtime/ModelViewViewModel/Source/ModelViewViewModelBlueprint/Private/MVVMConversionFunctionGraphSchema.h
@@ -6,7 +6,6 @@
 
 #include "MVVMConversionFunctionGraphSchema.generated.h"
 
-
 /**
  *
  */
@@ -16,13 +15,5 @@ class UMVVMConversionFunctionGraphSchema : public UEdGraphSchema_K2
 	GENERATED_BODY()
 
 public:
-	virtual const FPinConnectionResponse CanCreateConnection(const UEdGraphPin* A, const UEdGraphPin* B) const override
-	{
-		FPinConnectionResponse ConnectionResponse = Super::CanCreateConnection(A, B);
-		if (ConnectionResponse.Response == CONNECT_RESPONSE_MAKE_WITH_CONVERSION_NODE || ConnectionResponse.Response == CONNECT_RESPONSE_MAKE_WITH_PROMOTION)
-		{
-			ConnectionResponse.Response = CONNECT_RESPONSE_DISALLOW;
-		}
-		return ConnectionResponse;
-	}
+	virtual bool TryCreateConnection(UEdGraphPin* A, UEdGraphPin* B) const override;
 };
\ No newline at end of file
diff --git a/Engine/Plugins/Runtime/ModelViewViewModel/Source/ModelViewViewModelBlueprint/Public/Bindings/MVVMConversionFunctionHelper.h b/Engine/Plugins/Runtime/ModelViewViewModel/Source/ModelViewViewModelBlueprint/Public/Bindings/MVVMConversionFunctionHelper.h
--- a/Engine/Plugins/Runtime/ModelViewViewModel/Source/ModelViewViewModelBlueprint/Public/Bindings/MVVMConversionFunctionHelper.h
+++ b/Engine/Plugins/Runtime/ModelViewViewModel/Source/ModelViewViewModelBlueprint/Public/Bindings/MVVMConversionFunctionHelper.h
@@ -86,6 +86,12 @@ namespace UE::MVVM::ConversionFunctionHelper
 
 	/** Add metadata to the Graph/Function. */
 	MODELVIEWVIEWMODELBLUEPRINT_API void SetMetaData(UEdGraph* NewGraph, FName MetaData, FStringView Value);
+
+	/** Mark the node a auto promote. We try to hide those node in the editor. */
+	MODELVIEWVIEWMODELBLUEPRINT_API void MarkNodeAsAutoPromote(UEdGraphNode* Node);
+
+	/** Is the node an auto promote node. */
+	MODELVIEWVIEWMODELBLUEPRINT_API bool IsAutoPromoteNode(const UEdGraphNode* Node);
 } //namespace
 
 #if UE_ENABLE_INCLUDE_ORDER_DEPRECATED_IN_5_2
