Added missing #pragma once.

#rnx

diff --git a/Engine/Plugins/Runtime/ModelViewViewModel/Source/ModelViewViewModelEditor/Public/Customizations/MVVMBlueprintViewDesignerExtension.h b/Engine/Plugins/Runtime/ModelViewViewModel/Source/ModelViewViewModelEditor/Public/Customizations/MVVMBlueprintViewDesignerExtension.h
--- a/Engine/Plugins/Runtime/ModelViewViewModel/Source/ModelViewViewModelEditor/Public/Customizations/MVVMBlueprintViewDesignerExtension.h
+++ b/Engine/Plugins/Runtime/ModelViewViewModel/Source/ModelViewViewModelEditor/Public/Customizations/MVVMBlueprintViewDesignerExtension.h
@@ -1,5 +1,7 @@
 // Copyright Epic Games, Inc. All Rights Reserved.
 
+#pragma once
+
 #include "DesignerExtension.h"
 #include "IHasDesignerExtensibility.h"
 
