PR #11968: Request update to FViewModelClassFilter in SMVVMSelectViewModel.h

diff --git a/Engine/Plugins/Runtime/ModelViewViewModel/Source/ModelViewViewModelEditor/Private/Widgets/SMVVMSelectViewModel.h b/Engine/Plugins/Runtime/ModelViewViewModel/Source/ModelViewViewModelEditor/Private/Widgets/SMVVMSelectViewModel.h
--- a/Engine/Plugins/Runtime/ModelViewViewModel/Source/ModelViewViewModelEditor/Private/Widgets/SMVVMSelectViewModel.h
+++ b/Engine/Plugins/Runtime/ModelViewViewModel/Source/ModelViewViewModelEditor/Private/Widgets/SMVVMSelectViewModel.h
@@ -21,7 +21,7 @@ public:
 	virtual bool IsClassAllowed(const FClassViewerInitializationOptions& InInitOptions, const UClass* InClass, TSharedRef<FClassViewerFilterFuncs> InFilterFuncs) override;
 	virtual bool IsUnloadedClassAllowed(const FClassViewerInitializationOptions& InInitOptions, const TSharedRef<const IUnloadedBlueprintData> InUnloadedClassData, TSharedRef<FClassViewerFilterFuncs> InFilterFuncs) override;
 
-	EClassFlags DisallowedClassFlags = CLASS_HideDropDown | CLASS_Hidden | CLASS_Deprecated | CLASS_Abstract | CLASS_NotPlaceable;
+	EClassFlags DisallowedClassFlags = CLASS_HideDropDown | CLASS_Hidden | CLASS_Deprecated | CLASS_Abstract;
 };
 
 class SMVVMSelectViewModel : public SCompoundWidget
