MVVM: Deprecate bAllowConversionFunctionGeneratedGraphInEditor. The conversion function graph will be recreated everytime the asset is loaded to be in sync with the saved data.

diff --git a/Engine/Plugins/Runtime/ModelViewViewModel/Source/ModelViewViewModelBlueprint/Private/MVVMBlueprintViewConversionFunction.cpp b/Engine/Plugins/Runtime/ModelViewViewModel/Source/ModelViewViewModelBlueprint/Private/MVVMBlueprintViewConversionFunction.cpp
--- a/Engine/Plugins/Runtime/ModelViewViewModel/Source/ModelViewViewModelBlueprint/Private/MVVMBlueprintViewConversionFunction.cpp
+++ b/Engine/Plugins/Runtime/ModelViewViewModel/Source/ModelViewViewModelBlueprint/Private/MVVMBlueprintViewConversionFunction.cpp
@@ -186,7 +186,7 @@ void UMVVMBlueprintViewConversionFunction::Reset()
 {
 	ConversionFunction = FMVVMBlueprintFunctionReference();
 	GraphName = FName();
-	bWrapperGraphTransient = false;
+	bWrapperGraphTransient = true;
 	SavedPins.Reset();
 	SetCachedWrapperGraph(nullptr, nullptr, nullptr);
 }
@@ -200,7 +200,7 @@ void UMVVMBlueprintViewConversionFunction::Initialize(UBlueprint* InContext, FNa
 		ConversionFunction = InFunction;
 		check(GraphName.IsNone()); // the name needs to be set before a GetOrCreateWrapperGraph
 		GraphName = InGraphName;
-		bWrapperGraphTransient = !GetDefault<UMVVMDeveloperProjectSettings>()->bAllowConversionFunctionGeneratedGraphInEditor;
+		bWrapperGraphTransient = true;
 		GetOrCreateWrapperGraphInternal(InContext);
 		SavePinValues(InContext);
 	}
@@ -215,7 +215,7 @@ void UMVVMBlueprintViewConversionFunction::InitializeFromFunction(UBlueprint* In
 		ConversionFunction = FMVVMBlueprintFunctionReference(InContext, InFunction);
 		check(GraphName.IsNone()); // the name needs to be set before a GetOrCreateWrapperGraph
 		GraphName = InGraphName;
-		bWrapperGraphTransient = !GetDefault<UMVVMDeveloperProjectSettings>()->bAllowConversionFunctionGeneratedGraphInEditor;
+		bWrapperGraphTransient = true;
 		GetOrCreateWrapperGraphInternal(InContext);
 		SavePinValues(InContext);
 	}
@@ -240,13 +240,14 @@ void UMVVMBlueprintViewConversionFunction::Deprecation_InitializeFromWrapperGrap
 
 		check(GraphName.IsNone());
 		GraphName = CachedWrapperGraph->GetFName();
-		bWrapperGraphTransient = !GetDefault<UMVVMDeveloperProjectSettings>()->bAllowConversionFunctionGeneratedGraphInEditor;
+		bWrapperGraphTransient = true;
 
 		SavePinValues(SelfContext);
 
 		if (bWrapperGraphTransient && CachedWrapperNode)
 		{
 			SelfContext->FunctionGraphs.RemoveSingle(CachedWrapperGraph);
+			CachedWrapperGraph->SetFlags(RF_Transient);
 		}
 	}
 }
@@ -259,7 +260,7 @@ void UMVVMBlueprintViewConversionFunction::Deprecation_InitializeFromMemberRefer
 
 	check(GraphName.IsNone()); // the name needs to be set before a GetOrCreateWrapperGraph
 	GraphName = InGraphName;
-	bWrapperGraphTransient = !GetDefault<UMVVMDeveloperProjectSettings>()->bAllowConversionFunctionGeneratedGraphInEditor;
+	bWrapperGraphTransient = true;
 
 	// since it is a new object, we can't create a the graph right away
 	UClass* GeneratedClass = SelfContext->SkeletonGeneratedClass ? SelfContext->SkeletonGeneratedClass : SelfContext->GeneratedClass;
@@ -283,7 +284,7 @@ void UMVVMBlueprintViewConversionFunction::Deprecation_SetWrapperGraphName(UBlue
 	if (ensure(SavedPins.Num() == 0) && ensure(GraphName.IsNone()))
 	{
 		GraphName = InGraphName;
-		bWrapperGraphTransient = !GetDefault<UMVVMDeveloperProjectSettings>()->bAllowConversionFunctionGeneratedGraphInEditor;
+		bWrapperGraphTransient = true;
 
 		// since it is a new object, we can't create a the graph right away
 		UClass* GeneratedClass = SelfContext->SkeletonGeneratedClass ? SelfContext->SkeletonGeneratedClass : SelfContext->GeneratedClass;
@@ -360,10 +361,15 @@ UEdGraph* UMVVMBlueprintViewConversionFunction::GetOrCreateIntermediateWrapperGr
 	TObjectPtr<UEdGraph>* FoundGraph = !GraphName.IsNone() ? Context.Blueprint->FunctionGraphs.FindByPredicate([GraphName = GetWrapperGraphName()](const UEdGraph* Other) { return Other->GetFName() == GraphName; }) : nullptr;
 	if (FoundGraph)
 	{
-		ensureMsgf(!IsWrapperGraphTransient(), TEXT("The graph is transient. It should not be saved in the editor."));
 		UBlueprint* NullContext = nullptr; // do not register the callback
 		const_cast<UMVVMBlueprintViewConversionFunction*>(this)->SetCachedWrapperGraph(NullContext, *FoundGraph, UE::MVVM::ConversionFunctionHelper::GetWrapperNode(*FoundGraph));
 		LoadPinValuesInternal(Context.Blueprint);
+
+		// Conversion Function graph are not saved in the editor anymore.
+		check(CachedWrapperGraph == *FoundGraph);
+		Context.Blueprint->FunctionGraphs.RemoveSingle(CachedWrapperGraph);
+		CachedWrapperGraph->SetFlags(RF_Transient);
+
 		return CachedWrapperGraph;
 	}
 	else if (IsValid(Context.Blueprint))
@@ -391,9 +397,13 @@ UEdGraph* UMVVMBlueprintViewConversionFunction::GetOrCreateWrapperGraph(UBluepri
 	TObjectPtr<UEdGraph>* FoundGraph = Blueprint->FunctionGraphs.FindByPredicate([GraphName = GetWrapperGraphName()](const UEdGraph* Other) { return Other->GetFName() == GraphName; });
 	if (FoundGraph)
 	{
-		ensureMsgf(!IsWrapperGraphTransient(), TEXT("The graph is transient. It should not be saved in the editor."));
 		const_cast<UMVVMBlueprintViewConversionFunction*>(this)->SetCachedWrapperGraph(Blueprint, *FoundGraph, UE::MVVM::ConversionFunctionHelper::GetWrapperNode(*FoundGraph));
 		LoadPinValuesInternal(Blueprint);
+
+		// Conversion Function graph are not saved in the editor anymore.
+		check(CachedWrapperGraph == *FoundGraph);
+		Blueprint->FunctionGraphs.RemoveSingle(CachedWrapperGraph);
+		CachedWrapperGraph->SetFlags(RF_Transient);
 	}
 	else if (IsValid(Blueprint))
 	{
diff --git a/Engine/Plugins/Runtime/ModelViewViewModel/Source/ModelViewViewModelBlueprint/Public/MVVMBlueprintViewConversionFunction.h b/Engine/Plugins/Runtime/ModelViewViewModel/Source/ModelViewViewModelBlueprint/Public/MVVMBlueprintViewConversionFunction.h
--- a/Engine/Plugins/Runtime/ModelViewViewModel/Source/ModelViewViewModelBlueprint/Public/MVVMBlueprintViewConversionFunction.h
+++ b/Engine/Plugins/Runtime/ModelViewViewModel/Source/ModelViewViewModelBlueprint/Public/MVVMBlueprintViewConversionFunction.h
@@ -152,7 +152,7 @@ private:
 	TArray<FMVVMBlueprintPin> SavedPins;
 
 	UPROPERTY(VisibleAnywhere, Category = "Viewmodel")
-	bool bWrapperGraphTransient = false;
+	bool bWrapperGraphTransient = true;
 
 	UPROPERTY(Transient, DuplicateTransient)
 	mutable TObjectPtr<UEdGraph> CachedWrapperGraph;
diff --git a/Engine/Plugins/Runtime/ModelViewViewModel/Source/ModelViewViewModelBlueprint/Public/MVVMDeveloperProjectSettings.h b/Engine/Plugins/Runtime/ModelViewViewModel/Source/ModelViewViewModelBlueprint/Public/MVVMDeveloperProjectSettings.h
--- a/Engine/Plugins/Runtime/ModelViewViewModel/Source/ModelViewViewModelBlueprint/Public/MVVMDeveloperProjectSettings.h
+++ b/Engine/Plugins/Runtime/ModelViewViewModel/Source/ModelViewViewModelBlueprint/Public/MVVMDeveloperProjectSettings.h
@@ -137,12 +137,13 @@ public:
 	UPROPERTY(EditAnywhere, config, Category = "UX")
 	bool bShowDeveloperGenerateGraphSettings = true;
 
+	UE_DEPRECATED(5.5, "MVVM AllowConversionFunctionGeneratedGraphInEditor feature is disable. The graphs are now transient.")
 	/**
 	 * When a conversion function requires a wrapper graph, add and save the generated graph to the blueprint.
-	 * It is strongly suggested to to use this feature. It can easly break deprecation and the graph will not auto update with new features.
+	 * It is strongly suggested to not use this feature. It can easly break deprecation and the graph will not auto update with new features.
 	 */
-	UPROPERTY(EditAnywhere, config, Category = "Features")
-	bool bAllowConversionFunctionGeneratedGraphInEditor = false;
+	UPROPERTY()
+	bool bAllowConversionFunctionGeneratedGraphInEditor_DEPRECATED = false;
 
 	/** When binding to a multicast delegate property, allow to create an event. */
 	UPROPERTY(EditAnywhere, config, Category = "Features")
