MVVM: Binding are registered individually to make sure all sources that have the same value (pointer) can be triggered. Use the id because the handle callback will be trigger for each registration. Since the viewmodel pointer set at runtime, we can't really prevent more than one registration.

diff --git a/Engine/Plugins/Runtime/ModelViewViewModel/Source/ModelViewViewModel/Private/View/MVVMView.cpp b/Engine/Plugins/Runtime/ModelViewViewModel/Source/ModelViewViewModel/Private/View/MVVMView.cpp
--- a/Engine/Plugins/Runtime/ModelViewViewModel/Source/ModelViewViewModel/Private/View/MVVMView.cpp
+++ b/Engine/Plugins/Runtime/ModelViewViewModel/Source/ModelViewViewModel/Private/View/MVVMView.cpp
@@ -420,7 +420,7 @@ void UMVVMView::InitializeSourceBindings(FMVVMView_SourceKey SourceKey, bool bRu
 				{
 					if (ensureMsgf(FieldId.IsValid(), TEXT("Invalid field. It failed somewhere in the compiler.")))
 					{
-						UE::FieldNotification::FFieldMulticastDelegate::FDelegate Delegate = UE::FieldNotification::FFieldMulticastDelegate::FDelegate::CreateUObject(this, &UMVVMView::HandledLibraryBindingValueChanged);
+						UE::FieldNotification::FFieldMulticastDelegate::FDelegate Delegate = UE::FieldNotification::FFieldMulticastDelegate::FDelegate::CreateUObject(this, &UMVVMView::HandledLibraryBindingValueChanged, SourceKey);
 						NotifyFieldValueChanged->AddFieldValueChangedDelegate(FieldId.GetFieldId(), MoveTemp(Delegate));
 						++ViewSource.RegisteredCount;
 					}
@@ -594,24 +594,17 @@ bool UMVVMView::EvaluateSource(FMVVMViewClass_SourceKey SourceKey)
 }
 
 
-void UMVVMView::HandledLibraryBindingValueChanged(UObject* InSource, UE::FieldNotification::FFieldId InFieldId)
+void UMVVMView::HandledLibraryBindingValueChanged(UObject* InSource, UE::FieldNotification::FFieldId InFieldId, FMVVMView_SourceKey ViewSourceKey)
 {
 	SCOPE_CYCLE_COUNTER(STAT_UMG_Viewmodel_ExecuteBinding_ValueChanged);
 
 	check(InSource);
 	check(InFieldId.IsValid());
+	check(ViewSourceKey.IsValid());
 
-	if (ensure(GeneratedViewClass))
+	if (ensure(GeneratedViewClass) && ensure(Sources.IsValidIndex(ViewSourceKey.GetIndex())))
 	{
-		int32 ViewSourceIndex = Sources.IndexOfByPredicate([InSource](const FMVVMView_Source& Other){ return Other.Source == InSource; });
-		if (ViewSourceIndex == INDEX_NONE)
-		{
-			ensureMsgf(false, TEXT("No source was found. That means the source was removed but the binding was not."));
-			return;
-		}
-		const FMVVMView_SourceKey ViewSourceKey = FMVVMView_SourceKey(ViewSourceIndex);
-		const FMVVMView_Source& ViewSource = Sources[ViewSourceIndex];
-
+		const FMVVMView_Source& ViewSource = Sources[ViewSourceKey.GetIndex()];
 		if (!ViewSource.bBindingsInitialized || !ViewSource.bSourceInitialized)
 		{
 			// we do not want to run a binding while we are initializing the bindings.
diff --git a/Engine/Plugins/Runtime/ModelViewViewModel/Source/ModelViewViewModel/Public/View/MVVMView.h b/Engine/Plugins/Runtime/ModelViewViewModel/Source/ModelViewViewModel/Public/View/MVVMView.h
--- a/Engine/Plugins/Runtime/ModelViewViewModel/Source/ModelViewViewModel/Public/View/MVVMView.h
+++ b/Engine/Plugins/Runtime/ModelViewViewModel/Source/ModelViewViewModel/Public/View/MVVMView.h
@@ -192,7 +192,7 @@ private:
 	void InitializeSourceBindings(FMVVMView_SourceKey SourceKey, bool bRunAllBindings);
 	void InitializeSourceBindingsCommon();
 	void UninitializeSourceBindings(FMVVMViewClass_SourceKey SourceKey, const FMVVMViewClass_Source& ClassSource, FMVVMView_Source& ViewSource);
-	void HandledLibraryBindingValueChanged(UObject* InSource, UE::FieldNotification::FFieldId InFieldId);
+	void HandledLibraryBindingValueChanged(UObject* InSource, UE::FieldNotification::FFieldId InFieldId, FMVVMView_SourceKey SourceKey);
 	void ExecuteBindingInternal(const FMVVMViewClass_SourceBinding& SourceBinding) const;
 	void ExecuteBindingImmediately(const FMVVMViewClass_Binding& ClassBinding, FMVVMViewClass_BindingKey KeyForLog) const;
 	void ExecuteViewModelBindingsInternal(FMVVMViewClass_SourceKey SourceKey);
