MVVM: Add an event when a conversion function setting changes. The event is used to rebuild the library.

diff --git a/Engine/Plugins/Runtime/ModelViewViewModel/Source/ModelViewViewModelBlueprint/Private/MVVMDeveloperProjectSettings.cpp b/Engine/Plugins/Runtime/ModelViewViewModel/Source/ModelViewViewModelBlueprint/Private/MVVMDeveloperProjectSettings.cpp
--- a/Engine/Plugins/Runtime/ModelViewViewModel/Source/ModelViewViewModelBlueprint/Private/MVVMDeveloperProjectSettings.cpp
+++ b/Engine/Plugins/Runtime/ModelViewViewModel/Source/ModelViewViewModelBlueprint/Private/MVVMDeveloperProjectSettings.cpp
@@ -61,6 +61,22 @@ FText UMVVMDeveloperProjectSettings::GetSectionText() const
 	return LOCTEXT("MVVMProjectSettings", "UMG Model View Viewmodel");
 }
 
+#if WITH_EDITOR
+void UMVVMDeveloperProjectSettings::PostEditChangeChainProperty(FPropertyChangedChainEvent& PropertyChangedEvent)
+{
+	Super::PostEditChangeChainProperty(PropertyChangedEvent);
+
+	const FName PropertyName = PropertyChangedEvent.Property->GetFName();
+	if (PropertyName == GET_MEMBER_NAME_CHECKED(UMVVMDeveloperProjectSettings, ConversionFunctionFilter)
+		|| PropertyName == GET_MEMBER_NAME_CHECKED(UMVVMDeveloperProjectSettings, AllowedClassForConversionFunctions)
+		|| PropertyName == GET_MEMBER_NAME_CHECKED(UMVVMDeveloperProjectSettings, DeniedClassForConversionFunctions)
+		|| PropertyName == GET_MEMBER_NAME_CHECKED(UMVVMDeveloperProjectSettings, DeniedModuleForConversionFunctions))
+	{
+		OnLibrarySettingChanged.Broadcast();
+	}
+}
+#endif
+
 bool UMVVMDeveloperProjectSettings::PropertyHasFiltering(const UStruct* ObjectStruct, const FProperty* Property) const
 {
 	check(ObjectStruct);
diff --git a/Engine/Plugins/Runtime/ModelViewViewModel/Source/ModelViewViewModelBlueprint/Public/MVVMDeveloperProjectSettings.h b/Engine/Plugins/Runtime/ModelViewViewModel/Source/ModelViewViewModelBlueprint/Public/MVVMDeveloperProjectSettings.h
--- a/Engine/Plugins/Runtime/ModelViewViewModel/Source/ModelViewViewModelBlueprint/Public/MVVMDeveloperProjectSettings.h
+++ b/Engine/Plugins/Runtime/ModelViewViewModel/Source/ModelViewViewModelBlueprint/Public/MVVMDeveloperProjectSettings.h
@@ -71,6 +71,9 @@ public:
 
 	virtual FName GetCategoryName() const override;
 	virtual FText GetSectionText() const override;
+#if WITH_EDITOR
+	virtual void PostEditChangeChainProperty(FPropertyChangedChainEvent& PropertyChangedEvent) override;
+#endif
 
 	bool PropertyHasFiltering(const UStruct* ObjectStruct, const FProperty* Property) const;
 	bool IsPropertyAllowed(const UBlueprint* Context, const UStruct* ObjectStruct, const FProperty* Property) const;
@@ -176,6 +179,8 @@ public:
 	UPROPERTY(EditAnywhere, config, Category = "Defaults", meta = (EditCondition = "ConversionFunctionFilter == EMVVMDeveloperConversionFunctionFilterType::AllowedList"))
 	TSet<FName> DeniedModuleForConversionFunctions;
 
+	FSimpleMulticastDelegate OnLibrarySettingChanged;
+
 	/** The default value of UMVVMBlueprintViewSettings::bForceExecuteBindingOnSetSource. */
 	UPROPERTY(EditAnywhere, config, Category = "Defaults")
 	bool bForceExecuteBindingsOnSetSource = false;
diff --git a/Engine/Plugins/Runtime/ModelViewViewModel/Source/ModelViewViewModelEditor/Private/Framework/MVVMConversionFunctionLibrary.cpp b/Engine/Plugins/Runtime/ModelViewViewModel/Source/ModelViewViewModelEditor/Private/Framework/MVVMConversionFunctionLibrary.cpp
--- a/Engine/Plugins/Runtime/ModelViewViewModel/Source/ModelViewViewModelEditor/Private/Framework/MVVMConversionFunctionLibrary.cpp
+++ b/Engine/Plugins/Runtime/ModelViewViewModel/Source/ModelViewViewModelEditor/Private/Framework/MVVMConversionFunctionLibrary.cpp
@@ -87,7 +87,8 @@ public:
 
 FFunctionEntry::FFunctionEntry(FConversionFunctionValue Function)
 	: ConversionFunctionValue(Function)
-{}
+{
+}
 
 FCollection::FCollection()
 {
@@ -103,6 +104,8 @@ FCollection::FCollection()
 
 	FModuleManager::Get().OnModulesChanged().AddRaw(this, &FCollection::HandleModulesChanged);
 	FCoreUObjectDelegates::ReloadCompleteDelegate.AddRaw(this, &FCollection::HandleReloadComplete);
+
+	GetMutableDefault<UMVVMDeveloperProjectSettings>()->OnLibrarySettingChanged.AddRaw(this, &FCollection::Rebuild);
 }
 
 FCollection::~FCollection()
@@ -111,10 +114,13 @@ FCollection::~FCollection()
 
 	if (FAssetRegistryModule* AssetRegistryModule = FModuleManager::GetModulePtr<FAssetRegistryModule>(TEXT("AssetRegistry")))
 	{
-		IAssetRegistry& AssetRegistry = AssetRegistryModule->Get();
-		AssetRegistry.OnAssetAdded().RemoveAll(this);
-		AssetRegistry.OnAssetRemoved().RemoveAll(this);
-		AssetRegistry.OnAssetRenamed().RemoveAll(this);
+		if (AssetRegistryModule->IsValid())
+		{
+			IAssetRegistry& AssetRegistry = AssetRegistryModule->Get();
+			AssetRegistry.OnAssetAdded().RemoveAll(this);
+			AssetRegistry.OnAssetRemoved().RemoveAll(this);
+			AssetRegistry.OnAssetRenamed().RemoveAll(this);
+		}
 	}
 
 	FEditorDelegates::OnAssetsPreDelete.RemoveAll(this);
@@ -125,6 +131,8 @@ FCollection::~FCollection()
 
 	if (UObjectInitialized() && !IsEngineExitRequested())
 	{
+		GetMutableDefault<UMVVMDeveloperProjectSettings>()->OnLibrarySettingChanged.RemoveAll(this);
+
 		for (TPair<FObjectKey, FFunctionContainer> PairClass : ClassOrBlueprintToFunctions)
 		{
 			if (PairClass.Value.bIsUserWidget)
@@ -457,6 +465,11 @@ void FCollection::UnegisterBlueprintCallback(UObject* ObjectKey)
 	}
 }
 
+void FCollection::Rebuild()
+{
+	bRefreshAll = true;
+}
+
 TArray<::UE::MVVM::FConversionFunctionValue> FCollection::GetFunctions(const UWidgetBlueprint* WidgetBlueprint) const
 {
 	return GetFunctions(WidgetBlueprint, nullptr, nullptr);
diff --git a/Engine/Plugins/Runtime/ModelViewViewModel/Source/ModelViewViewModelEditor/Private/Framework/MVVMConversionFunctionLibrary.h b/Engine/Plugins/Runtime/ModelViewViewModel/Source/ModelViewViewModelEditor/Private/Framework/MVVMConversionFunctionLibrary.h
--- a/Engine/Plugins/Runtime/ModelViewViewModel/Source/ModelViewViewModelEditor/Private/Framework/MVVMConversionFunctionLibrary.h
+++ b/Engine/Plugins/Runtime/ModelViewViewModel/Source/ModelViewViewModelEditor/Private/Framework/MVVMConversionFunctionLibrary.h
@@ -56,6 +56,9 @@ public:
 
 	[[nodiscard]] const TSharedPtr<FFunctionEntry> FindFunction(::UE::MVVM::FConversionFunctionValue) const;
 
+	/** Rebuild the library when a setting outside of MVVM changed. */
+	void Rebuild();
+
 private:
 	void RefreshIfNeeded();
 	void Build();
