Fix MVVM Show Graph not working in view binding due to missing graph editor mode set

#tests editor

diff --git a/Engine/Plugins/Runtime/ModelViewViewModel/Source/ModelViewViewModelEditor/Private/Framework/MVVMRowHelper.cpp b/Engine/Plugins/Runtime/ModelViewViewModel/Source/ModelViewViewModelEditor/Private/Framework/MVVMRowHelper.cpp
--- a/Engine/Plugins/Runtime/ModelViewViewModel/Source/ModelViewViewModelEditor/Private/Framework/MVVMRowHelper.cpp
+++ b/Engine/Plugins/Runtime/ModelViewViewModel/Source/ModelViewViewModelEditor/Private/Framework/MVVMRowHelper.cpp
@@ -4,6 +4,7 @@
 
 #include "Blueprint/WidgetTree.h"
 #include "BlueprintEditor.h"
+#include "BlueprintModes/WidgetBlueprintApplicationModes.h"
 #include "Details/WidgetPropertyDragDropOp.h"
 #include "DragAndDrop/DecoratedDragDropOp.h"
 #include "Editor/EditorEngine.h"
@@ -115,6 +116,7 @@ void FRowHelper::ShowBlueprintGraph(FBlueprintEditor* BlueprintEditor, UWidgetBl
 			{
 				if (Graph->HasAnyFlags(RF_Transient))
 				{
+					BlueprintEditor->SetCurrentMode(FWidgetBlueprintApplicationModes::GraphMode);
 					BlueprintEditor->OpenDocument(Graph, FDocumentTracker::OpenNewDocument);
 				}
 				else
