MVVM: When a viewmodel is set, we execute all the bindings related to that source. But since the source might have multiples binding that are identical but with different Field that triggers it, then we should only execute the binding once.

diff --git a/Engine/Plugins/Runtime/ModelViewViewModel/Source/ModelViewViewModel/Private/View/MVVMView.cpp b/Engine/Plugins/Runtime/ModelViewViewModel/Source/ModelViewViewModel/Private/View/MVVMView.cpp
--- a/Engine/Plugins/Runtime/ModelViewViewModel/Source/ModelViewViewModel/Private/View/MVVMView.cpp
+++ b/Engine/Plugins/Runtime/ModelViewViewModel/Source/ModelViewViewModel/Private/View/MVVMView.cpp
@@ -429,12 +429,29 @@ void UMVVMView::InitializeSourceBindings(FMVVMView_SourceKey SourceKey, bool bRu
 		}
 
 		// Run the bindings
-		for (const FMVVMViewClass_SourceBinding& SourceBinding : ClassSource.GetBindings())
+		if (bRunAllBindings)
+		{
+			// Run all the bindings from this source but only once.
+			TArray<FMVVMViewClass_BindingKey, TInlineAllocator<8>> ExecutedBindings;
+			for (const FMVVMViewClass_SourceBinding& SourceBinding : ClassSource.GetBindings())
+			{
+				if (!ExecutedBindings.Contains(SourceBinding.GetBindingKey()))
+				{
+					ExecutedBindings.Add(SourceBinding.GetBindingKey());
+					const FMVVMViewClass_Binding& ClassBinding = GeneratedViewClass->GetBinding(SourceBinding.GetBindingKey());
+					ExecuteBindingImmediately(ClassBinding, SourceBinding.GetBindingKey());
+				}
+			}
+		}
+		else
 		{
-			if (SourceBinding.ExecuteAtInitialization() || bRunAllBindings)
+			for (const FMVVMViewClass_SourceBinding& SourceBinding : ClassSource.GetBindings())
 			{
-				const FMVVMViewClass_Binding& ClassBinding = GeneratedViewClass->GetBinding(SourceBinding.GetBindingKey());
-				ExecuteBindingImmediately(ClassBinding, SourceBinding.GetBindingKey());
+				if (SourceBinding.ExecuteAtInitialization())
+				{
+					const FMVVMViewClass_Binding& ClassBinding = GeneratedViewClass->GetBinding(SourceBinding.GetBindingKey());
+					ExecuteBindingImmediately(ClassBinding, SourceBinding.GetBindingKey());
+				}
 			}
 		}
 
@@ -901,7 +918,9 @@ bool UMVVMView::SetSourceInternal(FMVVMViewClass_SourceKey ClassSourceKey, TScri
 				}
 			}
 
-			// initialize bindings
+			// Initialize bindings and execute them without delay.
+			//A binding might have different source dependencies. The binding might not have the "ExecAtInit" flag for this source.
+			//Force the execution of all the bindings for this source.
 			InitializeSourceBindings(ViewSourceKey, true);
 		}
 
