UE-219574 missing icons on the BindingRow and EventRow

The image path was going toward a file that didn't exist. The code was identical to what can be found in widgetreflector. We noticed the discrepency between the CoreContentRoot between the 2 projects.

To get access to the right directory with the resource, we had to bring the root back up, which forced changes on RootCoreContentDir() accesses for the button style.

#tests Tested the presence of the icon for properties and events. Also verified the normal/hovered/pressed effect on the 2 buttons with NoStyleComboButton. I didn't see any other accesses to the MVVMEditorStyle.
#rnx

diff --git a/Engine/Plugins/Runtime/ModelViewViewModel/Source/ModelViewViewModelEditor/Private/Styling/MVVMEditorStyle.cpp b/Engine/Plugins/Runtime/ModelViewViewModel/Source/ModelViewViewModelEditor/Private/Styling/MVVMEditorStyle.cpp
--- a/Engine/Plugins/Runtime/ModelViewViewModel/Source/ModelViewViewModelEditor/Private/Styling/MVVMEditorStyle.cpp
+++ b/Engine/Plugins/Runtime/ModelViewViewModel/Source/ModelViewViewModelEditor/Private/Styling/MVVMEditorStyle.cpp
@@ -30,7 +30,7 @@ FMVVMEditorStyle::FMVVMEditorStyle()
 	if (ensure(MVVMPlugin))
 	{
 		SetContentRoot(MVVMPlugin->GetContentDir() / TEXT("Editor"));
-		SetCoreContentRoot(FPaths::EngineContentDir() / TEXT("Editor") / TEXT("Slate"));
+		SetCoreContentRoot(FPaths::EngineContentDir());
 	}
 
 	// Class Icons
@@ -66,9 +66,9 @@ FMVVMEditorStyle::FMVVMEditorStyle()
 	);
 
 	FButtonStyle NoStyleComboButtonButtonStyle = FButtonStyle()
-		.SetNormal(FSlateBoxBrush(RootToCoreContentDir("Common/ButtonHoverHint.png"), FMargin(4.0f / 16.0f), FLinearColor(1.0f, 1.0f, 1.0f, 0.15f)))
-		.SetHovered(FSlateBoxBrush(RootToCoreContentDir("Common/ButtonHoverHint.png"), FMargin(4.0f / 16.0f), FLinearColor(1.0f, 1.0f, 1.0f, 0.25f)))
-		.SetPressed(FSlateBoxBrush(RootToCoreContentDir("Common/ButtonHoverHint.png"), FMargin(4.0f / 16.0f), FLinearColor(1.0f, 1.0f, 1.0f, 0.30f)))
+		.SetNormal(FSlateBoxBrush(RootToCoreContentDir("Editor/Slate/Common/ButtonHoverHint.png"), FMargin(4.0f / 16.0f), FLinearColor(1.0f, 1.0f, 1.0f, 0.15f)))
+		.SetHovered(FSlateBoxBrush(RootToCoreContentDir("Editor/Slate/Common/ButtonHoverHint.png"), FMargin(4.0f / 16.0f), FLinearColor(1.0f, 1.0f, 1.0f, 0.25f)))
+		.SetPressed(FSlateBoxBrush(RootToCoreContentDir("Editor/Slate/Common/ButtonHoverHint.png"), FMargin(4.0f / 16.0f), FLinearColor(1.0f, 1.0f, 1.0f, 0.30f)))
 		.SetNormalPadding(FMargin(0.0f, 0.0f, 0.0f, 1.0f))
 		.SetPressedPadding(FMargin(0.0f, 1.0f, 0.0f, 0.0f));
 	Set("NoStyleComboButton", FComboButtonStyle()
@@ -103,7 +103,7 @@ FMVVMEditorStyle::FMVVMEditorStyle()
 	Set("ConversionFunction.DestToSource", new IMAGE_BRUSH_SVG("Slate/ConversionFunction_DestToSource", Icon16x16));
 	Set("ConversionFunction.SourceToDest", new IMAGE_BRUSH_SVG("Slate/ConversionFunction_SourceToDest", Icon16x16));
 
-	Set("Icon.Ellipsis", new CORE_IMAGE_BRUSH_SVG("Starship/Common/ellipsis-vertical-narrow", FVector2D(6.0, 24.0)));
+	Set("Icon.Ellipsis", new CORE_IMAGE_BRUSH_SVG("Slate/Starship/Common/ellipsis-vertical-narrow", FVector2D(6.0, 24.0)));
 
 	Set("FieldSelector.ComboButton", 
 		FComboButtonStyle(FAppStyle::Get().GetWidgetStyle<FComboButtonStyle>("ComboButton"))
@@ -120,7 +120,6 @@ FMVVMEditorStyle::~FMVVMEditorStyle()
 	FSlateStyleRegistry::UnRegisterSlateStyle(*this);
 }
 
-
 void FMVVMEditorStyle::CreateInstance()
 {
 	Instance.Reset(new FMVVMEditorStyle);
