MVVM: Add Callback on the WidgetBlueprintGeneratedClass to collect save override. The generated function are transient in editor but save during the cooking process.

diff --git a/Engine/Plugins/Runtime/ModelViewViewModel/Source/ModelViewViewModelBlueprint/Private/MVVMViewBlueprintCompiler.cpp b/Engine/Plugins/Runtime/ModelViewViewModel/Source/ModelViewViewModelBlueprint/Private/MVVMViewBlueprintCompiler.cpp
--- a/Engine/Plugins/Runtime/ModelViewViewModel/Source/ModelViewViewModelBlueprint/Private/MVVMViewBlueprintCompiler.cpp
+++ b/Engine/Plugins/Runtime/ModelViewViewModel/Source/ModelViewViewModelBlueprint/Private/MVVMViewBlueprintCompiler.cpp
@@ -1144,7 +1144,8 @@ void FMVVMViewBlueprintCompiler::CreatePublicFunctionsDeclaration(const FWidgetB
 			}
 
 			BlueprintView->TemporaryGraph.Add(Setter.SetterGraph);
-			FunctionPermissionsToAdd.Add(Setter.SetterGraph->GetFName());
+			FunctionPermissionsToAdd.AddUnique(Setter.SetterGraph->GetFName());
+			GeneratedFunctions.AddUnique(Setter.SetterGraph->GetFName());
 
 			UE::MVVM::FunctionGraphHelper::AddFunctionArgument(Setter.SetterGraph, const_cast<UClass*>(Setter.Class), "Viewmodel");
 		}
@@ -1460,6 +1461,7 @@ void FMVVMViewBlueprintCompiler::CreateIntermediateGraphFunctions(const FWidgetB
 
 			// the new path can only be set later, once the function is compiled.
 			GeneratedDestination->GeneratedFunctionName = GeneratedSetterGraph->GetFName();
+			GeneratedFunctions.Add(GeneratedSetterGraph->GetFName());
 		}
 	}
 
@@ -1481,6 +1483,7 @@ void FMVVMViewBlueprintCompiler::CreateIntermediateGraphFunctions(const FWidgetB
 				{
 					Context.AddGeneratedFunctionGraph(WrapperGraph);
 				}
+				GeneratedFunctions.Add(WrapperGraph->GetFName());
 			}
 		}
 	}
@@ -1490,12 +1493,14 @@ void FMVVMViewBlueprintCompiler::CreateIntermediateGraphFunctions(const FWidgetB
 	{
 		UMVVMBlueprintViewEvent* EventPtr = Event->Event.Get();
 		UEdGraph* WrapperGraph = EventPtr->GetOrCreateWrapperGraph();
-		ensure(WrapperGraph);
-
-		bool bAlreadyContained = WidgetBlueprintCompilerContext.Blueprint->FunctionGraphs.Contains(WrapperGraph);
-		if (ensure(!bAlreadyContained))
+		if (ensure(WrapperGraph))
 		{
-			Context.AddGeneratedFunctionGraph(WrapperGraph);
+			bool bAlreadyContained = WidgetBlueprintCompilerContext.Blueprint->FunctionGraphs.Contains(WrapperGraph);
+			if (ensure(!bAlreadyContained))
+			{
+				Context.AddGeneratedFunctionGraph(WrapperGraph);
+				GeneratedFunctions.Add(WrapperGraph->GetFName());
+			}
 		}
 	}
 }
@@ -1595,6 +1600,12 @@ bool FMVVMViewBlueprintCompiler::Compile(UWidgetBlueprintGeneratedClass* Class,
 }
 
 
+TArray<FName> FMVVMViewBlueprintCompiler::GetGeneratedFunctions() const
+{
+	return GeneratedFunctions;
+}
+
+
 void FMVVMViewBlueprintCompiler::FixFieldPathContext(UWidgetBlueprintGeneratedClass* Class)
 {
 	TSharedRef<FCompilerBindingSource>* SelfSourcePtr = NeededBindingSources.FindByPredicate([](const TSharedRef<FCompilerBindingSource>& Other)
diff --git a/Engine/Plugins/Runtime/ModelViewViewModel/Source/ModelViewViewModelBlueprint/Private/MVVMViewBlueprintCompiler.h b/Engine/Plugins/Runtime/ModelViewViewModel/Source/ModelViewViewModelBlueprint/Private/MVVMViewBlueprintCompiler.h
--- a/Engine/Plugins/Runtime/ModelViewViewModel/Source/ModelViewViewModelBlueprint/Private/MVVMViewBlueprintCompiler.h
+++ b/Engine/Plugins/Runtime/ModelViewViewModel/Source/ModelViewViewModelBlueprint/Private/MVVMViewBlueprintCompiler.h
@@ -68,6 +68,8 @@ public:
 	/** Compile the library and fill the view and viewclass */
 	bool Compile(UWidgetBlueprintGeneratedClass* Class, UMVVMViewClass* ViewExtension);
 
+	/** Get the list of generated function during the compilation process. */
+	TArray<FName> GetGeneratedFunctions() const;
 
 	static void TestGenerateSetter(const UBlueprint* Context, FStringView ObjectName, FStringView FieldPath, FStringView FunctionName);
 
@@ -349,6 +351,11 @@ private:
 	 */
 	TArray<FName> FunctionPermissionsToAdd;
 
+	/**
+	 * List of generated function
+	 */
+	TArray<FName> GeneratedFunctions;
+
 private:
 	TMap<FName, UWidget*> WidgetNameToWidgetPointerMap;
 	FWidgetBlueprintCompilerContext& WidgetBlueprintCompilerContext;
diff --git a/Engine/Plugins/Runtime/ModelViewViewModel/Source/ModelViewViewModelBlueprint/Private/MVVMWidgetBlueprintExtension_View.cpp b/Engine/Plugins/Runtime/ModelViewViewModel/Source/ModelViewViewModelBlueprint/Private/MVVMWidgetBlueprintExtension_View.cpp
--- a/Engine/Plugins/Runtime/ModelViewViewModel/Source/ModelViewViewModelBlueprint/Private/MVVMWidgetBlueprintExtension_View.cpp
+++ b/Engine/Plugins/Runtime/ModelViewViewModel/Source/ModelViewViewModelBlueprint/Private/MVVMWidgetBlueprintExtension_View.cpp
@@ -252,6 +252,7 @@ void UMVVMWidgetBlueprintExtension_View::HandleBeginCompilation(FWidgetBlueprint
 	}
 
 	CurrentCompilerContext.Reset();
+	GeneratedFunctions.Reset();
 	if (BlueprintView)
 	{
 		BlueprintView->ResetBindingMessages();
@@ -346,6 +347,8 @@ void UMVVMWidgetBlueprintExtension_View::HandleFinishCompilingClass(UWidgetBluep
 		}
 	}
 
+	GeneratedFunctions = CurrentCompilerContext->GetGeneratedFunctions();
+	
 	// If we can't auto-generate function add the transient flags
 	if (UE::MVVM::Private::GAutogeneratedFunctionsAreTransient)
 	{
diff --git a/Engine/Plugins/Runtime/ModelViewViewModel/Source/ModelViewViewModelBlueprint/Public/MVVMWidgetBlueprintExtension_View.h b/Engine/Plugins/Runtime/ModelViewViewModel/Source/ModelViewViewModelBlueprint/Public/MVVMWidgetBlueprintExtension_View.h
--- a/Engine/Plugins/Runtime/ModelViewViewModel/Source/ModelViewViewModelBlueprint/Public/MVVMWidgetBlueprintExtension_View.h
+++ b/Engine/Plugins/Runtime/ModelViewViewModel/Source/ModelViewViewModelBlueprint/Public/MVVMWidgetBlueprintExtension_View.h
@@ -71,6 +71,11 @@ public:
 		return BlueprintViewChangedDelegate;
 	}
 
+	const TArrayView<const FName> GetGeneratedFunctions() const
+	{
+		return GeneratedFunctions;
+	}
+
 public:
 	//~ Begin UObject interface
 	virtual void PostLoad() override;
@@ -112,6 +117,9 @@ private:
 
 	UPROPERTY(Transient)
 	FMVVMViewBindingFilterSettings FilterSettings;
+	
+	UPROPERTY(Transient)
+	TArray<FName> GeneratedFunctions;
 
 	FSimpleMulticastDelegate BlueprintViewChangedDelegate;
 	TPimplPtr<UE::MVVM::Private::FMVVMViewBlueprintCompiler> CurrentCompilerContext;
diff --git a/Engine/Plugins/Runtime/ModelViewViewModel/Source/ModelViewViewModelEditor/Private/ModelViewViewModelEditorModule.cpp b/Engine/Plugins/Runtime/ModelViewViewModel/Source/ModelViewViewModelEditor/Private/ModelViewViewModelEditorModule.cpp
--- a/Engine/Plugins/Runtime/ModelViewViewModel/Source/ModelViewViewModelEditor/Private/ModelViewViewModelEditorModule.cpp
+++ b/Engine/Plugins/Runtime/ModelViewViewModel/Source/ModelViewViewModelEditor/Private/ModelViewViewModelEditorModule.cpp
@@ -84,6 +84,7 @@ void FModelViewViewModelEditorModule::StartupModule()
 	FMVVMEditorCommands::Register();
 	FWidgetBlueprintDelegates::GetAssetTagsWithContext.AddRaw(this, &FModelViewViewModelEditorModule::HandleWidgetBlueprintAssetTags);
 	FWidgetBlueprintGeneratedClassDelegates::GetAssetTagsWithContext.AddRaw(this, &FModelViewViewModelEditorModule::HandleClassBlueprintAssetTags);
+	FWidgetBlueprintGeneratedClassDelegates::CollectSaveOverrides.AddRaw(this, &FModelViewViewModelEditorModule::HandleCollectSaveOverrides);
 
 	UToolMenus::RegisterStartupCallback(FSimpleMulticastDelegate::FDelegate::CreateRaw(this, &FModelViewViewModelEditorModule::HandleRegisterMenus));
 }
@@ -93,6 +94,7 @@ void FModelViewViewModelEditorModule::ShutdownModule()
 {
 	UnregisterMenus();
 
+	FWidgetBlueprintGeneratedClassDelegates::CollectSaveOverrides.RemoveAll(this);
 	FWidgetBlueprintGeneratedClassDelegates::GetAssetTagsWithContext.RemoveAll(this);
 	FWidgetBlueprintDelegates::GetAssetTagsWithContext.RemoveAll(this);
 	if (FMessageLogModule* MessageLogModule = FModuleManager::GetModulePtr<FMessageLogModule>("MessageLog"))
@@ -260,6 +262,44 @@ void FModelViewViewModelEditorModule::HandleClassBlueprintAssetTags(const UWidge
 	}
 }
 
+namespace UE::MVVM::Private
+{
+	bool GAutogeneratedFunctionsAreForceEditorTransient = true;
+	static FAutoConsoleVariableRef CVarAutogeneratedFunctionsAreForceEditorTransient(
+		TEXT("MVVM.AutogeneratedFunctionsAreForceEditorTransient"),
+		GAutogeneratedFunctionsAreForceEditorTransient,
+		TEXT("Is the autogenerated function are mark as transient in editor but still cooked."),
+		ECVF_ReadOnly
+	);
+}
+
+void FModelViewViewModelEditorModule::HandleCollectSaveOverrides(const UWidgetBlueprintGeneratedClass* GeneratedClass, FObjectCollectSaveOverridesContext SaveContext)
+{
+	// The UWidgetBlueprintGeneratedClass::CollectSaveOverrides calls this callback. It give us the opotunity to remove object/properties from a package.
+	//Do not save the UFunction auto generated by MVVM. Can be a viewmodel setter, a conversion function, an event or a inner struct setter.
+	//The function are still alive and will be generated again on the next editor load, when a BP compiles. They are saved in a cook package.
+	if (UE::MVVM::Private::GAutogeneratedFunctionsAreForceEditorTransient && !SaveContext.IsCooking() && GeneratedClass && GeneratedClass->ClassGeneratedBy)
+	{
+		if (UWidgetBlueprint* WidgetBlueprint = Cast<UWidgetBlueprint>(GeneratedClass->ClassGeneratedBy))
+		{
+			if (UMVVMWidgetBlueprintExtension_View* ExtensionView = UMVVMWidgetBlueprintExtension_View::GetExtension<UMVVMWidgetBlueprintExtension_View>(WidgetBlueprint))
+			{
+				// If we can't auto-generate function add the transient flags
+				FObjectSaveOverride ObjectSaveOverride;
+				ObjectSaveOverride.bForceTransient = true;
+				for (FName FunctionName : ExtensionView->GetGeneratedFunctions())
+				{
+					UFunction* Function = GeneratedClass->FindFunctionByName(FunctionName);
+					if (ensure(Function))
+					{
+						SaveContext.AddSaveOverride(Function, ObjectSaveOverride);
+					}
+				}
+			}
+		}
+	}
+}
+
 
 void FModelViewViewModelEditorModule::HandleRegisterMenus()
 {
diff --git a/Engine/Plugins/Runtime/ModelViewViewModel/Source/ModelViewViewModelEditor/Private/ModelViewViewModelEditorModule.h b/Engine/Plugins/Runtime/ModelViewViewModel/Source/ModelViewViewModelEditor/Private/ModelViewViewModelEditorModule.h
--- a/Engine/Plugins/Runtime/ModelViewViewModel/Source/ModelViewViewModelEditor/Private/ModelViewViewModelEditorModule.h
+++ b/Engine/Plugins/Runtime/ModelViewViewModel/Source/ModelViewViewModelEditor/Private/ModelViewViewModelEditorModule.h
@@ -43,6 +43,7 @@ private:
 	void HandleActivateMode(FWidgetBlueprintApplicationMode& InDesignerMode);
 	void HandleWidgetBlueprintAssetTags(const UWidgetBlueprint* Widget, FAssetRegistryTagsContext Context);
 	void HandleClassBlueprintAssetTags(const UWidgetBlueprintGeneratedClass* GeneratedClass, FAssetRegistryTagsContext Context);
+	void HandleCollectSaveOverrides(const UWidgetBlueprintGeneratedClass* GeneratedClass, FObjectCollectSaveOverridesContext SaveContext);
 	void HandleRegisterMenus();
 	void UnregisterMenus();
 
diff --git a/Engine/Source/Runtime/UMG/Private/WidgetBlueprintGeneratedClass.cpp b/Engine/Source/Runtime/UMG/Private/WidgetBlueprintGeneratedClass.cpp
--- a/Engine/Source/Runtime/UMG/Private/WidgetBlueprintGeneratedClass.cpp
+++ b/Engine/Source/Runtime/UMG/Private/WidgetBlueprintGeneratedClass.cpp
@@ -108,6 +108,7 @@ PRAGMA_DISABLE_DEPRECATION_WARNINGS;
 FWidgetBlueprintGeneratedClassDelegates::FGetAssetTags FWidgetBlueprintGeneratedClassDelegates::GetAssetTags;
 PRAGMA_ENABLE_DEPRECATION_WARNINGS;
 FWidgetBlueprintGeneratedClassDelegates::FGetAssetTagsWithContext FWidgetBlueprintGeneratedClassDelegates::GetAssetTagsWithContext;
+FWidgetBlueprintGeneratedClassDelegates::FCollectSaveOverrides FWidgetBlueprintGeneratedClassDelegates::CollectSaveOverrides;
 
 #endif
 
@@ -556,6 +557,12 @@ void UWidgetBlueprintGeneratedClass::GetAssetRegistryTags(FAssetRegistryTagsCont
 	PRAGMA_ENABLE_DEPRECATION_WARNINGS;
 	FWidgetBlueprintGeneratedClassDelegates::GetAssetTagsWithContext.Broadcast(this, Context);
 }
+
+void UWidgetBlueprintGeneratedClass::CollectSaveOverrides(FObjectCollectSaveOverridesContext SaveContext)
+{
+	Super::CollectSaveOverrides(SaveContext);
+	FWidgetBlueprintGeneratedClassDelegates::CollectSaveOverrides.Broadcast(this, SaveContext);
+}
 #endif
 
 #undef LOCTEXT_NAMESPACE
diff --git a/Engine/Source/Runtime/UMG/Public/Blueprint/WidgetBlueprintGeneratedClass.h b/Engine/Source/Runtime/UMG/Public/Blueprint/WidgetBlueprintGeneratedClass.h
--- a/Engine/Source/Runtime/UMG/Public/Blueprint/WidgetBlueprintGeneratedClass.h
+++ b/Engine/Source/Runtime/UMG/Public/Blueprint/WidgetBlueprintGeneratedClass.h
@@ -59,9 +59,12 @@ public:
 	// delegate for generating widget asset registry tags.
 	DECLARE_MULTICAST_DELEGATE_TwoParams(FGetAssetTagsWithContext, const UWidgetBlueprintGeneratedClass*, FAssetRegistryTagsContext);
 	DECLARE_MULTICAST_DELEGATE_TwoParams(FGetAssetTags, const UWidgetBlueprintGeneratedClass*, TArray<UObject::FAssetRegistryTag>&);
+	DECLARE_MULTICAST_DELEGATE_TwoParams(FCollectSaveOverrides, const UWidgetBlueprintGeneratedClass*, FObjectCollectSaveOverridesContext);
 
 	// called by UWidgetBlueprintGeneratedClass::GetAssetRegistryTags()
 	static UMG_API FGetAssetTagsWithContext GetAssetTagsWithContext;
+	// called by UWidgetBlueprintGeneratedClass::CollectSaveOverrides()
+	static UMG_API FCollectSaveOverrides CollectSaveOverrides;
 	UE_DEPRECATED(5.4, "Subscribe to GetAssetTagsWithContext instead.")
 	static UMG_API FGetAssetTags GetAssetTags;
 };
@@ -163,6 +166,7 @@ public:
 	UMG_API virtual void GetAssetRegistryTags(FAssetRegistryTagsContext Context) const override;
 	UE_DEPRECATED(5.4, "Implement the version that takes FAssetRegistryTagsContext instead.")
 	UMG_API virtual void GetAssetRegistryTags(TArray<FAssetRegistryTag>& OutTags) const override;
+	UMG_API virtual void CollectSaveOverrides(FObjectCollectSaveOverridesContext SaveContext) override;
 #endif
 	//~ End UObject interface
 
