MVVM: fix context sensitive option not working

diff --git a/Engine/Plugins/Runtime/ModelViewViewModel/Source/ModelViewViewModelEditor/Private/Widgets/SMVVMFieldSelectorMenu.cpp b/Engine/Plugins/Runtime/ModelViewViewModel/Source/ModelViewViewModelEditor/Private/Widgets/SMVVMFieldSelectorMenu.cpp
--- a/Engine/Plugins/Runtime/ModelViewViewModel/Source/ModelViewViewModelEditor/Private/Widgets/SMVVMFieldSelectorMenu.cpp
+++ b/Engine/Plugins/Runtime/ModelViewViewModel/Source/ModelViewViewModelEditor/Private/Widgets/SMVVMFieldSelectorMenu.cpp
@@ -72,15 +72,10 @@ void SFieldSelectorMenu::Construct(const FArguments& InArgs, const UWidgetBluepr
 	bIsMenuInitialized = false;
 	const bool bIsClearEnabled = InArgs._CurrentSelected.IsSet() && InArgs._CurrentSelected.GetValue().IsValid();
 
-	// If we're showing conversion functions, we don't want to set the AssignableTo property of SSourceBindingList, because then it will only show exact matches, 
-	// and since we're also showing conversion functions we know that's not what the user wants.
-	// However, in the case we're not showing conversion functions, we want only exact matches.
 	const FProperty* AssignableToProperty = nullptr;
 	const bool bShowConversionFunctions = SelectionContext.bAllowConversionFunctions;
-	if (!bShowConversionFunctions)
-	{
-		AssignableToProperty = SelectionContext.AssignableTo;
-	}
+
+	AssignableToProperty = SelectionContext.AssignableTo;
 
 	if (bShowConversionFunctions)
 	{
@@ -1005,12 +1000,11 @@ void SFieldSelectorMenu::HandleEnabledContextToggleChanged(ECheckBoxState CheckS
 	FilterSettings.FilterFlags = CheckState == ECheckBoxState::Checked ? EFilterFlag::All : EFilterFlag::None;
 	ExtensionView->SetFilterSettings(FilterSettings);
 
-	if (WidgetList.IsValid())
+	if (WidgetList.IsValid() && WidgetList->GetSelectedWidgets().Num() > 0)
 	{
 		HandleWidgetSelected(WidgetBlueprint->GetFName(), ESelectInfo::Direct);
 	}
-
-	if (ViewModelList.IsValid())
+	else if (ViewModelList.IsValid() && ViewModelList->GetNumItemsSelected() > 0)
 	{
 		HandleViewModelSelected(FBindingSource(), ESelectInfo::Direct);
 	}
