[MVVM] Add helper for setting an array of ftext using MVVM macros

#tests Used in local change

diff --git a/Engine/Plugins/Runtime/ModelViewViewModel/Source/ModelViewViewModel/Public/MVVMViewModelBase.h b/Engine/Plugins/Runtime/ModelViewViewModel/Source/ModelViewViewModel/Public/MVVMViewModelBase.h
--- a/Engine/Plugins/Runtime/ModelViewViewModel/Source/ModelViewViewModel/Public/MVVMViewModelBase.h
+++ b/Engine/Plugins/Runtime/ModelViewViewModel/Source/ModelViewViewModel/Public/MVVMViewModelBase.h
@@ -3,6 +3,7 @@
 #pragma once
 
 
+#include "Algo/Compare.h"
 #include "FieldNotificationDeclaration.h" // IWYU pragma: keep
 #include "IFieldNotificationClassDescriptor.h"
 #include "ViewModel/MVVMFieldNotificationDelegates.h"
@@ -99,6 +100,21 @@ protected:
 		return true;
 	}
 
+	bool SetPropertyValue(TArray<FText>& Value, const TArray<FText>& NewValue, UE::FieldNotification::FFieldId FieldId)
+	{
+		if (Algo::Compare(Value, NewValue, [](const FText& LHS, const FText& RHS)
+			{
+				return LHS.IdenticalTo(RHS);
+			}))
+		{
+			return false;
+		}
+
+		Value = NewValue;
+		BroadcastFieldValueChanged(FieldId);
+		return true;
+	}
+
 private:
 	DECLARE_FUNCTION(execK2_SetPropertyValue);
 
