Add comments related to delegates / events / MVVM

#tests PIE, Standalone, Cooked

diff --git a/Engine/Plugins/Runtime/ModelViewViewModel/Source/ModelViewViewModelBlueprint/Private/MVVMConversionFunctionGraphSchema.h b/Engine/Plugins/Runtime/ModelViewViewModel/Source/ModelViewViewModelBlueprint/Private/MVVMConversionFunctionGraphSchema.h
--- a/Engine/Plugins/Runtime/ModelViewViewModel/Source/ModelViewViewModelBlueprint/Private/MVVMConversionFunctionGraphSchema.h
+++ b/Engine/Plugins/Runtime/ModelViewViewModel/Source/ModelViewViewModelBlueprint/Private/MVVMConversionFunctionGraphSchema.h
@@ -7,7 +7,7 @@
 #include "MVVMConversionFunctionGraphSchema.generated.h"
 
 /**
- *
+ * Schema for conversion functions, adds pin metadata needed on connections for MVVM
  */
 UCLASS()
 class UMVVMConversionFunctionGraphSchema : public UEdGraphSchema_K2
diff --git a/Engine/Plugins/Runtime/ModelViewViewModel/Source/ModelViewViewModelBlueprint/Public/MVVMBlueprintViewEvent.h b/Engine/Plugins/Runtime/ModelViewViewModel/Source/ModelViewViewModelBlueprint/Public/MVVMBlueprintViewEvent.h
--- a/Engine/Plugins/Runtime/ModelViewViewModel/Source/ModelViewViewModelBlueprint/Public/MVVMBlueprintViewEvent.h
+++ b/Engine/Plugins/Runtime/ModelViewViewModel/Source/ModelViewViewModelBlueprint/Public/MVVMBlueprintViewEvent.h
@@ -15,7 +15,10 @@ class UMVVMK2Node_AreSourcesValidForEvent;
 class UWidgetBlueprint;
 
 /**
- *
+ * Binding for an event that MVVM will listen too. Does not imply 
+ * the MVVM graph itself will use events.
+ * 
+ * Ex: UButton::OnClick 
  */
 UCLASS(Within = MVVMBlueprintView)
 class MODELVIEWVIEWMODELBLUEPRINT_API UMVVMBlueprintViewEvent : public UObject
diff --git a/Engine/Source/Runtime/Engine/Classes/Engine/DynamicBlueprintBinding.h b/Engine/Source/Runtime/Engine/Classes/Engine/DynamicBlueprintBinding.h
--- a/Engine/Source/Runtime/Engine/Classes/Engine/DynamicBlueprintBinding.h
+++ b/Engine/Source/Runtime/Engine/Classes/Engine/DynamicBlueprintBinding.h
@@ -7,6 +7,13 @@
 #include "UObject/Object.h"
 #include "DynamicBlueprintBinding.generated.h"
 
+/**
+ * Used to tie events to delegates at runtime
+ * 
+ * To use, override 'K2Node::GetDynamicBindingClass' to specify a child of this class.
+ * Additionally, override 'K2Node::RegisterDynamicBinding'. During compilation, your node
+ * will be given an instance of the specified binding class to use for binding delegates.
+ */
 UCLASS(abstract, MinimalAPI)
 class UDynamicBlueprintBinding
 	: public UObject
