MVVM panel extension: fix entries being generated in reverse order if an array with completely new elements is set to the binding (when there's already entries on the screen).

diff --git a/Engine/Plugins/Runtime/ModelViewViewModel/Source/ModelViewViewModel/Private/Extensions/MVVMViewPanelWidgetExtension.cpp b/Engine/Plugins/Runtime/ModelViewViewModel/Source/ModelViewViewModel/Private/Extensions/MVVMViewPanelWidgetExtension.cpp
--- a/Engine/Plugins/Runtime/ModelViewViewModel/Source/ModelViewViewModel/Private/Extensions/MVVMViewPanelWidgetExtension.cpp
+++ b/Engine/Plugins/Runtime/ModelViewViewModel/Source/ModelViewViewModel/Private/Extensions/MVVMViewPanelWidgetExtension.cpp
@@ -81,6 +81,7 @@ void UMVVMPanelWidgetViewExtension::BP_SetItems(const TArray<UObject*>& InItems)
 			else
 			{
 				NewSlots.Add(TTuple<UPanelSlot*, UWidget*>(FoundObject->Key, FoundObject->Key->Content));
+				PreviousSlots.RemoveSingleSwap(*FoundObject);
 			}
 		}
 		else
@@ -123,19 +124,15 @@ void UMVVMPanelWidgetViewExtension::ReplaceAllSlots(TArrayView<TTuple<UPanelSlot
 		return;
 	}
 
-	const TArray<UPanelSlot*>& OldSlots = PanelWidget->GetSlots();
-	const int32 OldSlotsNum = OldSlots.Num();
+	const int32 OldSlotsNum = PanelWidget->GetSlots().Num();
 	const int32 NewSlotsNum = NewSlots.Num();
 	const int32 MinSlotNum = FMath::Min(OldSlotsNum, NewSlotsNum);
 
 	// as long as we're within the boundaries of both arrays, compare and replace elements
-	for (int32 SlotIndex = MinSlotNum-1; SlotIndex >= 0; --SlotIndex)
+	for (int32 SlotIndex = 0; SlotIndex < MinSlotNum; ++SlotIndex)
 	{
-		if (OldSlots[SlotIndex] != NewSlots[SlotIndex].Key)
-		{
-			PanelWidget->RemoveChildAt(SlotIndex);
-			UPanelSlot* NewSlot = PanelWidget->InsertChildAt(SlotIndex, NewSlots[SlotIndex].Value, NewSlots[SlotIndex].Key);
-		}
+		PanelWidget->RemoveChildAt(SlotIndex);
+		UPanelSlot* NewSlot = PanelWidget->InsertChildAt(SlotIndex, NewSlots[SlotIndex].Value, NewSlots[SlotIndex].Key);
 	}
 
 	// If we have more old slots than new ones, remove all the extra ones.
