MVVM: clean up extensions on destruct

diff --git a/Engine/Plugins/Runtime/ModelViewViewModel/Source/ModelViewViewModel/Private/View/MVVMView.cpp b/Engine/Plugins/Runtime/ModelViewViewModel/Source/ModelViewViewModel/Private/View/MVVMView.cpp
--- a/Engine/Plugins/Runtime/ModelViewViewModel/Source/ModelViewViewModel/Private/View/MVVMView.cpp
+++ b/Engine/Plugins/Runtime/ModelViewViewModel/Source/ModelViewViewModel/Private/View/MVVMView.cpp
@@ -115,9 +115,10 @@ void UMVVMView::Destruct()
 
 	const TArrayView<const TObjectPtr<UMVVMViewClassExtension>> ClassExtensions = GeneratedViewClass->GetViewClassExtensions();
 	check(ClassExtensions.Num() == Extensions.Num());
-	for (int32 Index = 0; Index < ClassExtensions.Num(); ++Index)
+	for (int32 Index = ClassExtensions.Num() - 1; Index >= 0; --Index)
 	{
 		ClassExtensions[Index]->OnViewDestructed(GetUserWidget(), this, Extensions[Index]);
+		Extensions.RemoveAtSwap(Index);
 	}
 }
 
