MVVM: fix context sensitive filtering non-field notify values when they're being written to by a field notify value.
This case should be accepted, and the correct case to filter out is when a non-field notify field is writing to another field (in a binding that's not one-time).

diff --git a/Engine/Plugins/Runtime/ModelViewViewModel/Source/ModelViewViewModelEditor/Private/Widgets/SMVVMViewModelBindingListWidget.cpp b/Engine/Plugins/Runtime/ModelViewViewModel/Source/ModelViewViewModelEditor/Private/Widgets/SMVVMViewModelBindingListWidget.cpp
--- a/Engine/Plugins/Runtime/ModelViewViewModel/Source/ModelViewViewModelEditor/Private/Widgets/SMVVMViewModelBindingListWidget.cpp
+++ b/Engine/Plugins/Runtime/ModelViewViewModel/Source/ModelViewViewModelEditor/Private/Widgets/SMVVMViewModelBindingListWidget.cpp
@@ -37,7 +37,7 @@ namespace Private
 			return TOptional<FFieldVariant>();
 		}
 
-		if (EnumHasAllFlags(FieldVisibilityFlags, EFieldVisibility::Notify) && !Binding.HasNotify())
+		if (EnumHasAllFlags(FieldVisibilityFlags, EFieldVisibility::Notify | EFieldVisibility::Readable) && !Binding.HasNotify())
 		{
 			return TOptional<FFieldVariant>();
 		}
