MVVM: disable auto-conversion from SoftObjectPointer to ObjectPointer in conversion function arguments

diff --git a/Engine/Plugins/Runtime/ModelViewViewModel/Source/ModelViewViewModelBlueprint/Private/MVVMConversionFunctionGraphSchema.cpp b/Engine/Plugins/Runtime/ModelViewViewModel/Source/ModelViewViewModelBlueprint/Private/MVVMConversionFunctionGraphSchema.cpp
--- a/Engine/Plugins/Runtime/ModelViewViewModel/Source/ModelViewViewModelBlueprint/Private/MVVMConversionFunctionGraphSchema.cpp
+++ b/Engine/Plugins/Runtime/ModelViewViewModel/Source/ModelViewViewModelBlueprint/Private/MVVMConversionFunctionGraphSchema.cpp
@@ -32,6 +32,20 @@ bool UMVVMConversionFunctionGraphSchema::TryCreateConnection(UEdGraphPin* A, UEd
 	return bResult;
 }
 
+const FPinConnectionResponse UMVVMConversionFunctionGraphSchema::CanCreateConnection(const UEdGraphPin* PinA, const UEdGraphPin* PinB) const
+{
+	check(PinA);
+	check(PinB);
+	if (PinA->PinType.PinCategory == UEdGraphSchema_K2::PC_SoftObject && PinB->PinType.PinCategory == UEdGraphSchema_K2::PC_Object)
+	{
+		FPinConnectionResponse Response(CONNECT_RESPONSE_DISALLOW, FString::Printf(TEXT("Conversion disallowed from %s to %s"), *UEdGraphSchema_K2::PC_SoftObject.ToString(), *UEdGraphSchema_K2::PC_Object.ToString()));
+		Response.SetFatal();
+		return Response;
+	}
+
+	return Super::CanCreateConnection(PinA, PinB);
+}
+
 ///////////////////////////////////////////////////////////////////////////////
 
 EGraphType UMVVMFakeTestUbergraphSchema::GetGraphType(const UEdGraph* TestEdGraph) const
diff --git a/Engine/Plugins/Runtime/ModelViewViewModel/Source/ModelViewViewModelBlueprint/Private/MVVMConversionFunctionGraphSchema.h b/Engine/Plugins/Runtime/ModelViewViewModel/Source/ModelViewViewModelBlueprint/Private/MVVMConversionFunctionGraphSchema.h
--- a/Engine/Plugins/Runtime/ModelViewViewModel/Source/ModelViewViewModelBlueprint/Private/MVVMConversionFunctionGraphSchema.h
+++ b/Engine/Plugins/Runtime/ModelViewViewModel/Source/ModelViewViewModelBlueprint/Private/MVVMConversionFunctionGraphSchema.h
@@ -16,6 +16,7 @@ class UMVVMConversionFunctionGraphSchema : public UEdGraphSchema_K2
 
 public:
 	virtual bool TryCreateConnection(UEdGraphPin* A, UEdGraphPin* B) const override;
+	virtual const FPinConnectionResponse CanCreateConnection(const UEdGraphPin* A, const UEdGraphPin* B) const override;
 };
 
 /**
