Fix incorrect API Exposure MVVM

#tests NonUnity

diff --git a/Engine/Plugins/Runtime/ModelViewViewModel/Source/ModelViewViewModelBlueprint/Public/Bindings/MVVMConversionFunctionHelper.h b/Engine/Plugins/Runtime/ModelViewViewModel/Source/ModelViewViewModelBlueprint/Public/Bindings/MVVMConversionFunctionHelper.h
--- a/Engine/Plugins/Runtime/ModelViewViewModel/Source/ModelViewViewModelBlueprint/Public/Bindings/MVVMConversionFunctionHelper.h
+++ b/Engine/Plugins/Runtime/ModelViewViewModel/Source/ModelViewViewModelBlueprint/Public/Bindings/MVVMConversionFunctionHelper.h
@@ -47,7 +47,7 @@ namespace UE::MVVM::ConversionFunctionHelper
 	 */
 	MODELVIEWVIEWMODELBLUEPRINT_API TValueOrError<void, FText> CanCreateSetterGraph(UBlueprint* WidgetBlueprint, const FMVVMBlueprintPropertyPath& PropertyPath);
 
-	MODELVIEWVIEWMODELBLUEPRINT_API struct FCreateGraphResult
+	struct MODELVIEWVIEWMODELBLUEPRINT_API FCreateGraphResult
 	{
 		/** The new graph created. */
 		UEdGraph* NewGraph = nullptr;
@@ -59,7 +59,7 @@ namespace UE::MVVM::ConversionFunctionHelper
 		bool bIsUbergraphPage = false;
 	};
 
-	MODELVIEWVIEWMODELBLUEPRINT_API struct FCreateGraphParams
+	struct MODELVIEWVIEWMODELBLUEPRINT_API FCreateGraphParams
 	{
 		bool bIsConst = false;
 		bool bTransient = false;
