Fix Static Analysis MVVM Async

#tests NonUnity

diff --git a/Engine/Plugins/Runtime/ModelViewViewModel/Source/ModelViewViewModelBlueprint/Private/Bindings/MVVMConversionFunctionHelper.cpp b/Engine/Plugins/Runtime/ModelViewViewModel/Source/ModelViewViewModelBlueprint/Private/Bindings/MVVMConversionFunctionHelper.cpp
--- a/Engine/Plugins/Runtime/ModelViewViewModel/Source/ModelViewViewModelBlueprint/Private/Bindings/MVVMConversionFunctionHelper.cpp
+++ b/Engine/Plugins/Runtime/ModelViewViewModel/Source/ModelViewViewModelBlueprint/Private/Bindings/MVVMConversionFunctionHelper.cpp
@@ -1445,7 +1445,6 @@ void SetPropertyPathForPin(const UBlueprint* Blueprint, const FMVVMBlueprintProp
 	UEdGraph* FunctionGraph = ConversionNode ? ConversionNode->GetGraph() : nullptr;
 	const UEdGraphSchema* Schema = GetDefault<UMVVMConversionFunctionGraphSchema>();
 
-
 	UK2Node* K2ConversionNode = Cast<UK2Node>(ConversionNode);
 	if (K2ConversionNode && FunctionGraph)
 	{
@@ -1459,15 +1458,15 @@ void SetPropertyPathForPin(const UBlueprint* Blueprint, const FMVVMBlueprintProp
 			// Add new nodes
 			if (PropertyPath.IsValid())
 			{
-				const int32 ArgumentIndex = ConversionNode->Pins.IndexOfByPredicate([PathPin](const UEdGraphPin* Other){ return Other == PathPin; });
-				if (!ensure(ConversionNode->Pins.IsValidIndex(ArgumentIndex)))
+				const int32 ArgumentIndex = K2ConversionNode->Pins.IndexOfByPredicate([PathPin](const UEdGraphPin* Other){ return Other == PathPin; });
+				if (!ensure(K2ConversionNode->Pins.IsValidIndex(ArgumentIndex)))
 				{
 					return;
 				}
 
 				const int32 NumberOfFields = PropertyPath.GetFieldPaths().Num();
-				const float PosX = ConversionNode->NodePosX;
-				const float PosY = ConversionNode->NodePosY + ArgumentIndex * 100;
+				const float PosX = K2ConversionNode->NodePosX;
+				const float PosY = K2ConversionNode->NodePosY + ArgumentIndex * 100;
 				TValueOrError<TArray<UEdGraphPin*>, void> BuildPropertyPathResult = Private::BuildPropertyPath(Blueprint, FunctionGraph, PropertyPath, NumberOfFields, FVector2f(PosX, PosY));
 				if (BuildPropertyPathResult.HasError())
 				{
@@ -1479,7 +1478,7 @@ void SetPropertyPathForPin(const UBlueprint* Blueprint, const FMVVMBlueprintProp
 			}
 
 			// Link Then / Exec pin
-			Private::LinkAllNodesForEvent(FunctionGraph, ConversionEventEntry, ConversionNode);
+			Private::LinkAllNodesForEvent(FunctionGraph, ConversionEventEntry, K2ConversionNode);
 		}
 	}
 
