MVVM: fix crash when adding extensions to panel widgets.

diff --git a/Engine/Plugins/Runtime/ModelViewViewModel/Source/ModelViewViewModelBlueprint/Private/Extensions/MVVMViewBlueprintPanelWidgetExtension.cpp b/Engine/Plugins/Runtime/ModelViewViewModel/Source/ModelViewViewModelBlueprint/Private/Extensions/MVVMViewBlueprintPanelWidgetExtension.cpp
--- a/Engine/Plugins/Runtime/ModelViewViewModel/Source/ModelViewViewModelBlueprint/Private/Extensions/MVVMViewBlueprintPanelWidgetExtension.cpp
+++ b/Engine/Plugins/Runtime/ModelViewViewModel/Source/ModelViewViewModelBlueprint/Private/Extensions/MVVMViewBlueprintPanelWidgetExtension.cpp
@@ -293,13 +293,16 @@ bool UMVVMBlueprintViewExtension_PanelWidget::WidgetRenamed(FName OldName, FName
 
 void UMVVMBlueprintViewExtension_PanelWidget::OnPreviewContentChanged(TSharedRef<SWidget> NewContent)
 {
-	const SObjectWidget* ObjectWidget = StaticCastSharedPtr<SObjectWidget>(NewContent.ToSharedPtr()).Get();
-	const UUserWidget* PreviewRoot = ObjectWidget ? ObjectWidget->GetWidgetObject() : nullptr;
-
-	if (UPanelWidget* PreviewWidget = PreviewRoot ? Cast<UPanelWidget>(PreviewRoot->GetWidgetFromName(WidgetName)) : nullptr)
+	if (NewContent != SNullWidget::NullWidget)
 	{
-		constexpr bool bFullRebuild = false;
-		RefreshDesignerPreviewEntries(PreviewWidget, EntryWidgetClass, SlotObj, NumDesignerPreviewEntries, bFullRebuild);
+		const SObjectWidget* ObjectWidget = StaticCastSharedPtr<SObjectWidget>(NewContent.ToSharedPtr()).Get();
+		const UUserWidget* PreviewRoot = ObjectWidget ? ObjectWidget->GetWidgetObject() : nullptr;
+
+		if (UPanelWidget* PreviewWidget = PreviewRoot ? Cast<UPanelWidget>(PreviewRoot->GetWidgetFromName(WidgetName)) : nullptr)
+		{
+			constexpr bool bFullRebuild = false;
+			RefreshDesignerPreviewEntries(PreviewWidget, EntryWidgetClass, SlotObj, NumDesignerPreviewEntries, bFullRebuild);
+		}
 	}
 }
 
