Fix Manual Arg refresh required when node added to MVVM Allowlist

#tests Editor

diff --git a/Engine/Plugins/Runtime/ModelViewViewModel/Source/ModelViewViewModelEditor/Private/Framework/MVVMConversionFunctionLibrary.cpp b/Engine/Plugins/Runtime/ModelViewViewModel/Source/ModelViewViewModelEditor/Private/Framework/MVVMConversionFunctionLibrary.cpp
--- a/Engine/Plugins/Runtime/ModelViewViewModel/Source/ModelViewViewModelEditor/Private/Framework/MVVMConversionFunctionLibrary.cpp
+++ b/Engine/Plugins/Runtime/ModelViewViewModel/Source/ModelViewViewModelEditor/Private/Framework/MVVMConversionFunctionLibrary.cpp
@@ -143,6 +143,16 @@ FCollection::~FCollection()
 	}
 }
 
+void FCollection::AddReferencedObjects(FReferenceCollector& Collector)
+{
+	Collector.AddReferencedObjects(ConversionFunctionNodes);
+}
+
+FString FCollection::GetReferencerName() const
+{
+	return TEXT("MVVMConversionFunctionLibrary::FCollection");
+}
+
 void FCollection::RefreshIfNeeded()
 {
 	if (bRefreshAll)
@@ -258,6 +268,8 @@ void FCollection::Build()
 {
 	FScopedSlowTask SlowTask = FScopedSlowTask(0, LOCTEXT("BuildingConversionFunctionLibrary", "Loading Conversion function library"));
 
+	ConversionFunctionNodes.Reset();
+
 	TArray<const UClass*> AllSupportedClass;
 	AllSupportedClass.Reserve(128);
 
@@ -411,10 +423,10 @@ void FCollection::AddClassFunctions(const UClass* Class)
 
 void FCollection::AddNode(TSubclassOf<UK2Node> Function)
 {
-	UEdGraph* OwnerOfTemporaries = NewObject<UEdGraph>(GetTransientPackage());
-	UK2Node* NewNode = NewObject<UK2Node>(OwnerOfTemporaries, Function.Get(), NAME_None, RF_Transient);
+	UK2Node* NewNode = NewObject<UK2Node>(GetTransientPackage(), Function.Get());
 	NewNode->AllocateDefaultPins();
 	NewNode->PostPlacedNewNode();
+	ConversionFunctionNodes.Add(NewNode);
 
  	TArray<UEdGraphPin*> InputPins = UE::MVVM::ConversionFunctionHelper::FindInputPins(NewNode);
 	if (InputPins.Num() == 0)
diff --git a/Engine/Plugins/Runtime/ModelViewViewModel/Source/ModelViewViewModelEditor/Private/Framework/MVVMConversionFunctionLibrary.h b/Engine/Plugins/Runtime/ModelViewViewModel/Source/ModelViewViewModelEditor/Private/Framework/MVVMConversionFunctionLibrary.h
--- a/Engine/Plugins/Runtime/ModelViewViewModel/Source/ModelViewViewModelEditor/Private/Framework/MVVMConversionFunctionLibrary.h
+++ b/Engine/Plugins/Runtime/ModelViewViewModel/Source/ModelViewViewModelEditor/Private/Framework/MVVMConversionFunctionLibrary.h
@@ -44,12 +44,17 @@ private:
  * Collection of all available conversion functions.
  * The collection rebuild when a new object is loaded and when the WidgetBlueprint is compiled.
  */
-class FCollection
+class FCollection : public FGCObject
 {
 public:
 	FCollection();
 	~FCollection();
 
+	//~ Begin FGCObject
+	virtual void AddReferencedObjects(FReferenceCollector& Collector) override;
+	virtual FString GetReferencerName() const override;
+	//~ End FGCObject
+
 public:
 	[[nodiscard]] TArray<::UE::MVVM::FConversionFunctionValue> GetFunctions(const UWidgetBlueprint* Blueprint) const;
 	[[nodiscard]] TArray<::UE::MVVM::FConversionFunctionValue> GetFunctions(const UWidgetBlueprint* Blueprint, const FProperty* ArgumentType, const FProperty* ReturnType) const;
@@ -90,6 +95,7 @@ private:
 		bool bIsNode = false;
 	};
 	TMap<FObjectKey, FFunctionContainer> ClassOrBlueprintToFunctions;
+	TArray<TObjectPtr<UK2Node>> ConversionFunctionNodes;
 
 	TSet<FObjectKey> ObjectToRefresh; // can only be GeneratedClass or Blueprint
 	TSet<FName> ModuleToRefresh;
