MVVM: fix generated functions not being cleared properly on undo.
The temporary array holding the generated graphs should not be affected by undo because it holds data created on the compiler context which is itself not transactional.

diff --git a/Engine/Plugins/Runtime/ModelViewViewModel/Source/ModelViewViewModelBlueprint/Private/MVVMViewBlueprintCompiler.cpp b/Engine/Plugins/Runtime/ModelViewViewModel/Source/ModelViewViewModelBlueprint/Private/MVVMViewBlueprintCompiler.cpp
--- a/Engine/Plugins/Runtime/ModelViewViewModel/Source/ModelViewViewModelBlueprint/Private/MVVMViewBlueprintCompiler.cpp
+++ b/Engine/Plugins/Runtime/ModelViewViewModel/Source/ModelViewViewModelBlueprint/Private/MVVMViewBlueprintCompiler.cpp
@@ -1173,6 +1173,12 @@ void FMVVMViewBlueprintCompiler::CreatePublicFunctionsDeclaration(const FWidgetB
 					FText::FromString(Setter.BlueprintSetter)
 				).ToString()
 				);
+
+				if (Setter.SetterGraph != nullptr)
+				{ 
+					RenameObjectToTransientPackage(Setter.SetterGraph);
+					Setter.SetterGraph = nullptr;
+				}
 				bIsCreateVariableStepValid = false;
 				continue;
 			}
diff --git a/Engine/Plugins/Runtime/ModelViewViewModel/Source/ModelViewViewModelBlueprint/Public/MVVMBlueprintView.h b/Engine/Plugins/Runtime/ModelViewViewModel/Source/ModelViewViewModelBlueprint/Public/MVVMBlueprintView.h
--- a/Engine/Plugins/Runtime/ModelViewViewModel/Source/ModelViewViewModelBlueprint/Public/MVVMBlueprintView.h
+++ b/Engine/Plugins/Runtime/ModelViewViewModel/Source/ModelViewViewModelBlueprint/Public/MVVMBlueprintView.h
@@ -175,7 +175,7 @@ public:
 	FOnViewModelsUpdated OnViewModelsUpdated;
 
 	// Use during compilation to clean the automatically generated graph.
-	UPROPERTY(Transient)
+	UPROPERTY(Transient, NonTransactional)
 	TArray<TObjectPtr<UEdGraph>> TemporaryGraph;
 
 private:
