[MVVM] Add ability to set default resolver to select by default in project settings.
* Fixed a mixed up condition
* Cleanup resolver reference when not using it as a creation type.

#tests Changed creation types. Added a VM that's defaulted to use resolver type.

diff --git a/Engine/Plugins/Runtime/ModelViewViewModel/Source/ModelViewViewModelBlueprint/Private/MVVMBlueprintViewModelContext.cpp b/Engine/Plugins/Runtime/ModelViewViewModel/Source/ModelViewViewModelBlueprint/Private/MVVMBlueprintViewModelContext.cpp
--- a/Engine/Plugins/Runtime/ModelViewViewModel/Source/ModelViewViewModelBlueprint/Private/MVVMBlueprintViewModelContext.cpp
+++ b/Engine/Plugins/Runtime/ModelViewViewModel/Source/ModelViewViewModelBlueprint/Private/MVVMBlueprintViewModelContext.cpp
@@ -2,7 +2,7 @@
 
 #include "MVVMBlueprintViewModelContext.h"
 #include "MVVMDeveloperProjectSettings.h"
-
+#include "View/MVVMViewModelContextResolver.h"
 
 #include UE_INLINE_GENERATED_CPP_BY_NAME(MVVMBlueprintViewModelContext)
 
@@ -20,6 +20,10 @@ FMVVMBlueprintViewModelContext::FMVVMBlueprintViewModelContext(const UClass* InC
 			NotifyFieldValueClass = const_cast<UClass*>(InClass);
 			ViewModelName = InViewModelName;
 			bExposeInstanceInEditor = GetDefault<UMVVMDeveloperProjectSettings>()->bExposeViewModelInstanceInEditor;
+
+#if WITH_EDITOR
+			Resolver = CreateDefaultResolver();
+#endif
 		}
 	}
 }
@@ -31,6 +35,27 @@ FText FMVVMBlueprintViewModelContext::GetDisplayName() const
 }
 
 #if WITH_EDITOR
+TObjectPtr<UMVVMViewModelContextResolver> FMVVMBlueprintViewModelContext::CreateDefaultResolver() const
+{
+	TObjectPtr<UMVVMViewModelContextResolver> DefaultResolver;
+
+	if (CreationType == EMVVMBlueprintViewModelContextCreationType::Resolver)
+	{
+		if (const UClass* DefaultClass = GetDefault<UMVVMDeveloperProjectSettings>()->DefaultResolverValue.ResolveClass())
+		{
+			if (UMVVMViewModelContextResolver* DefaultResolverClass = Cast<UMVVMViewModelContextResolver>(DefaultClass->GetDefaultObject()))
+			{
+				if (DefaultResolverClass->DoesSupportViewModelClass(GetViewModelClass()))
+				{
+					DefaultResolver = NewObject<UMVVMViewModelContextResolver>(GetTransientPackage(), DefaultClass);
+				}
+			}
+		}
+	}
+
+	return DefaultResolver;
+}
+
 namespace UE::MVVM
 {
 namespace Private
@@ -45,7 +70,7 @@ void GenerateAllowedList(const UClass* Class, TArray<FString>& AllowedStrings, T
 	{
 		Class->GetMetaData(NAME_MVVMAllowedContextCreationType).ParseIntoArray(AllowedStrings, TEXT("|"));
 	}
-	if (AllowedStrings.Num() == 0 && Class->HasMetaData(NAME_MVVMDisallowedContextCreationType))
+	if (DisallowedStrings.Num() == 0 && Class->HasMetaData(NAME_MVVMDisallowedContextCreationType))
 	{
 		Class->GetMetaData(NAME_MVVMDisallowedContextCreationType).ParseIntoArray(DisallowedStrings, TEXT("|"));
 	}
diff --git a/Engine/Plugins/Runtime/ModelViewViewModel/Source/ModelViewViewModelBlueprint/Public/MVVMBlueprintViewModelContext.h b/Engine/Plugins/Runtime/ModelViewViewModel/Source/ModelViewViewModelBlueprint/Public/MVVMBlueprintViewModelContext.h
--- a/Engine/Plugins/Runtime/ModelViewViewModel/Source/ModelViewViewModelBlueprint/Public/MVVMBlueprintViewModelContext.h
+++ b/Engine/Plugins/Runtime/ModelViewViewModel/Source/ModelViewViewModelBlueprint/Public/MVVMBlueprintViewModelContext.h
@@ -97,6 +97,10 @@ public:
 		return NotifyFieldValueClass != nullptr;
 	}
 
+#if WITH_EDITOR
+	[[nodiscard]] TObjectPtr<UMVVMViewModelContextResolver> CreateDefaultResolver() const;
+#endif
+
 private:
 	/** When the view is spawn, create an instance of the viewmodel. */
 	UPROPERTY(VisibleAnywhere, AdvancedDisplay, Category = "Viewmodel", meta = (DisplayName = "Viewmodel Context Id", NoResetToDefault))
diff --git a/Engine/Plugins/Runtime/ModelViewViewModel/Source/ModelViewViewModelBlueprint/Public/MVVMDeveloperProjectSettings.h b/Engine/Plugins/Runtime/ModelViewViewModel/Source/ModelViewViewModelBlueprint/Public/MVVMDeveloperProjectSettings.h
--- a/Engine/Plugins/Runtime/ModelViewViewModel/Source/ModelViewViewModelBlueprint/Public/MVVMDeveloperProjectSettings.h
+++ b/Engine/Plugins/Runtime/ModelViewViewModel/Source/ModelViewViewModelBlueprint/Public/MVVMDeveloperProjectSettings.h
@@ -196,6 +196,10 @@ public:
 	/** Sub-classes of ListViewBase that are supported to have an extension for binding their entries to viewmodels. */
 	UPROPERTY(EditAnywhere, config, Category = "Widget Extension")
 	TSet<TSoftClassPtr<UListViewBase>> SupportedListViewBaseClassesForExtension;
+
+	/** Resolver class to use as the default value when selecting resolver creation mode */
+	UPROPERTY(EditAnywhere, config, Category = "Defaults", meta = (MetaClass = "/Script/ModelViewViewModel.UMVVMViewModelContextResolver"))
+	FSoftClassPath DefaultResolverValue;
 };
 
 #if UE_ENABLE_INCLUDE_ORDER_DEPRECATED_IN_5_2
diff --git a/Engine/Plugins/Runtime/ModelViewViewModel/Source/ModelViewViewModelEditor/Private/Customizations/MVVMBlueprintViewModelContextCustomization.cpp b/Engine/Plugins/Runtime/ModelViewViewModel/Source/ModelViewViewModelEditor/Private/Customizations/MVVMBlueprintViewModelContextCustomization.cpp
--- a/Engine/Plugins/Runtime/ModelViewViewModel/Source/ModelViewViewModelEditor/Private/Customizations/MVVMBlueprintViewModelContextCustomization.cpp
+++ b/Engine/Plugins/Runtime/ModelViewViewModel/Source/ModelViewViewModelEditor/Private/Customizations/MVVMBlueprintViewModelContextCustomization.cpp
@@ -207,6 +207,7 @@ void FBlueprintViewModelContextDetailCustomization::CustomizeChildren(TSharedRef
 	OptionalHandle = PropertyHandle->GetChildHandle(GET_MEMBER_NAME_CHECKED(FMVVMBlueprintViewModelContext, bOptional), false);
 	CreateSetterFunctionHandle = PropertyHandle->GetChildHandle(GET_MEMBER_NAME_CHECKED(FMVVMBlueprintViewModelContext, bCreateSetterFunction), false);
 	ForceExecuteBindingsOnSetSourceHandle = PropertyHandle->GetChildHandle(GET_MEMBER_NAME_CHECKED(FMVVMBlueprintViewModelContext, bForceExecuteBindingsOnSetSource), false);
+	ResolverHandle = PropertyHandle->GetChildHandle(GET_MEMBER_NAME_CHECKED(FMVVMBlueprintViewModelContext, Resolver), false);
 
 	if (ensure(NotifyFieldValueClassHandle))
 	{
@@ -367,7 +368,6 @@ void FBlueprintViewModelContextDetailCustomization::CustomizeChildren(TSharedRef
 			}
 		}
 
-		TSharedPtr<IPropertyHandle> ResolverHandle = PropertyHandle->GetChildHandle(GET_MEMBER_NAME_CHECKED(FMVVMBlueprintViewModelContext, Resolver), false);
 		if (ensure(ResolverHandle))
 		{
 			TSharedRef<Private::FResolverClassFilter> ClassFilter = MakeShared<Private::FResolverClassFilter>();
@@ -509,7 +509,8 @@ void FBlueprintViewModelContextDetailCustomization::HandleCreationTypeChanged()
 	{
 		if (FMVVMBlueprintViewModelContext* ContextPtr = Private::GetViewModelContext(ContextHandle.ToSharedRef()))
 		{
-			const bool bIsManual = (EMVVMBlueprintViewModelContextCreationType)NewValue == EMVVMBlueprintViewModelContextCreationType::Manual;
+			const EMVVMBlueprintViewModelContextCreationType CreationType = (EMVVMBlueprintViewModelContextCreationType)NewValue;
+			const bool bIsManual = CreationType == EMVVMBlueprintViewModelContextCreationType::Manual;
 			if (ContextPtr->bOptional != bIsManual)
 			{
 				OptionalHandle->SetValue(bIsManual);
@@ -518,6 +519,20 @@ void FBlueprintViewModelContextDetailCustomization::HandleCreationTypeChanged()
 			{
 				CreateSetterFunctionHandle->SetValue(bIsManual);
 			}
+
+			// Set default resolver only if not already set to a valid value
+			if (CreationType == EMVVMBlueprintViewModelContextCreationType::Resolver)
+			{
+				UObject* ExistingResolver = nullptr;
+				if (ResolverHandle->GetValue(ExistingResolver) == FPropertyAccess::Fail || ExistingResolver == nullptr)
+				{
+					TObjectPtr<UMVVMViewModelContextResolver> NewResolver = ContextPtr->CreateDefaultResolver();
+
+					// Bypass SetValue, Resolver is set to Instanced which will block it
+					FString PropertyText = NewResolver ? NewResolver->GetPathName() : TEXT("None");
+					ResolverHandle->SetValueFromFormattedString(PropertyText);
+				}
+			}
 		}
 	}
 }
diff --git a/Engine/Plugins/Runtime/ModelViewViewModel/Source/ModelViewViewModelEditor/Private/Customizations/MVVMBlueprintViewModelContextCustomization.h b/Engine/Plugins/Runtime/ModelViewViewModel/Source/ModelViewViewModelEditor/Private/Customizations/MVVMBlueprintViewModelContextCustomization.h
--- a/Engine/Plugins/Runtime/ModelViewViewModel/Source/ModelViewViewModelEditor/Private/Customizations/MVVMBlueprintViewModelContextCustomization.h
+++ b/Engine/Plugins/Runtime/ModelViewViewModel/Source/ModelViewViewModelEditor/Private/Customizations/MVVMBlueprintViewModelContextCustomization.h
@@ -73,6 +73,7 @@ private:
 	TSharedPtr<IPropertyHandle> OptionalHandle;
 	TSharedPtr<IPropertyHandle> CreateSetterFunctionHandle;
 	TSharedPtr<IPropertyHandle> ForceExecuteBindingsOnSetSourceHandle;
+	TSharedPtr<IPropertyHandle> ResolverHandle;
 	TArray<EMVVMBlueprintViewModelContextCreationType> AllowedCreationTypes;
 };
 
