Remove meta class specifier causing crash in editor

diff --git a/Engine/Plugins/Runtime/ModelViewViewModel/Source/ModelViewViewModelBlueprint/Public/MVVMDeveloperProjectSettings.h b/Engine/Plugins/Runtime/ModelViewViewModel/Source/ModelViewViewModelBlueprint/Public/MVVMDeveloperProjectSettings.h
--- a/Engine/Plugins/Runtime/ModelViewViewModel/Source/ModelViewViewModelBlueprint/Public/MVVMDeveloperProjectSettings.h
+++ b/Engine/Plugins/Runtime/ModelViewViewModel/Source/ModelViewViewModelBlueprint/Public/MVVMDeveloperProjectSettings.h
@@ -198,7 +198,7 @@ public:
 	TSet<TSoftClassPtr<UListViewBase>> SupportedListViewBaseClassesForExtension;
 
 	/** Resolver class to use as the default value when selecting resolver creation mode */
-	UPROPERTY(EditAnywhere, config, Category = "Defaults", meta = (MetaClass = "/Script/ModelViewViewModel.UMVVMViewModelContextResolver"))
+	UPROPERTY(EditAnywhere, config, Category = "Defaults")
 	FSoftClassPath DefaultResolverValue;
 };
 
