[MVVM] Change default resolver to use TSoftClassPtr

#tests Launch editor and look at options in field in settings

diff --git a/Engine/Plugins/Runtime/ModelViewViewModel/Source/ModelViewViewModelBlueprint/Private/MVVMBlueprintViewModelContext.cpp b/Engine/Plugins/Runtime/ModelViewViewModel/Source/ModelViewViewModelBlueprint/Private/MVVMBlueprintViewModelContext.cpp
--- a/Engine/Plugins/Runtime/ModelViewViewModel/Source/ModelViewViewModelBlueprint/Private/MVVMBlueprintViewModelContext.cpp
+++ b/Engine/Plugins/Runtime/ModelViewViewModel/Source/ModelViewViewModelBlueprint/Private/MVVMBlueprintViewModelContext.cpp
@@ -41,7 +41,7 @@ TObjectPtr<UMVVMViewModelContextResolver> FMVVMBlueprintViewModelContext::Create
 
 	if (CreationType == EMVVMBlueprintViewModelContextCreationType::Resolver)
 	{
-		if (const UClass* DefaultClass = GetDefault<UMVVMDeveloperProjectSettings>()->DefaultResolverValue.ResolveClass())
+		if (const UClass* DefaultClass = GetDefault<UMVVMDeveloperProjectSettings>()->DefaultResolverValue.LoadSynchronous())
 		{
 			if (UMVVMViewModelContextResolver* DefaultResolverClass = Cast<UMVVMViewModelContextResolver>(DefaultClass->GetDefaultObject()))
 			{
diff --git a/Engine/Plugins/Runtime/ModelViewViewModel/Source/ModelViewViewModelBlueprint/Public/MVVMDeveloperProjectSettings.h b/Engine/Plugins/Runtime/ModelViewViewModel/Source/ModelViewViewModelBlueprint/Public/MVVMDeveloperProjectSettings.h
--- a/Engine/Plugins/Runtime/ModelViewViewModel/Source/ModelViewViewModelBlueprint/Public/MVVMDeveloperProjectSettings.h
+++ b/Engine/Plugins/Runtime/ModelViewViewModel/Source/ModelViewViewModelBlueprint/Public/MVVMDeveloperProjectSettings.h
@@ -8,6 +8,7 @@
 
 class UK2Node;
 class UListViewBase;
+class UMVVMViewModelContextResolver;
 class UPanelWidget;
 enum class EMVVMBlueprintViewModelContextCreationType : uint8;
 enum class EMVVMExecutionMode : uint8;
@@ -199,7 +200,7 @@ public:
 
 	/** Resolver class to use as the default value when selecting resolver creation mode */
 	UPROPERTY(EditAnywhere, config, Category = "Defaults")
-	FSoftClassPath DefaultResolverValue;
+	TSoftClassPtr<UMVVMViewModelContextResolver> DefaultResolverValue;
 };
 
 #if UE_ENABLE_INCLUDE_ORDER_DEPRECATED_IN_5_2
