MVVM/UMG: use the link icon as the Bind drop down for UMG details panel.
Add a checkmark next to the bound viewmodel property when vewing from the Bind menu.
fix pin color not reflecting the property color
fix conversion functions not showing as a binding in the Bind button

diff --git a/Engine/Plugins/Runtime/ModelViewViewModel/Source/ModelViewViewModelEditor/Private/Customizations/MVVMPropertyBindingExtension.cpp b/Engine/Plugins/Runtime/ModelViewViewModel/Source/ModelViewViewModelEditor/Private/Customizations/MVVMPropertyBindingExtension.cpp
--- a/Engine/Plugins/Runtime/ModelViewViewModel/Source/ModelViewViewModelEditor/Private/Customizations/MVVMPropertyBindingExtension.cpp
+++ b/Engine/Plugins/Runtime/ModelViewViewModel/Source/ModelViewViewModelEditor/Private/Customizations/MVVMPropertyBindingExtension.cpp
@@ -2,16 +2,20 @@
 
 #include "Customizations/MVVMPropertyBindingExtension.h"
 
+#include "Bindings/MVVMBindingHelper.h"
 #include "BlueprintEditor.h"
 #include "Components/Widget.h"
 #include "Framework/MultiBox/MultiBoxBuilder.h"
 #include "Kismet2/BlueprintEditorUtils.h"
 #include "MVVMBlueprintView.h"
+#include "MVVMBlueprintViewConversionFunction.h"
 #include "MVVMDeveloperProjectSettings.h"
 #include "MVVMEditorSubsystem.h"
 #include "MVVMWidgetBlueprintExtension_View.h"
 #include "PropertyHandle.h"
 #include "PropertyPathHelpers.h"
+#include "ScopedTransaction.h"
+#include "Styling/CoreStyle.h"
 #include "Styling/StyleColors.h"
 #include "Widgets/ViewModelFieldDragDropOp.h"
 #include "Widgets/Images/SImage.h"
@@ -21,7 +25,7 @@
 
 namespace UE::MVVM
 {
-static void ExtendBindingsMenu(FMenuBuilder& MenuBuilder, const UWidgetBlueprint* WidgetBlueprint, UWidget* Widget, TSharedPtr<IPropertyHandle> WidgetPropertyHandle)
+void FMVVMPropertyBindingExtension::ExtendBindingsMenu(FMenuBuilder& MenuBuilder, TSharedRef<FMVVMPropertyBindingExtension> MVVMPropertyBindingExtension, const UWidgetBlueprint* WidgetBlueprint, UWidget* Widget, TSharedPtr<IPropertyHandle> WidgetPropertyHandle)
 {
 	MenuBuilder.BeginSection("ViewModels", LOCTEXT("ViewModels", "View Models"));
 
@@ -38,7 +42,7 @@ static void ExtendBindingsMenu(FMenuBuilder& MenuBuilder, const UWidgetBlueprint
 
 	const UEdGraphSchema_K2* Schema = GetDefault<UEdGraphSchema_K2>();
 
-	auto CreatePropertyWidget = [Schema](const FProperty* Property, bool bRequiresConversion)
+	auto CreatePropertyWidget = [Schema, WidgetPropertyHandle, WidgetBlueprint, Widget, MVVMPropertyBindingExtension](const FProperty* Property, FGuid OwningViewModelId ,bool bRequiresConversion)
 		-> TSharedRef<SWidget>
 	{
 		FEdGraphPinType PinType;
@@ -49,8 +53,17 @@ static void ExtendBindingsMenu(FMenuBuilder& MenuBuilder, const UWidgetBlueprint
 			SNew(SHorizontalBox)
 			.ToolTipText(Property->GetDisplayNameText())
 			+ SHorizontalBox::Slot()
+			.AutoWidth()
+			.VAlign(VAlign_Center)
+			.Padding(0.0f, 0.0f, 5.0f, 0.0f)
+			[
+				SNew(SImage)
+				.Image(FCoreStyle::Get().GetBrush("Icons.Check"))
+				.Visibility(MVVMPropertyBindingExtension, &UE::MVVM::FMVVMPropertyBindingExtension::GetCheckmarkVisibility, WidgetBlueprint, Widget, Property, OwningViewModelId, WidgetPropertyHandle)
+			]
+			+ SHorizontalBox::Slot()
 			.VAlign(VAlign_Center)
-			.Padding(0, 0, 5, 0)
+			.Padding(0.0f, 0.0f, 5.0f, 0.0f)
 			.AutoWidth()
 			[
 				SNew(SImage)
@@ -58,7 +71,7 @@ static void ExtendBindingsMenu(FMenuBuilder& MenuBuilder, const UWidgetBlueprint
 				.ColorAndOpacity(Schema->GetPinTypeColor(PinType))
 			]
 			+ SHorizontalBox::Slot()
-			.Padding(0, 0, 5, 0)
+			.Padding(0.0f, 0.0f, 5.0f, 0.0f)
 			.VAlign(VAlign_Center)
 			[
 				SNew(STextBlock)
@@ -145,7 +158,7 @@ static void ExtendBindingsMenu(FMenuBuilder& MenuBuilder, const UWidgetBlueprint
 						UIAction.ExecuteAction = FExecuteAction::CreateLambda(CreateBinding, Widget, WidgetPropertyHandle, ViewModel.GetViewModelId(), VMProperty);
 						MenuBuilder.AddMenuEntry(
 							UIAction,
-							CreatePropertyWidget(VMProperty, false)
+							CreatePropertyWidget(VMProperty, ViewModel.GetViewModelId(), false)
 						);
 					}
 
@@ -165,7 +178,7 @@ static void ExtendBindingsMenu(FMenuBuilder& MenuBuilder, const UWidgetBlueprint
 						UIAction.ExecuteAction = FExecuteAction::CreateLambda(CreateBinding, Widget, WidgetPropertyHandle, ViewModel.GetViewModelId(), VMProperty);
 						MenuBuilder.AddMenuEntry(
 							UIAction,
-							CreatePropertyWidget(VMProperty, true)
+							CreatePropertyWidget(VMProperty, ViewModel.GetViewModelId(), true)
 						);
 					}
 				}));
@@ -174,6 +187,32 @@ static void ExtendBindingsMenu(FMenuBuilder& MenuBuilder, const UWidgetBlueprint
 	MenuBuilder.EndSection();
 }
 
+EVisibility FMVVMPropertyBindingExtension::GetCheckmarkVisibility(const UWidgetBlueprint* WidgetBlueprint, UWidget* Widget, const FProperty* Property, FGuid OwningViewModelId, TSharedPtr<IPropertyHandle> WidgetPropertyHandle) const
+{
+	const UMVVMWidgetBlueprintExtension_View* MVVMExtensionPtr = UMVVMWidgetBlueprintExtension_View::GetExtension<UMVVMWidgetBlueprintExtension_View>(WidgetBlueprint);
+	if (MVVMExtensionPtr == nullptr)
+	{
+		return EVisibility::Hidden;
+	}
+
+	const UMVVMBlueprintView* MVVMBlueprintView = MVVMExtensionPtr->GetBlueprintView();
+	const FMVVMBlueprintViewBinding* Binding = MVVMBlueprintView ? MVVMBlueprintView->FindBinding(Widget, WidgetPropertyHandle->GetProperty()) : nullptr;
+	if (Binding == nullptr)
+	{
+		return EVisibility::Hidden;
+	}
+
+	TArray<UE::MVVM::FMVVMConstFieldVariant> Fields = Binding->SourcePath.GetFields(WidgetBlueprint->GeneratedClass);
+
+	// Currently the bind menu only supports top-level properties in viewmodels, so we can only check the first field to find a match.
+	// This should be updated once we are able to expand the full tree of properties in the bind menu.
+	if (Fields.Num() > 0 && Fields[0].IsProperty() && Property == Fields[0].GetProperty() && Binding->SourcePath.GetViewModelId() == OwningViewModelId)
+	{
+		return EVisibility::Visible;
+	}
+	return EVisibility::Hidden;
+}
+
 TOptional<FName> FMVVMPropertyBindingExtension::GetCurrentValue(const UWidgetBlueprint* WidgetBlueprint, const UWidget* Widget, const FProperty* Property) const
 {
 	const UMVVMWidgetBlueprintExtension_View* MVVMExtensionPtr = UMVVMWidgetBlueprintExtension_View::GetExtension<UMVVMWidgetBlueprintExtension_View>(WidgetBlueprint);
@@ -194,6 +233,10 @@ TOptional<FName> FMVVMPropertyBindingExtension::GetCurrentValue(const UWidgetBlu
 	{
 		return Names.Last();
 	}
+	if (UMVVMBlueprintViewConversionFunction* ConversionFunction = Binding->Conversion.GetConversionFunction(true))
+	{
+		return ConversionFunction->GetConversionFunction().GetName();
+	}
 	return TOptional<FName>();
 }
 
@@ -233,6 +276,58 @@ const FSlateBrush* FMVVMPropertyBindingExtension::GetCurrentIcon(const UWidgetBl
 	return nullptr;
 }
 
+TOptional<FLinearColor> FMVVMPropertyBindingExtension::GetCurrentIconColor(const UWidgetBlueprint* WidgetBlueprint, const UWidget* Widget, const FProperty* Property) const
+{
+	const UMVVMWidgetBlueprintExtension_View* MVVMExtensionPtr = UMVVMWidgetBlueprintExtension_View::GetExtension<UMVVMWidgetBlueprintExtension_View>(WidgetBlueprint);
+	if (MVVMExtensionPtr == nullptr)
+	{
+		return TOptional<FLinearColor>();
+	}
+	const UMVVMBlueprintView* MVVMBlueprintView = MVVMExtensionPtr->GetBlueprintView();
+
+	const FMVVMBlueprintViewBinding* Binding = MVVMBlueprintView ? MVVMBlueprintView->FindBinding(Widget, Property) : nullptr;
+	if (Binding == nullptr)
+	{
+		return TOptional<FLinearColor>();
+	}
+
+	TArray<UE::MVVM::FMVVMConstFieldVariant> Fields = Binding->SourcePath.GetFields(WidgetBlueprint->GeneratedClass);
+	if (Fields.IsEmpty())
+	{
+		return TOptional<FLinearColor>();
+	}
+
+	UE::MVVM::FMVVMConstFieldVariant Field = Fields.Last();
+	const FProperty* IconProperty = nullptr;
+
+	if (Field.IsProperty())
+	{
+		IconProperty = Field.GetProperty();
+	}
+	else if (Field.IsFunction())
+	{
+		const UFunction* Function = Field.GetFunction();
+		const FProperty* ReturnProperty = BindingHelper::GetReturnProperty(Function);
+		if (ReturnProperty != nullptr)
+		{
+			IconProperty = ReturnProperty;
+		}
+		else
+		{
+			IconProperty = BindingHelper::GetFirstArgumentProperty(Function);
+		}
+	}
+	if (IconProperty != nullptr)
+	{
+		FSlateColor PrimaryColor, SecondaryColor;
+		const FSlateBrush* SecondaryBrush = nullptr;
+		const FSlateBrush* PrimaryBrush = FBlueprintEditor::GetVarIconAndColorFromProperty(IconProperty, PrimaryColor, SecondaryBrush, SecondaryColor);
+		return PrimaryColor.GetSpecifiedColor();
+	}
+
+	return TOptional<FLinearColor>();
+}
+
 void FMVVMPropertyBindingExtension::ClearCurrentValue(const UWidgetBlueprint* WidgetBlueprint, const UWidget* Widget, const FProperty* Property)
 {
 	if (UMVVMWidgetBlueprintExtension_View* MVVMExtensionPtr = UMVVMWidgetBlueprintExtension_View::GetExtension<UMVVMWidgetBlueprintExtension_View>(WidgetBlueprint))
@@ -241,7 +336,9 @@ void FMVVMPropertyBindingExtension::ClearCurrentValue(const UWidgetBlueprint* Wi
 		{
 			if (FMVVMBlueprintViewBinding* Binding = MVVMBlueprintView->FindBinding(Widget, Property))
 			{
-				Binding->SourcePath.ResetPropertyPath();
+				FScopedTransaction Transaction(LOCTEXT("DeleteBindingsTransaction", "Delete Binding"));
+				MVVMBlueprintView->Modify();
+				MVVMBlueprintView->RemoveBinding(Binding);
 			}
 		}
 	}
@@ -250,7 +347,7 @@ void FMVVMPropertyBindingExtension::ClearCurrentValue(const UWidgetBlueprint* Wi
 TSharedPtr<FExtender> FMVVMPropertyBindingExtension::CreateMenuExtender(const UWidgetBlueprint* WidgetBlueprint, UWidget* Widget, TSharedPtr<IPropertyHandle> WidgetPropertyHandle)
 {
 	TSharedPtr<FExtender> Extender = MakeShared<FExtender>();
-	Extender->AddMenuExtension("BindingActions", EExtensionHook::After, nullptr, FMenuExtensionDelegate::CreateStatic(&ExtendBindingsMenu, WidgetBlueprint, Widget, WidgetPropertyHandle));
+	Extender->AddMenuExtension("BindingActions", EExtensionHook::After, nullptr, FMenuExtensionDelegate::CreateStatic(&ExtendBindingsMenu, MakeShared<FMVVMPropertyBindingExtension>(*this), WidgetBlueprint, Widget, WidgetPropertyHandle));
 	return Extender;
 }
 
diff --git a/Engine/Plugins/Runtime/ModelViewViewModel/Source/ModelViewViewModelEditor/Private/Customizations/MVVMPropertyBindingExtension.h b/Engine/Plugins/Runtime/ModelViewViewModel/Source/ModelViewViewModelEditor/Private/Customizations/MVVMPropertyBindingExtension.h
--- a/Engine/Plugins/Runtime/ModelViewViewModel/Source/ModelViewViewModelEditor/Private/Customizations/MVVMPropertyBindingExtension.h
+++ b/Engine/Plugins/Runtime/ModelViewViewModel/Source/ModelViewViewModelEditor/Private/Customizations/MVVMPropertyBindingExtension.h
@@ -6,6 +6,7 @@
 
 class UWidgetBlueprint;
 class UWidget;
+struct EVisibility;
 
 namespace UE::MVVM
 {
@@ -13,14 +14,19 @@ namespace UE::MVVM
 class FMVVMPropertyBindingExtension
 	: public IPropertyBindingExtension
 {
+	static void ExtendBindingsMenu(FMenuBuilder& MenuBuilder, TSharedRef<FMVVMPropertyBindingExtension> MVVMPropertyBindingExtension, const UWidgetBlueprint* WidgetBlueprint, UWidget* Widget, TSharedPtr<IPropertyHandle> WidgetPropertyHandle);
+
 	virtual TOptional<FName> GetCurrentValue(const UWidgetBlueprint* WidgetBlueprint, const UWidget* Widget, const FProperty* Property) const override;
 	virtual const FSlateBrush* GetCurrentIcon(const UWidgetBlueprint* WidgetBlueprint, const UWidget* Widget, const FProperty* Property) const override;
+	virtual TOptional<FLinearColor> GetCurrentIconColor(const UWidgetBlueprint* WidgetBlueprint, const UWidget* Widget, const FProperty* Property) const override;
 
 	virtual void ClearCurrentValue(const UWidgetBlueprint* WidgetBlueprint, const UWidget* Widget, const FProperty* Property) override;
 	virtual TSharedPtr<FExtender> CreateMenuExtender(const UWidgetBlueprint* WidgetBlueprint, UWidget* Widget, TSharedPtr<IPropertyHandle> WidgetPropertyHandle) override;
 	virtual bool CanExtend(const UWidgetBlueprint* WidgetBlueprint, const UWidget* Widget, const FProperty* Property) const;
 	virtual EDropResult OnDrop(const FGeometry& Geometry, const FDragDropEvent& DragDropEvent, UWidgetBlueprint* WidgetBlueprint, UWidget* Widget, TSharedPtr<IPropertyHandle> WidgetPropertyHandle) override;
 	virtual TSharedPtr<FExtender> CreateMenuExtender(const UWidgetBlueprint* WidgetBlueprint, const UWidget* Widget, const FProperty* Property) override;
+
+	EVisibility GetCheckmarkVisibility(const UWidgetBlueprint* WidgetBlueprint, UWidget* Widget, const FProperty* Property, FGuid OwningViewModelId, TSharedPtr<IPropertyHandle> WidgetPropertyHandle) const;
 };
 
 }
diff --git a/Engine/Plugins/Runtime/PropertyAccess/Source/PropertyAccessEditor/Private/SPropertyBinding.cpp b/Engine/Plugins/Runtime/PropertyAccess/Source/PropertyAccessEditor/Private/SPropertyBinding.cpp
--- a/Engine/Plugins/Runtime/PropertyAccess/Source/PropertyAccessEditor/Private/SPropertyBinding.cpp
+++ b/Engine/Plugins/Runtime/PropertyAccess/Source/PropertyAccessEditor/Private/SPropertyBinding.cpp
@@ -4,6 +4,7 @@
 #include "Framework/MultiBox/MultiBoxBuilder.h"
 #include "Widgets/Input/SButton.h"
 #include "Widgets/Input/SComboButton.h"
+#include "Widgets/SNullWidget.h"
 
 #if WITH_EDITOR
 #include "Engine/BlueprintGeneratedClass.h"
@@ -19,6 +20,7 @@
 #include "Kismet/BlueprintFunctionLibrary.h"
 #include "BlueprintActionDatabase.h"
 #include "Preferences/PersonaOptions.h"
+#include "Styling/SlateTypes.h"
 
 #define LOCTEXT_NAMESPACE "PropertyBinding"
 
@@ -59,6 +61,15 @@ void SPropertyBinding::Construct(const FArguments& InArgs, UBlueprint* InBluepri
 			BuildContextStructCategoryRecursive(Categories, Section->SubCategories, Index);
 		}
 	}
+	TSharedRef<SWidget> LinkIcon = Args.bUseLinkIconStyle ?
+		SNew(SBox)
+		.HeightOverride(16.0f)
+		[
+			SNew(SImage)
+			.Image(this, &SPropertyBinding::GetLinkIcon)
+		]
+		:
+		SNullWidget::NullWidget;
 
 	ChildSlot
 	[
@@ -72,6 +83,7 @@ void SPropertyBinding::Construct(const FArguments& InArgs, UBlueprint* InBluepri
 			.ToolTipText(this, &SPropertyBinding::GetCurrentBindingToolTipText)
 			.OnGetMenuContent(this, &SPropertyBinding::OnGenerateDelegateMenu)
 			.ContentPadding(1)
+			.HasDownArrow(!Args.bUseLinkIconStyle)
 			.ButtonContent()
 			[
 				SNew(SHorizontalBox)
@@ -79,6 +91,12 @@ void SPropertyBinding::Construct(const FArguments& InArgs, UBlueprint* InBluepri
 				+ SHorizontalBox::Slot()
 				.AutoWidth()
 				.VAlign(VAlign_Center)
+				[
+					LinkIcon
+				]
+				+ SHorizontalBox::Slot()
+				.AutoWidth()
+				.VAlign(VAlign_Center)
 				[
 					SNew(SBox)
 					.HeightOverride(16.0f)
@@ -1048,6 +1066,12 @@ void SPropertyBinding::FillCategoryMenu(FMenuBuilder& MenuBuilder, const FBindin
 	MenuBuilder.EndSection();
 }
 
+const FSlateBrush* SPropertyBinding::GetLinkIcon() const
+{
+	bool CheckState = HasAnyBindings();
+	return (CheckState) ? FAppStyle::GetBrush("Icons.Link") : FAppStyle::GetBrush("Icons.Unlink");
+}
+
 const FSlateBrush* SPropertyBinding::GetCurrentBindingImage() const
 {
 	if(Args.CurrentBindingImage.IsSet())
@@ -1297,4 +1321,14 @@ FReply SPropertyBinding::OnDrop(const FGeometry& MyGeometry, const FDragDropEven
 	return FReply::Unhandled();
 }
 
+bool SPropertyBinding::HasAnyBindings() const
+{
+	if (Args.OnHasAnyBindings.IsBound())
+	{
+		return Args.OnHasAnyBindings.Execute();
+	}
+
+	return false;
+}
+
 #undef LOCTEXT_NAMESPACE
diff --git a/Engine/Plugins/Runtime/PropertyAccess/Source/PropertyAccessEditor/Private/SPropertyBinding.h b/Engine/Plugins/Runtime/PropertyAccess/Source/PropertyAccessEditor/Private/SPropertyBinding.h
--- a/Engine/Plugins/Runtime/PropertyAccess/Source/PropertyAccessEditor/Private/SPropertyBinding.h
+++ b/Engine/Plugins/Runtime/PropertyAccess/Source/PropertyAccessEditor/Private/SPropertyBinding.h
@@ -58,6 +58,7 @@ protected:
 	void FillPropertyMenu(FMenuBuilder& MenuBuilder, UStruct* InOwnerStruct, TArray<TSharedPtr<FBindingChainElement>> InBindingChain);
 	void FillCategoryMenu(FMenuBuilder& MenuBuilder, const FBindingContextStructCategory* Category);
 
+	const FSlateBrush* GetLinkIcon() const;
 	const FSlateBrush* GetCurrentBindingImage() const;
 	FText GetCurrentBindingText() const;
 	FText GetCurrentBindingToolTipText() const;
@@ -85,6 +86,8 @@ protected:
 
 	virtual FReply OnDrop(const FGeometry& MyGeometry, const FDragDropEvent& DragDropEvent) override;
 
+	bool HasAnyBindings() const;
+
 private:
 
 	bool IsClassDenied(UClass* OwnerClass) const;
diff --git a/Engine/Source/Editor/UMGEditor/Private/Customizations/UMGDetailCustomizations.cpp b/Engine/Source/Editor/UMGEditor/Private/Customizations/UMGDetailCustomizations.cpp
--- a/Engine/Source/Editor/UMGEditor/Private/Customizations/UMGDetailCustomizations.cpp
+++ b/Engine/Source/Editor/UMGEditor/Private/Customizations/UMGDetailCustomizations.cpp
@@ -38,6 +38,7 @@
 #include "UMGEditorModule.h"
 #include "WidgetBlueprintEditor.h"
 #include "Animation/WidgetAnimation.h"
+#include "Styling/SlateTypes.h"
 
 #define LOCTEXT_NAMESPACE "UMG"
 
@@ -149,6 +150,8 @@ TSharedRef<SWidget> FBlueprintWidgetCustomization::MakePropertyBindingWidget(TWe
 	Args.MenuExtender = FExtender::Combine(MenuExtenders);
 	Args.Property = InPropertyHandle->GetProperty();
 	Args.BindableSignature = SignatureFunction;
+	Args.BindButtonStyle = &FCoreStyle::Get().GetWidgetStyle<FButtonStyle>("Button");
+
 	Args.OnGenerateBindingName = FOnGenerateBindingName::CreateLambda([WidgetName]()
 	{
 		return WidgetName;
@@ -315,6 +318,39 @@ TSharedRef<SWidget> FBlueprintWidgetCustomization::MakePropertyBindingWidget(TWe
 		return InClass->IsChildOf(UWidget::StaticClass());
 	});
 
+	Args.OnHasAnyBindings = FOnHasAnyBindings::CreateLambda([InEditor, InPropertyHandle]()
+		{
+			return HasPropertyBindings(InEditor, InPropertyHandle);
+		});
+
+	Args.CurrentBindingColor = MakeAttributeLambda([InEditor, Objects, InPropertyHandle, ActiveExtensions]()
+		{
+			UWidgetBlueprint* ThisBlueprint = InEditor.Pin()->GetWidgetBlueprintObj();
+			for (const TWeakObjectPtr<UObject>& ObjectPtr : Objects)
+			{
+				UObject* Object = ObjectPtr.Get();
+
+				// Ignore null outer objects
+				if (Object == nullptr)
+				{
+					continue;
+				}
+				if (UWidget* Widget = Cast<UWidget>(Object))
+				{
+					for (const TSharedPtr<IPropertyBindingExtension>& Extension : ActiveExtensions)
+					{
+						TOptional<FLinearColor> Color = Extension->GetCurrentIconColor(ThisBlueprint, Widget, InPropertyHandle->GetProperty());
+						if (Color.IsSet())
+						{
+							return Color.GetValue();
+						}
+					}
+				}
+			}
+
+			return FLinearColor::White;
+		});
+
 	Args.OnAddBinding = FOnAddBinding::CreateLambda([InEditor, Objects](FName InPropertyName, const TArray<FBindingChainElement>& InBindingChain)
 	{
 		UWidgetBlueprint* ThisBlueprint = InEditor.Pin()->GetWidgetBlueprintObj();
@@ -550,7 +586,7 @@ TSharedRef<SWidget> FBlueprintWidgetCustomization::MakePropertyBindingWidget(TWe
 			break;
 		}
 
-		return LOCTEXT("Bind", "Bind");
+		return FText::GetEmpty();
 	});
 
 	Args.CurrentBindingImage = MakeAttributeLambda([InEditor, Objects, InPropertyHandle, ActiveExtensions]() -> const FSlateBrush*
@@ -607,6 +643,7 @@ TSharedRef<SWidget> FBlueprintWidgetCustomization::MakePropertyBindingWidget(TWe
 
 	Args.bGeneratePureBindings = bInGeneratePureBindings;
 	Args.bAllowNewBindings = bAllowDetailsPanelLegacyBinding;
+	Args.bUseLinkIconStyle = true;
 
 	IPropertyAccessEditor& PropertyAccessEditor = IModularFeatures::Get().GetModularFeature<IPropertyAccessEditor>("PropertyAccessEditor");
 	return PropertyAccessEditor.MakePropertyBindingWidget(InEditor.Pin()->GetBlueprintObj(), Args);
diff --git a/Engine/Source/Editor/UMGEditor/Public/IHasPropertyBindingExtensibility.h b/Engine/Source/Editor/UMGEditor/Public/IHasPropertyBindingExtensibility.h
--- a/Engine/Source/Editor/UMGEditor/Public/IHasPropertyBindingExtensibility.h
+++ b/Engine/Source/Editor/UMGEditor/Public/IHasPropertyBindingExtensibility.h
@@ -29,6 +29,7 @@ public:
 	virtual void ClearCurrentValue(const UWidgetBlueprint* WidgetBlueprint, const UWidget* Widget, const FProperty* Property) = 0;
 	virtual TOptional<FName> GetCurrentValue(const UWidgetBlueprint* WidgetBlueprint, const UWidget* Widget, const FProperty* Property) const = 0;
 	virtual const FSlateBrush* GetCurrentIcon(const UWidgetBlueprint* WidgetBlueprint, const UWidget* Widget, const FProperty* Property) const = 0;
+	virtual TOptional<FLinearColor> GetCurrentIconColor(const UWidgetBlueprint* WidgetBlueprint, const UWidget* Widget, const FProperty* Property) const = 0;
 	virtual EDropResult OnDrop(const FGeometry& Geometry, const FDragDropEvent& DragDropEvent, UWidgetBlueprint* WidgetBlueprint, UWidget* Widget, TSharedPtr<IPropertyHandle> WidgetPropertyHandle) = 0;
 
 	UE_DEPRECATED(5.3, "The function CreateMenuExtender with FProperty parameter is deprecated. Please pass the property handler as the parameter instead.")
diff --git a/Engine/Source/Editor/UnrealEd/Public/IPropertyAccessEditor.h b/Engine/Source/Editor/UnrealEd/Public/IPropertyAccessEditor.h
--- a/Engine/Source/Editor/UnrealEd/Public/IPropertyAccessEditor.h
+++ b/Engine/Source/Editor/UnrealEd/Public/IPropertyAccessEditor.h
@@ -108,7 +108,7 @@ DECLARE_DELEGATE_RetVal(FString, FOnGenerateBindingName);
 /** Delegate used to open a binding (e.g. a function) */
 DECLARE_DELEGATE_RetVal_OneParam(bool, FOnGotoBinding, FName /*InPropertyName*/);
 
-/** Delegate used to se if we can open a binding (e.g. a function) */
+/** Delegate used to see if we can open a binding (e.g. a function) */
 DECLARE_DELEGATE_RetVal_OneParam(bool, FOnCanGotoBinding, FName /*InPropertyName*/);
 
 /** Delegate used to check whether a property is considered for binding. Returning false will discard the property and all child properties. */
@@ -144,7 +144,7 @@ DECLARE_DELEGATE_TwoParams(FOnAddBinding, FName /*InPropertyName*/, const TArray
 /** Delegate called to remove a binding */
 DECLARE_DELEGATE_OneParam(FOnRemoveBinding, FName /*InPropertyName*/);
 
-/** Delegate called to see if we can remove remove a binding (ie. if it exists) */
+/** Delegate called to see if we can remove a binding (ie. if it exists) */
 DECLARE_DELEGATE_RetVal_OneParam(bool, FOnCanRemoveBinding, FName /*InPropertyName*/);
 
 /** Delegate called once a new function binding has been created */
@@ -156,6 +156,9 @@ DECLARE_DELEGATE_RetVal_OneParam(UStruct*, FOnResolveIndirection, const TArray<F
 /** Delegate called once a drag-drop event is dropped on the binding widget */
 DECLARE_DELEGATE_RetVal_TwoParams(FReply, FOnDrop, const FGeometry&, const FDragDropEvent&);
 
+/** Delegate called to see if the property has any bindings */
+DECLARE_DELEGATE_RetVal(bool, FOnHasAnyBindings);
+
 /** Setup arguments structure for a property binding widget */
 struct FPropertyBindingWidgetArgs
 {
@@ -219,6 +222,9 @@ struct FPropertyBindingWidgetArgs
 	/** Delegate called to see if we can remove remove a binding (ie. if it exists) */
 	FOnCanRemoveBinding OnCanRemoveBinding;
 
+	/** Delegate called to see if the property has any bindings */
+	FOnHasAnyBindings OnHasAnyBindings;
+
 	/** Delegate called once a new function binding has been created */
 	FOnNewFunctionBindingCreated OnNewFunctionBindingCreated;
 
@@ -278,6 +284,9 @@ struct FPropertyBindingWidgetArgs
 
 	/** Whether to allow UScriptStruct functions as non-leaf nodes */
 	bool bAllowStructFunctions = false;	
+
+	/** Whether to format the SPropertyBinding widget to use a link icon (versus the standard combo button) */
+	bool bUseLinkIconStyle = false;
 };
 
 /** Enum describing the result of ResolvePropertyAccess */
