Add AreSourcesValidForBinding to Async, needed should async finish after binding src / dst become gc'd

#tests PIE, cooked, validation

diff --git a/Engine/Plugins/Runtime/ModelViewViewModel/Source/ModelViewViewModel/Private/View/MVVMView.cpp b/Engine/Plugins/Runtime/ModelViewViewModel/Source/ModelViewViewModel/Private/View/MVVMView.cpp
--- a/Engine/Plugins/Runtime/ModelViewViewModel/Source/ModelViewViewModel/Private/View/MVVMView.cpp
+++ b/Engine/Plugins/Runtime/ModelViewViewModel/Source/ModelViewViewModel/Private/View/MVVMView.cpp
@@ -1237,4 +1237,37 @@ bool UMVVMView::AreSourcesValidForEvent(int32 EventKeyIndex) const
 	return false;
 }
 
+
+bool UMVVMView::AreSourcesValidForBinding(int32 InBindingKey) const
+{
+	FMVVMViewClass_BindingKey BindingKey = FMVVMViewClass_BindingKey(InBindingKey);
+	if (GeneratedViewClass)
+	{
+		const FMVVMViewClass_Binding& ClassBinding = GeneratedViewClass->GetBinding(BindingKey);
+		uint64 BindingSources = ClassBinding.GetSources();
+		if ((BindingSources & ValidSources) == BindingSources)
+		{
+			return true;
+		}
+		
+		const uint64 MissingSources = BindingSources & (~ValidSources);
+		if ((MissingSources & GeneratedViewClass->GetOptionalSources()) != MissingSources)
+		{
+			// Only log invalid bindings as info, since it is possible expected behavior for async functions
+#if UE_WITH_MVVM_DEBUGGING
+			UE::MVVM::FMessageLog Log(GetUserWidget());
+			Log.Info(FText::Format(LOCTEXT("ExecuteBindingFailInvalidSource", "The Binding '{0}' was not executed. There are invalid sources.")
+				, FText::FromString(ClassBinding.ToString(GeneratedViewClass, FMVVMViewClass_Binding::FToStringArgs::Short()))
+			));
+#else
+			UE::MVVM::FMessageLog Log(GetUserWidget());
+			Log.Info(FText::Format(LOCTEXT("ExecuteBindingFailInvalidSource", "The Binding '{0}' was not executed. There are invalid sources.")
+				, FText::AsNumber(BindingKey.GetIndex())
+			));
+#endif
+		}
+	}
+	return false;
+}
+
 #undef LOCTEXT_NAMESPACE
diff --git a/Engine/Plugins/Runtime/ModelViewViewModel/Source/ModelViewViewModel/Public/View/MVVMView.h b/Engine/Plugins/Runtime/ModelViewViewModel/Source/ModelViewViewModel/Public/View/MVVMView.h
--- a/Engine/Plugins/Runtime/ModelViewViewModel/Source/ModelViewViewModel/Public/View/MVVMView.h
+++ b/Engine/Plugins/Runtime/ModelViewViewModel/Source/ModelViewViewModel/Public/View/MVVMView.h
@@ -181,6 +181,13 @@ private:
 	UFUNCTION(BlueprintCallable, Category = "View", meta=(BlueprintInternalUseOnly="true"))
 	bool AreSourcesValidForEvent(int32 EventKey) const;
 
+	/**
+	 * Checks if all source bindings (Src / Dst objects) are valid for a binding
+	 * Currently used to verify if after effects of async conversion functions are safe to trigger.
+	 */
+	UFUNCTION(BlueprintCallable, Category = "View", meta = (BlueprintInternalUseOnly = "true"))
+	bool AreSourcesValidForBinding(int32 BindingKey) const;
+
 private:
 	//~ Source
 	void InitializeSource(FMVVMView_SourceKey SourceKey);
diff --git a/Engine/Plugins/Runtime/ModelViewViewModel/Source/ModelViewViewModelBlueprint/Private/Bindings/MVVMConversionFunctionHelper.cpp b/Engine/Plugins/Runtime/ModelViewViewModel/Source/ModelViewViewModelBlueprint/Private/Bindings/MVVMConversionFunctionHelper.cpp
--- a/Engine/Plugins/Runtime/ModelViewViewModel/Source/ModelViewViewModelBlueprint/Private/Bindings/MVVMConversionFunctionHelper.cpp
+++ b/Engine/Plugins/Runtime/ModelViewViewModel/Source/ModelViewViewModelBlueprint/Private/Bindings/MVVMConversionFunctionHelper.cpp
@@ -23,6 +23,7 @@
 #include "K2Node_VariableSet.h"
 #include "Kismet2/BlueprintEditorUtils.h"
 #include "MVVMConversionFunctionGraphSchema.h"
+#include "Node/MVVMK2Node_AreSourcesValidForBinding.h"
 
 #define LOCTEXT_NAMESPACE "MVVMConversionFunctionHelper"
 
@@ -1207,10 +1208,23 @@ TValueOrError<FCreateGraphResult, FText> CreateSetterGraph(UBlueprint* Blueprint
 			UK2Node* CallFunctionNode = *CallFunctionNodeItr;
 			UEdGraphPin* AsyncCompletedPin = CallFunctionNode->FindPinChecked(UEdGraphSchema_K2::PN_Completed);
 
-			// Link the Async result Exec with the setter
+			// Link the Async result Exec with the setter, adding a valid binding check should sources expire after the async
 			if (AsyncCompletedPin)
 			{
-				GraphSchema->TryCreateConnection(AsyncCompletedPin, SetterNode->GetExecPin());
+				UMVVMK2Node_AreSourcesValidForBinding* BranchNode = nullptr;
+				{
+					FGraphNodeCreator<UMVVMK2Node_AreSourcesValidForBinding> BranchNodeCreator(*LinkedSetterGraph.GetValue().NewGraph);
+					BranchNode = BranchNodeCreator.CreateNode(false, UMVVMK2Node_AreSourcesValidForBinding::StaticClass());
+					BranchNode->NodePosX = SetterNode->NodePosX;
+					BranchNode->NodePosY = SetterNode->NodePosY - 300;
+					BranchNodeCreator.Finalize();
+				}
+
+				if (ensure(BranchNode))
+				{
+					GraphSchema->TryCreateConnection(AsyncCompletedPin, BranchNode->GetExecPin());
+					GraphSchema->TryCreateConnection(BranchNode->GetThenPin(), SetterNode->GetExecPin());
+				}
 			}
 
 			// Link the Async result value with the setter args
diff --git a/Engine/Plugins/Runtime/ModelViewViewModel/Source/ModelViewViewModelBlueprint/Private/MVVMViewBlueprintCompiler.h b/Engine/Plugins/Runtime/ModelViewViewModel/Source/ModelViewViewModelBlueprint/Private/MVVMViewBlueprintCompiler.h
--- a/Engine/Plugins/Runtime/ModelViewViewModel/Source/ModelViewViewModelBlueprint/Private/MVVMViewBlueprintCompiler.h
+++ b/Engine/Plugins/Runtime/ModelViewViewModel/Source/ModelViewViewModelBlueprint/Private/MVVMViewBlueprintCompiler.h
@@ -93,6 +93,7 @@ private:
 	void CreateWriteFieldContexts(const FWidgetBlueprintCompilerContext::FCreateFunctionContext& Context);
 	void CreateViewModelSetters(const FWidgetBlueprintCompilerContext::FCreateFunctionContext& Context);
 	void CreateIntermediateGraphFunctions(const FWidgetBlueprintCompilerContext::FCreateFunctionContext& Context);
+	void CategorizeAsyncFunctions(const FWidgetBlueprintCompilerContext::FCreateFunctionContext& Context);
 
 	// PreCompile
 	void FixCompilerBindingSelfSource(UWidgetBlueprintGeneratedClass* Class);
diff --git a/Engine/Plugins/Runtime/ModelViewViewModel/Source/ModelViewViewModelBlueprint/Private/Node/MVVMK2Node_AreSourcesValidForBinding.cpp b/Engine/Plugins/Runtime/ModelViewViewModel/Source/ModelViewViewModelBlueprint/Private/Node/MVVMK2Node_AreSourcesValidForBinding.cpp
new file mode 100644
--- /dev/null
+++ b/Engine/Plugins/Runtime/ModelViewViewModel/Source/ModelViewViewModelBlueprint/Private/Node/MVVMK2Node_AreSourcesValidForBinding.cpp
@@ -0,0 +1,113 @@
+// Copyright Epic Games, Inc. All Rights Reserved.
+
+#include "Node/MVVMK2Node_AreSourcesValidForBinding.h"
+
+#include "BlueprintCompiledStatement.h"
+#include "BlueprintNodeSpawner.h"
+#include "EdGraph/EdGraphPin.h"
+#include "EdGraphSchema_K2.h"
+#include "Kismet2/CompilerResultsLog.h"
+#include "KismetCompiledFunctionContext.h"
+#include "KismetCompiler.h"
+#include "MVVMSubsystem.h"
+#include "Subsystems/SubsystemBlueprintLibrary.h"
+#include "View/MVVMView.h"
+
+#include "K2Node_CallFunction.h"
+#include "K2Node_IfThenElse.h"
+#include "K2Node_Self.h"
+
+#define LOCTEXT_NAMESPACE "MVVMK2Node_AreSourcesValidForBinding"
+
+//////////////////////////////////////////////////////////////////////////
+
+void UMVVMK2Node_AreSourcesValidForBinding::AllocateDefaultPins()
+{
+	const UEdGraphSchema_K2* K2Schema = GetDefault<UEdGraphSchema_K2>();
+
+	CreatePin(EGPD_Input, UEdGraphSchema_K2::PC_Exec, UEdGraphSchema_K2::PN_Execute);
+
+	UEdGraphPin* TruePin = CreatePin(EGPD_Output, UEdGraphSchema_K2::PC_Exec, UEdGraphSchema_K2::PN_Then);
+	TruePin->PinFriendlyName = LOCTEXT("true", "true");
+
+	UEdGraphPin* FalsePin = CreatePin(EGPD_Output, UEdGraphSchema_K2::PC_Exec, UEdGraphSchema_K2::PN_Else);
+	FalsePin->PinFriendlyName = LOCTEXT("false", "false");
+
+	Super::AllocateDefaultPins();
+}
+
+void UMVVMK2Node_AreSourcesValidForBinding::ExpandNode(FKismetCompilerContext& CompilerContext, UEdGraph* SourceGraph)
+{
+	Super::ExpandNode(CompilerContext, SourceGraph);
+
+	if (!BindingKey.IsValid())
+	{
+		CompilerContext.MessageLog.Error(*LOCTEXT("NoBindingKey", "Node @@ doesn't have a valid event key.").ToString(), this);
+		BreakAllNodeLinks();
+		return;
+	}
+
+	const UEdGraphSchema_K2* Schema = CompilerContext.GetSchema();
+
+	UK2Node_CallFunction* CallGetSubsystemNode = CompilerContext.SpawnIntermediateNode<UK2Node_CallFunction>(this, SourceGraph);
+	{
+		CallGetSubsystemNode->FunctionReference.SetExternalMember(GET_FUNCTION_NAME_CHECKED(USubsystemBlueprintLibrary, GetEngineSubsystem), USubsystemBlueprintLibrary::StaticClass());
+		CallGetSubsystemNode->AllocateDefaultPins();
+		UEdGraphPin* CallCreateClassTypePin = CallGetSubsystemNode->FindPinChecked(FName("Class"));
+		CallCreateClassTypePin->DefaultObject = UMVVMSubsystem::StaticClass();
+	}
+
+	UK2Node_CallFunction* CallGetViewNode = CompilerContext.SpawnIntermediateNode<UK2Node_CallFunction>(this, SourceGraph);
+	{
+		CallGetViewNode->FunctionReference.SetExternalMember(GET_FUNCTION_NAME_CHECKED(UMVVMSubsystem, K2_GetViewFromUserWidget), UMVVMSubsystem::StaticClass());
+		CallGetViewNode->AllocateDefaultPins();
+	}
+
+	UK2Node_Self* SelfNode = CompilerContext.SpawnIntermediateNode<UK2Node_Self>(this, SourceGraph);
+	{
+		SelfNode->AllocateDefaultPins();
+	}
+
+	UK2Node_CallFunction* CallAreSourcesValidForBindingNode = CompilerContext.SpawnIntermediateNode<UK2Node_CallFunction>(this, SourceGraph);
+	{
+		CallAreSourcesValidForBindingNode->FunctionReference.SetExternalMember(FName("AreSourcesValidForBinding"), UMVVMView::StaticClass());
+		CallAreSourcesValidForBindingNode->AllocateDefaultPins();
+
+		UEdGraphPin* BindingKeyPin = CallAreSourcesValidForBindingNode->FindPinChecked(FName("BindingKey"));
+		Schema->TrySetDefaultValue(*BindingKeyPin, LexToString(BindingKey.GetIndex()));
+	}
+
+	UK2Node_IfThenElse* BranchNode = CompilerContext.SpawnIntermediateNode<UK2Node_IfThenElse>(this, SourceGraph);
+	{
+		BranchNode->AllocateDefaultPins();
+	}
+
+	// Casting subsytem.result into MVVMSubsystem
+	CallGetSubsystemNode->GetReturnValuePin()->PinType = CallGetViewNode->FindPinChecked(UEdGraphSchema_K2::PN_Self)->PinType;
+	// subsystem.result to GetViewFromUserWidget.target
+	check(Schema->TryCreateConnection(CallGetSubsystemNode->GetReturnValuePin(), CallGetViewNode->FindPinChecked(UEdGraphSchema_K2::PN_Self)));
+	// self to GetViewFromUserWidget.UserWidget
+	check(Schema->TryCreateConnection(SelfNode->FindPinChecked(UEdGraphSchema_K2::PN_Self), CallGetViewNode->FindPinChecked(FName("UserWidget"))));
+	// GetViewFromUserWidget.Result to AreSourcesValidForBinding.target
+	check(Schema->TryCreateConnection(CallGetViewNode->FindPinChecked(UEdGraphSchema_K2::PN_ReturnValue), CallAreSourcesValidForBindingNode->FindPinChecked(UEdGraphSchema_K2::PN_Self)));
+	// AreSourcesValidForBinding.result to branch.condition
+	check(Schema->TryCreateConnection(CallAreSourcesValidForBindingNode->FindPinChecked(UEdGraphSchema_K2::PN_ReturnValue), BranchNode->GetConditionPin()));
+	// move this.exec to branch.exec
+	CompilerContext.MovePinLinksToIntermediate(*GetExecPin(), *BranchNode->GetExecPin());
+	// move this.then to branch.then
+	CompilerContext.MovePinLinksToIntermediate(*GetThenPin(), *BranchNode->GetThenPin());
+
+	BreakAllNodeLinks();
+}
+
+UEdGraphPin* UMVVMK2Node_AreSourcesValidForBinding::GetThenPin() const
+{
+	return FindPinChecked(UEdGraphSchema_K2::PN_Then);
+}
+
+UEdGraphPin* UMVVMK2Node_AreSourcesValidForBinding::GetElsePin() const
+{
+	return FindPinChecked(UEdGraphSchema_K2::PN_Else);
+}
+
+#undef LOCTEXT_NAMESPACE
diff --git a/Engine/Plugins/Runtime/ModelViewViewModel/Source/ModelViewViewModelBlueprint/Private/Node/MVVMK2Node_AreSourcesValidForBinding.h b/Engine/Plugins/Runtime/ModelViewViewModel/Source/ModelViewViewModelBlueprint/Private/Node/MVVMK2Node_AreSourcesValidForBinding.h
new file mode 100644
--- /dev/null
+++ b/Engine/Plugins/Runtime/ModelViewViewModel/Source/ModelViewViewModelBlueprint/Private/Node/MVVMK2Node_AreSourcesValidForBinding.h
@@ -0,0 +1,33 @@
+// Copyright Epic Games, Inc. All Rights Reserved.
+
+#pragma once
+
+#include "CoreMinimal.h"
+#include "K2Node.h"
+#include "View/MVVMViewTypes.h"
+
+#include "MVVMK2Node_AreSourcesValidForBinding.generated.h"
+
+class UEdGraphPin;
+class FKismetCompilerContext;
+
+UCLASS(MinimalAPI, Hidden, HideDropDown)
+class UMVVMK2Node_AreSourcesValidForBinding : public UK2Node
+{
+	GENERATED_BODY()
+
+public:
+	//~ Begin K2Node Interface.
+	virtual void ExpandNode(FKismetCompilerContext& CompilerContext, UEdGraph* SourceGraph) override;
+	virtual void AllocateDefaultPins() override;
+	//~ End K2Node Interface.
+
+	/** Get the then output pin */
+	UEdGraphPin* GetThenPin() const;
+	/** Get the else output pin */
+	UEdGraphPin* GetElsePin() const;
+
+	UPROPERTY()
+	FMVVMViewClass_BindingKey BindingKey;
+};
+
