Add Async MakeBrush & LoadTexture / LoadMaterial Nodes for MVVM

#tests PIE, cooked, validation

diff --git a/Engine/Plugins/Runtime/ModelViewViewModel/Source/ModelViewViewModelBlueprint/Private/MVVMConversionFunctionGraphSchema.h b/Engine/Plugins/Runtime/ModelViewViewModel/Source/ModelViewViewModelBlueprint/Private/MVVMConversionFunctionGraphSchema.h
--- a/Engine/Plugins/Runtime/ModelViewViewModel/Source/ModelViewViewModelBlueprint/Private/MVVMConversionFunctionGraphSchema.h
+++ b/Engine/Plugins/Runtime/ModelViewViewModel/Source/ModelViewViewModelBlueprint/Private/MVVMConversionFunctionGraphSchema.h
@@ -24,7 +24,7 @@ public:
  * 
  * Note: Keep in private header.
  */
-UCLASS()
+UCLASS(Hidden, HideDropDown)
 class UMVVMFakeTestUbergraphSchema : public UEdGraphSchema_K2
 {
 	GENERATED_BODY()
@@ -38,7 +38,7 @@ public:
  *
  * Note: Keep in private header.
  */
-UCLASS()
+UCLASS(Hidden, HideDropDown)
 class UMVVMFakeTestUbergraph : public UEdGraph
 {
 	GENERATED_BODY()
@@ -52,7 +52,7 @@ public:
  *
  * Note: Keep in private header.
  */
-UCLASS()
+UCLASS(Hidden, HideDropDown)
 class UMVVMFakeTestFunctiongraph : public UEdGraph
 {
 	GENERATED_BODY()
diff --git a/Engine/Plugins/Runtime/ModelViewViewModel/Source/ModelViewViewModelEditor/ModelViewViewModelEditor.Build.cs b/Engine/Plugins/Runtime/ModelViewViewModel/Source/ModelViewViewModelEditor/ModelViewViewModelEditor.Build.cs
--- a/Engine/Plugins/Runtime/ModelViewViewModel/Source/ModelViewViewModelEditor/ModelViewViewModelEditor.Build.cs
+++ b/Engine/Plugins/Runtime/ModelViewViewModel/Source/ModelViewViewModelEditor/ModelViewViewModelEditor.Build.cs
@@ -12,10 +12,12 @@ public class ModelViewViewModelEditor : ModuleRules
 			new string[]
 			{
 				"Engine",
+				"BlueprintGraph",
 				"Core",
 				"CoreUObject",
 				"EditorSubsystem",
 				"FieldNotification",
+				"KismetCompiler",
 				"ModelViewViewModel",
 				"ModelViewViewModelBlueprint",
 				"PropertyPath",
@@ -27,7 +29,6 @@ public class ModelViewViewModelEditor : ModuleRules
 				"AdvancedWidgets",
 				"AssetTools",
 				"EditorWidgets",
-				"BlueprintGraph",
 				"GraphEditor",
 				"InputCore",
 				"Kismet",
diff --git a/Engine/Plugins/Runtime/ModelViewViewModel/Source/ModelViewViewModelEditor/Private/Node/MVVMK2Node_LoadSoftResource.cpp b/Engine/Plugins/Runtime/ModelViewViewModel/Source/ModelViewViewModelEditor/Private/Node/MVVMK2Node_LoadSoftResource.cpp
new file mode 100644
--- /dev/null
+++ b/Engine/Plugins/Runtime/ModelViewViewModel/Source/ModelViewViewModelEditor/Private/Node/MVVMK2Node_LoadSoftResource.cpp
@@ -0,0 +1,162 @@
+// Copyright Epic Games, Inc. All Rights Reserved.
+
+#include "Node/MVVMK2Node_LoadSoftResource.h"
+
+#include "BlueprintCompiledStatement.h"
+#include "BlueprintNodeSpawner.h"
+#include "Blueprint/WidgetBlueprintLibrary.h"
+#include "EdGraph/EdGraphPin.h"
+#include "EdGraphSchema_K2.h"
+#include "Kismet2/CompilerResultsLog.h"
+#include "KismetCompiledFunctionContext.h"
+#include "KismetCompiler.h"
+#include "MVVMSubsystem.h"
+#include "Subsystems/SubsystemBlueprintLibrary.h"
+#include "View/MVVMView.h"
+
+#include "K2Node_CallFunction.h"
+#include "K2Node_ClassDynamicCast.h"
+#include "K2Node_IfThenElse.h"
+#include "K2Node_Self.h"
+
+#include "Engine/Texture2D.h"
+#include "Materials/MaterialInterface.h"
+
+
+#define LOCTEXT_NAMESPACE "MVVMK2Node_LoadSoftResource"
+
+//////////////////////////////////////////////////////////////////////////
+
+void UMVVMK2Node_LoadSoftResource::AllocateDefaultPins()
+{
+	const UEdGraphSchema_K2* K2Schema = GetDefault<UEdGraphSchema_K2>();
+
+	CreatePin(EGPD_Input, UEdGraphSchema_K2::PC_Exec, UEdGraphSchema_K2::PN_Execute);
+
+	// The immediate continue pin
+	CreatePin(EGPD_Output, UEdGraphSchema_K2::PC_Exec, UEdGraphSchema_K2::PN_Then);
+
+	// The delayed completed pin, this used to be called Then
+	CreatePin(EGPD_Output, UEdGraphSchema_K2::PC_Exec, UEdGraphSchema_K2::PN_Completed);
+
+	CreatePin(EGPD_Input, UEdGraphSchema_K2::PC_SoftObject, GetInputResourceClass(), GetInputPinName());
+	CreatePin(EGPD_Output, UEdGraphSchema_K2::PC_Object, GetInputResourceClass(), GetOutputPinName());
+
+	UK2Node::AllocateDefaultPins();
+}
+
+void UMVVMK2Node_LoadSoftResource::ExpandNode(FKismetCompilerContext& CompilerContext, UEdGraph* SourceGraph)
+{
+	UK2Node::ExpandNode(CompilerContext, SourceGraph);
+	
+	const UEdGraphSchema_K2* Schema = CompilerContext.GetSchema();
+
+	UK2Node_LoadAsset* CallLoadAssetNode = CompilerContext.SpawnIntermediateNode<UK2Node_LoadAsset>(this, SourceGraph);
+	{
+		CallLoadAssetNode->AllocateDefaultPins();
+	}
+	UEdGraphPin* LoadAssetInput = CallLoadAssetNode->FindPinChecked(CallLoadAssetNode->GetInputPinName());
+	UEdGraphPin* LoadAssetOutput = CallLoadAssetNode->FindPinChecked(CallLoadAssetNode->GetOutputPinName());
+
+	UK2Node_DynamicCast* CastNode = CompilerContext.SpawnIntermediateNode<UK2Node_DynamicCast>(this, SourceGraph);
+	{
+		// CastNode->SetPurity(true);
+		CastNode->TargetType = GetInputResourceClass();
+		CastNode->AllocateDefaultPins();
+	}
+	UEdGraphPin* CastInput = CastNode->GetCastSourcePin();
+	UEdGraphPin* CastOutput = CastNode->GetCastResultPin();
+
+	// move this.exec to CallLoadAssetNode.exec
+	CompilerContext.MovePinLinksToIntermediate(*GetExecPin(), *CallLoadAssetNode->GetExecPin());
+	// move this.then to CallLoadAssetNode.then
+	CompilerContext.MovePinLinksToIntermediate(*GetThenPin(), *CallLoadAssetNode->GetThenPin());
+	// CallLoadAssetNode.completed to CastNode.exec
+	check(Schema->TryCreateConnection(CallLoadAssetNode->FindPinChecked(UEdGraphSchema_K2::PN_Completed), CastNode->GetExecPin()));
+	// move this.completed to CastNode.then
+	CompilerContext.MovePinLinksToIntermediate(*GetCompletedPin(), *CastNode->GetThenPin());
+
+	// move this.resource to CallLoadAssetNode.arg
+	CompilerContext.MovePinLinksToIntermediate(*FindPinChecked(GetInputPinName()), *LoadAssetInput);
+	// CallLoadAssetNode.result to CastNode.Input
+	check(Schema->TryCreateConnection(LoadAssetOutput, CastInput));
+	// move this.result to CastNode.Output
+	CompilerContext.MovePinLinksToIntermediate(*FindPinChecked(GetOutputPinName()), *CastOutput);
+
+	BreakAllNodeLinks();
+}
+
+UEdGraphPin* UMVVMK2Node_LoadSoftResource::GetThenPin() const
+{
+	return FindPinChecked(UEdGraphSchema_K2::PN_Then);
+}
+
+UEdGraphPin* UMVVMK2Node_LoadSoftResource::GetCompletedPin() const
+{
+	return FindPinChecked(UEdGraphSchema_K2::PN_Completed);
+}
+
+const FName& UMVVMK2Node_LoadSoftResource::GetOutputPinName() const
+{
+	static const FName OutputPinName("Result");
+	return OutputPinName;
+}
+
+UClass* UMVVMK2Node_LoadSoftResource::GetInputResourceClass() const
+{
+	return nullptr;
+}
+
+
+//////////////////////////////////////////////////////////////////////////
+
+
+FText UMVVMK2Node_LoadSoftTexture::GetTooltipText() const
+{
+	return FText(LOCTEXT("UMVVMK2Node_LoadSoftTextureGetTooltipText", "Asynchronously loads a Soft Texture Reference and returns that texture on successful load"));
+}
+
+FText UMVVMK2Node_LoadSoftTexture::GetNodeTitle(ENodeTitleType::Type TitleType) const
+{
+	return FText(LOCTEXT("UMVVMK2Node_LoadSoftTextureGetNodeTitle", "Load From Soft Texture"));
+}
+
+const FName& UMVVMK2Node_LoadSoftTexture::GetInputPinName() const
+{
+	// Note: Our pin input should match the input of the specified MakeBrush node
+	static const FName InputPinName("Texture");
+	return InputPinName;
+}
+
+UClass* UMVVMK2Node_LoadSoftTexture::GetInputResourceClass() const
+{
+	return UTexture2D::StaticClass();
+}
+
+
+//////////////////////////////////////////////////////////////////////////
+
+
+FText UMVVMK2Node_LoadSoftMaterial::GetTooltipText() const
+{
+	return FText(LOCTEXT("UMVVMK2Node_LoadSoftMaterialGetTooltipText", "Asynchronously loads a Soft Material Reference and returns that material on successful load"));
+}
+
+FText UMVVMK2Node_LoadSoftMaterial::GetNodeTitle(ENodeTitleType::Type TitleType) const
+{
+	return FText(LOCTEXT("UMVVMK2Node_LoadSoftMaterialGetNodeTitle", "Load From Soft Material"));
+}
+
+const FName& UMVVMK2Node_LoadSoftMaterial::GetInputPinName() const
+{
+	// Note: Our pin input should match the input of the specified MakeBrush node
+	static const FName InputPinName("Material");
+	return InputPinName;
+}
+
+UClass* UMVVMK2Node_LoadSoftMaterial::GetInputResourceClass() const
+{
+	return UMaterialInterface::StaticClass();
+}
+
+#undef LOCTEXT_NAMESPACE
diff --git a/Engine/Plugins/Runtime/ModelViewViewModel/Source/ModelViewViewModelEditor/Private/Node/MVVMK2Node_LoadSoftResource.h b/Engine/Plugins/Runtime/ModelViewViewModel/Source/ModelViewViewModelEditor/Private/Node/MVVMK2Node_LoadSoftResource.h
new file mode 100644
--- /dev/null
+++ b/Engine/Plugins/Runtime/ModelViewViewModel/Source/ModelViewViewModelEditor/Private/Node/MVVMK2Node_LoadSoftResource.h
@@ -0,0 +1,98 @@
+// Copyright Epic Games, Inc. All Rights Reserved.
+
+#pragma once
+
+#include "CoreMinimal.h"
+#include "K2Node.h"
+#include "K2Node_LoadAsset.h"
+#include "View/MVVMViewTypes.h"
+
+#include "MVVMK2Node_LoadSoftResource.generated.h"
+
+class UEdGraphPin;
+class FKismetCompilerContext;
+
+
+/**
+ * Utility async node to load a type-specific a soft resource (Texture, Material, etc). 
+ * Needed since we current do not support automatic dynamic casting in MVVM
+ */
+UCLASS(MinimalAPI, Abstract, Hidden, HideDropDown)
+class UMVVMK2Node_LoadSoftResource : public UK2Node_LoadAsset
+{
+	GENERATED_BODY()
+
+public:
+	//~ Begin UEdGraphNode Interface.
+	virtual void AllocateDefaultPins() override;
+	//~ End UEdGraphNode Interface.
+	
+	//~ Begin K2Node Interface.
+	virtual void ExpandNode(FKismetCompilerContext& CompilerContext, UEdGraph* SourceGraph) override;
+	//~ End K2Node Interface.
+
+	/** Get the then output pin */
+	UEdGraphPin* GetThenPin() const;
+
+	/** Get the then output pin */
+	UEdGraphPin* GetCompletedPin() const;
+
+protected:
+
+	//~ Begin UK2Node_LoadAsset Interface.
+	virtual const FName& GetOutputPinName() const;
+	//~ End UK2Node_LoadAsset Interface.
+
+	/** Get the type of the input arg pin */
+	virtual UClass* GetInputResourceClass() const;
+};
+
+
+//////////////////////////////////////////////////////////////////////////
+
+
+UCLASS(MinimalAPI)
+class UMVVMK2Node_LoadSoftTexture : public UMVVMK2Node_LoadSoftResource
+{
+	GENERATED_BODY()
+
+public:
+	//~ Begin UEdGraphNode Interface.
+	virtual FText GetTooltipText() const override;
+	virtual FText GetNodeTitle(ENodeTitleType::Type TitleType) const override;
+	//~ End UEdGraphNode Interface.
+
+protected:
+	//~ Begin UK2Node_LoadAsset Interface.
+	virtual const FName& GetInputPinName() const;
+	//~ End UK2Node_LoadAsset Interface.
+
+	//~ Begin UMVVMK2Node_LoadSoftResource Interface.
+	virtual UClass* GetInputResourceClass() const override;
+	//~ End UMVVMK2Node_LoadSoftResource Interface.
+};
+
+
+//////////////////////////////////////////////////////////////////////////
+
+
+UCLASS(MinimalAPI)
+class UMVVMK2Node_LoadSoftMaterial: public UMVVMK2Node_LoadSoftResource
+{
+	GENERATED_BODY()
+
+public:
+	//~ Begin UEdGraphNode Interface.
+	virtual FText GetTooltipText() const override;
+	virtual FText GetNodeTitle(ENodeTitleType::Type TitleType) const override;
+	//~ End UEdGraphNode Interface.
+
+protected:
+	//~ Begin UK2Node_LoadAsset Interface.
+	virtual const FName& GetInputPinName() const;
+	//~ End UK2Node_LoadAsset Interface.
+
+	//~ Begin UMVVMK2Node_LoadSoftResource Interface.
+	virtual UClass* GetInputResourceClass() const override;
+	//~ End UMVVMK2Node_LoadSoftResource Interface.
+};
diff --git a/Engine/Plugins/Runtime/ModelViewViewModel/Source/ModelViewViewModelEditor/Private/Node/MVVMK2Node_MakeBrushFromSoftResource.cpp b/Engine/Plugins/Runtime/ModelViewViewModel/Source/ModelViewViewModelEditor/Private/Node/MVVMK2Node_MakeBrushFromSoftResource.cpp
new file mode 100644
--- /dev/null
+++ b/Engine/Plugins/Runtime/ModelViewViewModel/Source/ModelViewViewModelEditor/Private/Node/MVVMK2Node_MakeBrushFromSoftResource.cpp
@@ -0,0 +1,205 @@
+// Copyright Epic Games, Inc. All Rights Reserved.
+
+#include "Node/MVVMK2Node_MakeBrushFromSoftResource.h"
+
+#include "BlueprintCompiledStatement.h"
+#include "BlueprintNodeSpawner.h"
+#include "Blueprint/WidgetBlueprintLibrary.h"
+#include "EdGraph/EdGraphPin.h"
+#include "EdGraphSchema_K2.h"
+#include "Kismet2/CompilerResultsLog.h"
+#include "KismetCompiledFunctionContext.h"
+#include "KismetCompiler.h"
+#include "MVVMSubsystem.h"
+#include "Subsystems/SubsystemBlueprintLibrary.h"
+#include "View/MVVMView.h"
+
+#include "K2Node_CallFunction.h"
+#include "K2Node_ClassDynamicCast.h"
+#include "K2Node_IfThenElse.h"
+#include "K2Node_Self.h"
+
+#include "Engine/Texture2D.h"
+#include "Materials/MaterialInterface.h"
+
+
+#define LOCTEXT_NAMESPACE "MVVMK2Node_MakeBrushFromSoftResource"
+
+namespace UE::MVVM::Private
+{
+	namespace PinNames
+	{
+		static const FName Width = "Width";
+		static const FName Height = "Height";
+	}
+}
+
+//////////////////////////////////////////////////////////////////////////
+
+void UMVVMK2Node_MakeBrushFromSoftResource::AllocateDefaultPins()
+{
+	const UEdGraphSchema_K2* K2Schema = GetDefault<UEdGraphSchema_K2>();
+
+	CreatePin(EGPD_Input, UEdGraphSchema_K2::PC_Exec, UEdGraphSchema_K2::PN_Execute);
+
+	// The immediate continue pin
+	CreatePin(EGPD_Output, UEdGraphSchema_K2::PC_Exec, UEdGraphSchema_K2::PN_Then);
+
+	// The delayed completed pin, this used to be called Then
+	CreatePin(EGPD_Output, UEdGraphSchema_K2::PC_Exec, UEdGraphSchema_K2::PN_Completed);
+
+	CreatePin(EGPD_Input, UEdGraphSchema_K2::PC_SoftObject, GetInputResourceClass(), GetInputPinName());
+	CreatePin(EGPD_Input, UEdGraphSchema_K2::PC_Int, nullptr, UE::MVVM::Private::PinNames::Width);
+	CreatePin(EGPD_Input, UEdGraphSchema_K2::PC_Int, nullptr, UE::MVVM::Private::PinNames::Height);
+	CreatePin(EGPD_Output, UEdGraphSchema_K2::PC_Struct, FSlateBrush::StaticStruct(), GetOutputPinName());
+
+	UK2Node::AllocateDefaultPins();
+}
+
+void UMVVMK2Node_MakeBrushFromSoftResource::ExpandNode(FKismetCompilerContext& CompilerContext, UEdGraph* SourceGraph)
+{
+	UK2Node::ExpandNode(CompilerContext, SourceGraph);
+	
+	const UEdGraphSchema_K2* Schema = CompilerContext.GetSchema();
+
+	UK2Node_LoadAsset* CallLoadAssetNode = CompilerContext.SpawnIntermediateNode<UK2Node_LoadAsset>(this, SourceGraph);
+	{
+		CallLoadAssetNode->AllocateDefaultPins();
+	}
+	UEdGraphPin* LoadAssetInput = CallLoadAssetNode->FindPinChecked(CallLoadAssetNode->GetInputPinName());
+	UEdGraphPin* LoadAssetOutput = CallLoadAssetNode->FindPinChecked(CallLoadAssetNode->GetOutputPinName());
+
+	UK2Node_CallFunction* CallMakeBrushFrom = CompilerContext.SpawnIntermediateNode<UK2Node_CallFunction>(this, SourceGraph);
+	{
+		CallMakeBrushFrom->FunctionReference.SetExternalMember(GetMakeBrushFunctionName(), UWidgetBlueprintLibrary::StaticClass());
+		CallMakeBrushFrom->AllocateDefaultPins();
+	}
+	// Note: Our pin input should match the input of the specified MakeBrush node
+	UEdGraphPin* CallMakeBrushFromResourcePin = CallMakeBrushFrom->FindPinChecked(GetInputPinName());
+	UEdGraphPin* CallMakeBrushFromWidthPin = CallMakeBrushFrom->FindPinChecked(UE::MVVM::Private::PinNames::Width);
+	UEdGraphPin* CallMakeBrushFromHeightPin = CallMakeBrushFrom->FindPinChecked(UE::MVVM::Private::PinNames::Height);
+	UEdGraphPin* CallMakeBrushFromResultPin = CallMakeBrushFrom->GetReturnValuePin();
+
+	UK2Node_DynamicCast* CastNode = CompilerContext.SpawnIntermediateNode<UK2Node_DynamicCast>(this, SourceGraph);
+	{
+		// CastNode->SetPurity(true);
+		CastNode->TargetType = GetInputResourceClass();
+		CastNode->AllocateDefaultPins();
+	}
+	UEdGraphPin* CastInput = CastNode->GetCastSourcePin();
+	UEdGraphPin* CastOutput = CastNode->GetCastResultPin();
+
+	// move this.exec to CallLoadAssetNode.exec
+	CompilerContext.MovePinLinksToIntermediate(*GetExecPin(), *CallLoadAssetNode->GetExecPin());
+	// move this.then to CallLoadAssetNode.then
+	CompilerContext.MovePinLinksToIntermediate(*GetThenPin(), *CallLoadAssetNode->GetThenPin());
+	// CallLoadAssetNode.completed to CastNode.exec
+	check(Schema->TryCreateConnection(CallLoadAssetNode->FindPinChecked(UEdGraphSchema_K2::PN_Completed), CastNode->GetExecPin()));
+	// move this.completed to CastNode.then
+	CompilerContext.MovePinLinksToIntermediate(*GetCompletedPin(), *CastNode->GetThenPin());
+
+	// move this.resource to CallLoadAssetNode.arg
+	CompilerContext.MovePinLinksToIntermediate(*FindPinChecked(GetInputPinName()), *LoadAssetInput);
+	// move this.width to CallMakeBrushFrom.width
+	CompilerContext.MovePinLinksToIntermediate(*FindPinChecked(UE::MVVM::Private::PinNames::Width), *CallMakeBrushFromWidthPin);
+	// move this.height to CallMakeBrushFrom.height
+	CompilerContext.MovePinLinksToIntermediate(*FindPinChecked(UE::MVVM::Private::PinNames::Height), *CallMakeBrushFromHeightPin);
+	// CallLoadAssetNode.result to CastNode.Input
+	check(Schema->TryCreateConnection(LoadAssetOutput, CastInput));
+	// CastNode.Output to CallMakeBrushFrom.Resource
+	check(Schema->TryCreateConnection(CastOutput, CallMakeBrushFromResourcePin));
+	// move this.result to CallLoadAssetNode.result
+	CompilerContext.MovePinLinksToIntermediate(*FindPinChecked(GetOutputPinName()), *CallMakeBrushFromResultPin);
+
+	BreakAllNodeLinks();
+}
+
+UEdGraphPin* UMVVMK2Node_MakeBrushFromSoftResource::GetThenPin() const
+{
+	return FindPinChecked(UEdGraphSchema_K2::PN_Then);
+}
+
+UEdGraphPin* UMVVMK2Node_MakeBrushFromSoftResource::GetCompletedPin() const
+{
+	return FindPinChecked(UEdGraphSchema_K2::PN_Completed);
+}
+
+const FName& UMVVMK2Node_MakeBrushFromSoftResource::GetOutputPinName() const
+{
+	static const FName OutputPinName("SlateBrush");
+	return OutputPinName;
+}
+
+UClass* UMVVMK2Node_MakeBrushFromSoftResource::GetInputResourceClass() const
+{
+	return nullptr;
+}
+
+FName UMVVMK2Node_MakeBrushFromSoftResource::GetMakeBrushFunctionName() const
+{
+	return NAME_None;
+}
+
+
+//////////////////////////////////////////////////////////////////////////
+
+
+FText UMVVMK2Node_MakeBrushFromSoftTexture::GetTooltipText() const
+{
+	return FText(LOCTEXT("UMVVMK2Node_MakeBrushFromSoftTextureGetTooltipText", "Asynchronously loads a Soft Texture Reference and returns a slate brush using that texture on successful load"));
+}
+
+FText UMVVMK2Node_MakeBrushFromSoftTexture::GetNodeTitle(ENodeTitleType::Type TitleType) const
+{
+	return FText(LOCTEXT("UMVVMK2Node_MakeBrushFromSoftTextureGetNodeTitle", "Make Brush From Soft Texture"));
+}
+
+const FName& UMVVMK2Node_MakeBrushFromSoftTexture::GetInputPinName() const
+{
+	// Note: Our pin input should match the input of the specified MakeBrush node
+	static const FName InputPinName("Texture");
+	return InputPinName;
+}
+
+UClass* UMVVMK2Node_MakeBrushFromSoftTexture::GetInputResourceClass() const
+{
+	return UTexture2D::StaticClass();
+}
+
+FName UMVVMK2Node_MakeBrushFromSoftTexture::GetMakeBrushFunctionName() const
+{
+	return GET_FUNCTION_NAME_CHECKED(UWidgetBlueprintLibrary, MakeBrushFromTexture);
+}
+
+
+//////////////////////////////////////////////////////////////////////////
+
+
+FText UMVVMK2Node_MakeBrushFromSoftMaterial::GetTooltipText() const
+{
+	return FText(LOCTEXT("UMVVMK2Node_MakeBrushFromSoftMaterialGetTooltipText", "Asynchronously loads a Soft Material Reference and returns a slate brush using that material on successful load"));
+}
+
+FText UMVVMK2Node_MakeBrushFromSoftMaterial::GetNodeTitle(ENodeTitleType::Type TitleType) const
+{
+	return FText(LOCTEXT("UMVVMK2Node_MakeBrushFromSoftMaterialGetNodeTitle", "Make Brush From Soft Material"));
+}
+
+const FName& UMVVMK2Node_MakeBrushFromSoftMaterial::GetInputPinName() const
+{
+	// Note: Our pin input should match the input of the specified MakeBrush node
+	static const FName InputPinName("Material");
+	return InputPinName;
+}
+
+UClass* UMVVMK2Node_MakeBrushFromSoftMaterial::GetInputResourceClass() const
+{
+	return UMaterialInterface::StaticClass();
+}
+
+FName UMVVMK2Node_MakeBrushFromSoftMaterial::GetMakeBrushFunctionName() const
+{
+	return GET_FUNCTION_NAME_CHECKED(UWidgetBlueprintLibrary, MakeBrushFromMaterial);
+}
+
+#undef LOCTEXT_NAMESPACE
diff --git a/Engine/Plugins/Runtime/ModelViewViewModel/Source/ModelViewViewModelEditor/Private/Node/MVVMK2Node_MakeBrushFromSoftResource.h b/Engine/Plugins/Runtime/ModelViewViewModel/Source/ModelViewViewModelEditor/Private/Node/MVVMK2Node_MakeBrushFromSoftResource.h
new file mode 100644
--- /dev/null
+++ b/Engine/Plugins/Runtime/ModelViewViewModel/Source/ModelViewViewModelEditor/Private/Node/MVVMK2Node_MakeBrushFromSoftResource.h
@@ -0,0 +1,103 @@
+// Copyright Epic Games, Inc. All Rights Reserved.
+
+#pragma once
+
+#include "CoreMinimal.h"
+#include "K2Node.h"
+#include "K2Node_LoadAsset.h"
+#include "View/MVVMViewTypes.h"
+
+#include "MVVMK2Node_MakeBrushFromSoftResource.generated.h"
+
+class UEdGraphPin;
+class FKismetCompilerContext;
+
+
+/**
+ * Utility async node to create a slate brush from a soft resource (Texture, Material, etc). 
+ * Needed since we cannot nest conversion functions (Ex: LoadAsset->MakeBrushFromTexture).
+ */
+UCLASS(MinimalAPI, Abstract, Hidden, HideDropDown)
+class UMVVMK2Node_MakeBrushFromSoftResource : public UK2Node_LoadAsset
+{
+	GENERATED_BODY()
+
+public:
+	//~ Begin UEdGraphNode Interface.
+	virtual void AllocateDefaultPins() override;
+	//~ End UEdGraphNode Interface.
+	
+	//~ Begin K2Node Interface.
+	virtual void ExpandNode(FKismetCompilerContext& CompilerContext, UEdGraph* SourceGraph) override;
+	//~ End K2Node Interface.
+
+	/** Get the then output pin */
+	UEdGraphPin* GetThenPin() const;
+
+	/** Get the then output pin */
+	UEdGraphPin* GetCompletedPin() const;
+
+protected:
+
+	//~ Begin UK2Node_LoadAsset Interface.
+	virtual const FName& GetOutputPinName() const;
+	//~ End UK2Node_LoadAsset Interface.
+
+	/** Get the type of the input arg pin */
+	virtual UClass* GetInputResourceClass() const;
+
+	/** Get the type of the input arg pin */
+	virtual FName GetMakeBrushFunctionName() const;
+};
+
+
+//////////////////////////////////////////////////////////////////////////
+
+
+UCLASS(MinimalAPI)
+class UMVVMK2Node_MakeBrushFromSoftTexture : public UMVVMK2Node_MakeBrushFromSoftResource
+{
+	GENERATED_BODY()
+
+public:
+	//~ Begin UEdGraphNode Interface.
+	virtual FText GetTooltipText() const override;
+	virtual FText GetNodeTitle(ENodeTitleType::Type TitleType) const override;
+	//~ End UEdGraphNode Interface.
+
+protected:
+	//~ Begin UK2Node_LoadAsset Interface.
+	virtual const FName& GetInputPinName() const;
+	//~ End UK2Node_LoadAsset Interface.
+
+	//~ Begin UMVVMK2Node_MakeBrushFromSoftResource Interface.
+	virtual UClass* GetInputResourceClass() const override;
+	virtual FName GetMakeBrushFunctionName() const override;
+	//~ End UMVVMK2Node_MakeBrushFromSoftResource Interface.
+};
+
+
+//////////////////////////////////////////////////////////////////////////
+
+
+UCLASS(MinimalAPI)
+class UMVVMK2Node_MakeBrushFromSoftMaterial: public UMVVMK2Node_MakeBrushFromSoftResource
+{
+	GENERATED_BODY()
+
+public:
+	//~ Begin UEdGraphNode Interface.
+	virtual FText GetTooltipText() const override;
+	virtual FText GetNodeTitle(ENodeTitleType::Type TitleType) const override;
+	//~ End UEdGraphNode Interface.
+
+protected:
+	//~ Begin UK2Node_LoadAsset Interface.
+	virtual const FName& GetInputPinName() const;
+	//~ End UK2Node_LoadAsset Interface.
+
+	//~ Begin UMVVMK2Node_MakeBrushFromSoftResource Interface.
+	virtual UClass* GetInputResourceClass() const override;
+	virtual FName GetMakeBrushFunctionName() const override;
+	//~ End UMVVMK2Node_MakeBrushFromSoftResource Interface.
+};
diff --git a/Engine/Source/Editor/BlueprintGraph/Classes/K2Node_LoadAsset.h b/Engine/Source/Editor/BlueprintGraph/Classes/K2Node_LoadAsset.h
--- a/Engine/Source/Editor/BlueprintGraph/Classes/K2Node_LoadAsset.h
+++ b/Engine/Source/Editor/BlueprintGraph/Classes/K2Node_LoadAsset.h
@@ -19,8 +19,8 @@ class UEdGraph;
 class UEdGraphPin;
 class UObject;
 
-UCLASS(MinimalAPI)
-class UK2Node_LoadAsset : public UK2Node
+UCLASS()
+class BLUEPRINTGRAPH_API UK2Node_LoadAsset : public UK2Node
 {
 	GENERATED_BODY()
 public:
@@ -41,14 +41,14 @@ public:
 	virtual void ReallocatePinsDuringReconstruction(TArray<UEdGraphPin*>& OldPins) override;
 	// End of UK2Node interface
 
+	virtual const FName& GetInputPinName() const;
+	virtual const FName& GetOutputPinName() const;
+
 protected:
 	virtual FName NativeFunctionName() const;
 
 	virtual const FName& GetInputCategory() const;
 	virtual const FName& GetOutputCategory() const;
-
-	virtual const FName& GetInputPinName() const;
-	virtual const FName& GetOutputPinName() const;
 };
 
 UCLASS(MinimalAPI)
