Lower log level on multiple bindings. Make message more descriptive.
Add validation to MVVM compiler for async events

#tests PIE, standalone, validation

diff --git a/Engine/Plugins/Runtime/ModelViewViewModel/Source/ModelViewViewModelBlueprint/Private/MVVMViewBlueprintCompiler.cpp b/Engine/Plugins/Runtime/ModelViewViewModel/Source/ModelViewViewModelBlueprint/Private/MVVMViewBlueprintCompiler.cpp
--- a/Engine/Plugins/Runtime/ModelViewViewModel/Source/ModelViewViewModelBlueprint/Private/MVVMViewBlueprintCompiler.cpp
+++ b/Engine/Plugins/Runtime/ModelViewViewModel/Source/ModelViewViewModelBlueprint/Private/MVVMViewBlueprintCompiler.cpp
@@ -20,6 +20,7 @@
 #include "PropertyPermissionList.h"
 #include "MVVMFunctionGraphHelper.h"
 #include "MVVMWidgetBlueprintExtension_View.h"
+#include "Node/MVVMK2Node_AreSourcesValidForBinding.h"
 #include "Templates/ValueOrError.h"
 #include "Types/MVVMBindingName.h"
 #include "UObject/LinkerLoad.h"
@@ -1207,6 +1208,7 @@ void FMVVMViewBlueprintCompiler::CreateFunctions(const FWidgetBlueprintCompilerC
 	CreateWriteFieldContexts(Context);
 	CreateViewModelSetters(Context);
 	CreateIntermediateGraphFunctions(Context);
+	CategorizeAsyncFunctions(Context);
 }
 
 
@@ -1367,8 +1369,8 @@ void FMVVMViewBlueprintCompiler::CreateWriteFieldContexts(const FWidgetBlueprint
 					if (ValidBinding->Key.bIsForwardBinding)
 					{
 						AddMessageForBinding(Binding
-							, FText::Format(LOCTEXT("PropertyPathAlreadyUsed", "The property path '{0}' is already used by another binding."), PropertyPathToText(NewSkeletonClass, BlueprintView.Get(), DestinationPropertyPath))
-							, Compiler::EMessageType::Warning
+							, FText::Format(LOCTEXT("PropertyPathAlreadyUsed", "Note: the property path '{0}' is already used by another binding. Binding initialization order may not be deterministic."), PropertyPathToText(NewSkeletonClass, BlueprintView.Get(), DestinationPropertyPath))
+							, Compiler::EMessageType::Info
 							, FMVVMBlueprintPinId()
 						);
 					}
@@ -1563,6 +1565,38 @@ void FMVVMViewBlueprintCompiler::CreateIntermediateGraphFunctions(const FWidgetB
 	}
 }
 
+void FMVVMViewBlueprintCompiler::CategorizeAsyncFunctions(const FWidgetBlueprintCompilerContext::FCreateFunctionContext& Context)
+{
+	// Populate source binding keys for async functions
+	int32 Index = 0;
+	for (const TSharedRef<FCompilerBinding>& ValidBinding : ValidBindings)
+	{
+		FMVVMBlueprintViewBinding& Binding = *BlueprintView->GetBindingAt(ValidBinding->Key.ViewBindingIndex);
+
+		if (ValidBinding->Type == FCompilerBinding::EType::ComplexConversionFunction)
+		{
+			UMVVMBlueprintViewConversionFunction* ConversionFunction = Binding.Conversion.GetConversionFunction(ValidBinding->Key.bIsForwardBinding);
+			if (ensure(ConversionFunction) && ConversionFunction->IsUbergraphPage())
+			{
+				UEdGraph* WrapperGraph = ConversionFunction->GetOrCreateIntermediateWrapperGraph(WidgetBlueprintCompilerContext);
+				if (ensure(WrapperGraph))
+				{
+					TArray<UMVVMK2Node_AreSourcesValidForBinding*> AreSourcesValidNodes;
+					WrapperGraph->GetNodesOfClass<UMVVMK2Node_AreSourcesValidForBinding>(AreSourcesValidNodes);
+
+					// Like for events, this is a temp index to avoid log errors. The real index will get set after compilation
+					for (UMVVMK2Node_AreSourcesValidForBinding* AreSourcesValidNode : AreSourcesValidNodes)
+					{
+						AreSourcesValidNode->BindingKey = FMVVMViewClass_BindingKey(Index);
+					}
+				}
+			}
+		}
+
+		Index++;
+	}
+}
+
 
 bool FMVVMViewBlueprintCompiler::PreCompile(UWidgetBlueprintGeneratedClass* Class)
 {
@@ -3065,6 +3099,26 @@ void FMVVMViewBlueprintCompiler::CompileBindings(const FCompiledBindingLibraryCo
 				}
 			}
 		}
+
+		// Set the real binding key for the async validation nodes
+		if (ValidBinding->Type == FCompilerBinding::EType::ComplexConversionFunction)
+		{
+			UMVVMBlueprintViewConversionFunction* ConversionFunction = Binding.Conversion.GetConversionFunction(ValidBinding->Key.bIsForwardBinding);
+			if (ensure(ConversionFunction) && ConversionFunction->IsUbergraphPage())
+			{
+				UEdGraph* WrapperGraph = ConversionFunction->GetOrCreateIntermediateWrapperGraph(WidgetBlueprintCompilerContext);
+				if (ensure(WrapperGraph))
+				{
+					TArray<UMVVMK2Node_AreSourcesValidForBinding*> AreSourcesValidNodes;
+					WrapperGraph->GetNodesOfClass<UMVVMK2Node_AreSourcesValidForBinding>(AreSourcesValidNodes);
+
+					for (UMVVMK2Node_AreSourcesValidForBinding* AreSourcesValidNode : AreSourcesValidNodes)
+					{
+						AreSourcesValidNode->BindingKey = BindingKey;
+					}
+				}
+			}
+		}
 	}
 }
 
