Fix crash on exit MVVMConversionFunctionLibrary
#rnx

diff --git a/Engine/Plugins/Runtime/ModelViewViewModel/Source/ModelViewViewModelEditor/Private/Framework/MVVMConversionFunctionLibrary.cpp b/Engine/Plugins/Runtime/ModelViewViewModel/Source/ModelViewViewModelEditor/Private/Framework/MVVMConversionFunctionLibrary.cpp
--- a/Engine/Plugins/Runtime/ModelViewViewModel/Source/ModelViewViewModelEditor/Private/Framework/MVVMConversionFunctionLibrary.cpp
+++ b/Engine/Plugins/Runtime/ModelViewViewModel/Source/ModelViewViewModelEditor/Private/Framework/MVVMConversionFunctionLibrary.cpp
@@ -114,12 +114,11 @@ FCollection::~FCollection()
 
 	if (FAssetRegistryModule* AssetRegistryModule = FModuleManager::GetModulePtr<FAssetRegistryModule>(TEXT("AssetRegistry")))
 	{
-		if (AssetRegistryModule->IsValid())
+		if (IAssetRegistry* AssetRegistry = AssetRegistryModule->TryGet())
 		{
-			IAssetRegistry& AssetRegistry = AssetRegistryModule->Get();
-			AssetRegistry.OnAssetAdded().RemoveAll(this);
-			AssetRegistry.OnAssetRemoved().RemoveAll(this);
-			AssetRegistry.OnAssetRenamed().RemoveAll(this);
+			AssetRegistry->OnAssetAdded().RemoveAll(this);
+			AssetRegistry->OnAssetRemoved().RemoveAll(this);
+			AssetRegistry->OnAssetRenamed().RemoveAll(this);
 		}
 	}
 
