Remove uncessary commented code in async MVVM nodes

diff --git a/Engine/Plugins/Runtime/ModelViewViewModel/Source/ModelViewViewModelEditor/Private/Node/MVVMK2Node_LoadSoftResource.cpp b/Engine/Plugins/Runtime/ModelViewViewModel/Source/ModelViewViewModelEditor/Private/Node/MVVMK2Node_LoadSoftResource.cpp
--- a/Engine/Plugins/Runtime/ModelViewViewModel/Source/ModelViewViewModelEditor/Private/Node/MVVMK2Node_LoadSoftResource.cpp
+++ b/Engine/Plugins/Runtime/ModelViewViewModel/Source/ModelViewViewModelEditor/Private/Node/MVVMK2Node_LoadSoftResource.cpp
@@ -60,7 +60,6 @@ void UMVVMK2Node_LoadSoftResource::ExpandNode(FKismetCompilerContext& CompilerCo
 
 	UK2Node_DynamicCast* CastNode = CompilerContext.SpawnIntermediateNode<UK2Node_DynamicCast>(this, SourceGraph);
 	{
-		// CastNode->SetPurity(true);
 		CastNode->TargetType = GetInputResourceClass();
 		CastNode->AllocateDefaultPins();
 	}
diff --git a/Engine/Plugins/Runtime/ModelViewViewModel/Source/ModelViewViewModelEditor/Private/Node/MVVMK2Node_MakeBrushFromSoftResource.cpp b/Engine/Plugins/Runtime/ModelViewViewModel/Source/ModelViewViewModelEditor/Private/Node/MVVMK2Node_MakeBrushFromSoftResource.cpp
--- a/Engine/Plugins/Runtime/ModelViewViewModel/Source/ModelViewViewModelEditor/Private/Node/MVVMK2Node_MakeBrushFromSoftResource.cpp
+++ b/Engine/Plugins/Runtime/ModelViewViewModel/Source/ModelViewViewModelEditor/Private/Node/MVVMK2Node_MakeBrushFromSoftResource.cpp
@@ -82,7 +82,6 @@ void UMVVMK2Node_MakeBrushFromSoftResource::ExpandNode(FKismetCompilerContext& C
 
 	UK2Node_DynamicCast* CastNode = CompilerContext.SpawnIntermediateNode<UK2Node_DynamicCast>(this, SourceGraph);
 	{
-		// CastNode->SetPurity(true);
 		CastNode->TargetType = GetInputResourceClass();
 		CastNode->AllocateDefaultPins();
 	}
