#UE Add more info to an ensure

diff --git a/Engine/Plugins/Runtime/ModelViewViewModel/Source/ModelViewViewModelBlueprint/Private/Bindings/MVVMConversionFunctionHelper.cpp b/Engine/Plugins/Runtime/ModelViewViewModel/Source/ModelViewViewModelBlueprint/Private/Bindings/MVVMConversionFunctionHelper.cpp
--- a/Engine/Plugins/Runtime/ModelViewViewModel/Source/ModelViewViewModelBlueprint/Private/Bindings/MVVMConversionFunctionHelper.cpp
+++ b/Engine/Plugins/Runtime/ModelViewViewModel/Source/ModelViewViewModelBlueprint/Private/Bindings/MVVMConversionFunctionHelper.cpp
@@ -617,7 +617,7 @@ namespace NamedNodes
 
 			if (!NewPreviousDataPin)
 			{
-				ensureMsgf(false, TEXT("A node in path doesn't have a return value."));
+				ensureMsgf(false, TEXT("A node in path doesn't have a return value. Node:%s"), *GetPathNameSafe(NewNode));
 				return MakeError();
 			}
 
