Add LoadAsset & derived classes to allowlist for MVVM
Make Node Bindings show display name rather than node class name

#tests editor

diff --git a/Engine/Plugins/Runtime/ModelViewViewModel/Source/ModelViewViewModelBlueprint/Private/MVVMDeveloperProjectSettings.cpp b/Engine/Plugins/Runtime/ModelViewViewModel/Source/ModelViewViewModelBlueprint/Private/MVVMDeveloperProjectSettings.cpp
--- a/Engine/Plugins/Runtime/ModelViewViewModel/Source/ModelViewViewModelBlueprint/Private/MVVMDeveloperProjectSettings.cpp
+++ b/Engine/Plugins/Runtime/ModelViewViewModel/Source/ModelViewViewModelBlueprint/Private/MVVMDeveloperProjectSettings.cpp
@@ -39,9 +39,11 @@ UMVVMDeveloperProjectSettings::UMVVMDeveloperProjectSettings()
 	FTopLevelAssetPath BlueprintFunctionLibrary = FTopLevelAssetPath("/Script/Engine", "BlueprintFunctionLibrary");
 	FTopLevelAssetPath FormatText = FTopLevelAssetPath("/Script/BlueprintGraph", "K2Node_FormatText");
 	FTopLevelAssetPath GenericToText = FTopLevelAssetPath("/Script/BlueprintGraph", "K2Node_GenericToText");
+	FTopLevelAssetPath LoadAsset = FTopLevelAssetPath("/Script/BlueprintGraph", "K2Node_LoadAsset");
 	AllowedClassForConversionFunctions.Add(FSoftClassPath(BlueprintFunctionLibrary.ToString()));
 	AllowedClassForConversionFunctions.Add(FSoftClassPath(FormatText.ToString()));
 	AllowedClassForConversionFunctions.Add(FSoftClassPath(GenericToText.ToString()));
+	AllowedClassForConversionFunctions.Add(FSoftClassPath(LoadAsset.ToString()));
 
 	SupportedListViewBaseClassesForExtension.Add(UListView::StaticClass());
 	SupportedPanelClassesForExtension.Add(UHorizontalBox::StaticClass());
diff --git a/Engine/Plugins/Runtime/ModelViewViewModel/Source/ModelViewViewModelEditor/Private/Widgets/SMVVMCachedViewBindingConversionFunction.cpp b/Engine/Plugins/Runtime/ModelViewViewModel/Source/ModelViewViewModelEditor/Private/Widgets/SMVVMCachedViewBindingConversionFunction.cpp
--- a/Engine/Plugins/Runtime/ModelViewViewModel/Source/ModelViewViewModelEditor/Private/Widgets/SMVVMCachedViewBindingConversionFunction.cpp
+++ b/Engine/Plugins/Runtime/ModelViewViewModel/Source/ModelViewViewModelEditor/Private/Widgets/SMVVMCachedViewBindingConversionFunction.cpp
@@ -65,6 +65,10 @@ void SCachedViewBindingConversionFunction::Construct(const FArguments& InArgs, c
 				else if (ConversionFunction.IsType<TSubclassOf<UK2Node>>())
 				{
 					TSubclassOf<UK2Node> FoundFunction = ConversionFunction.Get<TSubclassOf<UK2Node>>();
+					if (FoundFunction && FoundFunction->GetDefaultObject<UK2Node>())
+					{
+						return FoundFunction->GetDefaultObject<UK2Node>()->GetNodeTitle(ENodeTitleType::ListView);
+					}
 					return FoundFunction.Get() != nullptr ? FoundFunction->GetDisplayNameText() : FText::GetEmpty();
 				}
 				return FText::GetEmpty();
