Add verbose info to invalid conversion functions.

#tests editor compile

diff --git a/Engine/Plugins/Runtime/ModelViewViewModel/Source/ModelViewViewModelBlueprint/Private/MVVMViewBlueprintCompiler.cpp b/Engine/Plugins/Runtime/ModelViewViewModel/Source/ModelViewViewModelBlueprint/Private/MVVMViewBlueprintCompiler.cpp
--- a/Engine/Plugins/Runtime/ModelViewViewModel/Source/ModelViewViewModelBlueprint/Private/MVVMViewBlueprintCompiler.cpp
+++ b/Engine/Plugins/Runtime/ModelViewViewModel/Source/ModelViewViewModelBlueprint/Private/MVVMViewBlueprintCompiler.cpp
@@ -64,6 +64,12 @@ FAutoConsoleVariable CVarLogViewCompiledResult(
 	TEXT("After the view is compiled log the compiled bindings and sources.")
 );
 
+FAutoConsoleVariable CVarLogInvalidBindingVerbose(
+	TEXT("MVVM.LogInvalidBindingVerbose"),
+	true,
+	TEXT("If a binding is invalid. Log verbose amounts of various different content data.")
+);
+
 #if UE_BUILD_DEVELOPMENT || UE_BUILD_DEBUG
 FAutoConsoleCommand CVarTestGenerateSetter(
 	TEXT("MVVM.TestGenerateSetter"),
@@ -1229,6 +1235,40 @@ void FMVVMViewBlueprintCompiler::CategorizeBindings(const FWidgetBlueprintCompil
 			if (!ConversionFunction->IsValid(WidgetBlueprintCompilerContext.WidgetBlueprint()))
 			{
 				AddMessageForBinding(Binding, LOCTEXT("InvalidConversionFunction", "The conversion function is invalid."), Compiler::EMessageType::Error, FMVVMBlueprintPinId());
+
+				if (CVarLogInvalidBindingVerbose->GetBool())
+				{
+					FText ConversionFuncInfo = FText::Format(LOCTEXT("ConversionFuncInfo", "Conversion Function Info: Is Ubergraph: {0}, Conv Type: {1}, Conv Func: {2}")
+						, FText::FromString(ConversionFunction->IsUbergraphPage() ? "Y" : "N")
+						, FText::FromString(ConversionFunction->GetConversionFunction().GetType() == EMVVMBlueprintFunctionReferenceType::Function ? "Function" : "Node")
+						, FText::FromString(ConversionFunction->GetConversionFunction().ToString())
+					);
+					AddMessageForBinding(Binding, ConversionFuncInfo, Compiler::EMessageType::Error, FMVVMBlueprintPinId());
+
+					FText DeveloperSettingsInfo = FText::Format(LOCTEXT("DeveloperSettingsInfo", "Developer Settings Info: Filter Type: {0}")
+						, FText::FromString(GetDefault<UMVVMDeveloperProjectSettings>()->ConversionFunctionFilter == EMVVMDeveloperConversionFunctionFilterType::AllowedList ? "AllowedList" : "BlueprintActionRegistry")
+					);
+					AddMessageForBinding(Binding, DeveloperSettingsInfo, Compiler::EMessageType::Error, FMVVMBlueprintPinId());
+
+					for (const FSoftClassPath& AllowedClassForConversionFunction : GetDefault<UMVVMDeveloperProjectSettings>()->AllowedClassForConversionFunctions)
+					{
+						FText AllowedClass = FText::FromString("AllowedClass: " + AllowedClassForConversionFunction.ToString());
+						AddMessageForBinding(Binding, AllowedClass, Compiler::EMessageType::Error, FMVVMBlueprintPinId());
+					}
+
+					for (const FSoftClassPath& DeniedClassForConversionFunction : GetDefault<UMVVMDeveloperProjectSettings>()->DeniedClassForConversionFunctions)
+					{
+						FText DeniedClass = FText::FromString("DeniedClass: " + DeniedClassForConversionFunction.ToString());
+						AddMessageForBinding(Binding, DeniedClass, Compiler::EMessageType::Error, FMVVMBlueprintPinId());
+					}
+
+					for (const FName DeniedModuleForConversionFunction : GetDefault<UMVVMDeveloperProjectSettings>()->DeniedModuleForConversionFunctions)
+					{
+						FText DeniedModule = FText::FromString("DeniedModule: " + DeniedModuleForConversionFunction.ToString());
+						AddMessageForBinding(Binding, DeniedModule, Compiler::EMessageType::Error, FMVVMBlueprintPinId());
+					}
+				}
+
 				bIsCreateFunctionsStepValid = false;
 				continue;
 			}
