Add even more node info on MVVM compile fail

#tests LCW nonunity editor

diff --git a/Engine/Plugins/Runtime/ModelViewViewModel/Source/ModelViewViewModelBlueprint/Private/MVVMViewBlueprintCompiler.cpp b/Engine/Plugins/Runtime/ModelViewViewModel/Source/ModelViewViewModelBlueprint/Private/MVVMViewBlueprintCompiler.cpp
--- a/Engine/Plugins/Runtime/ModelViewViewModel/Source/ModelViewViewModelBlueprint/Private/MVVMViewBlueprintCompiler.cpp
+++ b/Engine/Plugins/Runtime/ModelViewViewModel/Source/ModelViewViewModelBlueprint/Private/MVVMViewBlueprintCompiler.cpp
@@ -1267,6 +1267,20 @@ void FMVVMViewBlueprintCompiler::CategorizeBindings(const FWidgetBlueprintCompil
 						FText DeniedModule = FText::FromString("DeniedModule: " + DeniedModuleForConversionFunction.ToString());
 						AddMessageForBinding(Binding, DeniedModule, Compiler::EMessageType::Error, FMVVMBlueprintPinId());
 					}
+
+					if (TSubclassOf<UK2Node> InvalidNode = ConversionFunction->GetConversionFunction().GetNode())
+					{
+						TStringBuilder<512> FunctionClassPath;
+						InvalidNode->GetPathName(nullptr, FunctionClassPath);
+
+						FText NodeInfo = FText::Format(LOCTEXT("NodeInfo", "Conversion Node Info: FName: {0}, ClassPathName: {1}, Module: {2}, FunctionClassPath: {3}")
+							, FText::FromString(InvalidNode->GetFName().ToString())
+							, FText::FromString(InvalidNode->GetClassPathName().ToString())
+							, FText::FromString(InvalidNode->GetClassPathName().GetPackageName().ToString())
+							, FText::FromString(FunctionClassPath.ToString())
+						);
+						AddMessageForBinding(Binding, NodeInfo, Compiler::EMessageType::Error, FMVVMBlueprintPinId());
+					}
 				}
 
 				bIsCreateFunctionsStepValid = false;
