[MVVM] Fix VM default resolvers to be created in the correct package so they save properly

#tests Add a new VM that defaults to resolver, restart editor to see data being saved correctly

diff --git a/Engine/Plugins/Runtime/ModelViewViewModel/Source/ModelViewViewModelBlueprint/Private/MVVMBlueprintViewModelContext.cpp b/Engine/Plugins/Runtime/ModelViewViewModel/Source/ModelViewViewModelBlueprint/Private/MVVMBlueprintViewModelContext.cpp
--- a/Engine/Plugins/Runtime/ModelViewViewModel/Source/ModelViewViewModelBlueprint/Private/MVVMBlueprintViewModelContext.cpp
+++ b/Engine/Plugins/Runtime/ModelViewViewModel/Source/ModelViewViewModelBlueprint/Private/MVVMBlueprintViewModelContext.cpp
@@ -20,10 +20,6 @@ FMVVMBlueprintViewModelContext::FMVVMBlueprintViewModelContext(const UClass* InC
 			NotifyFieldValueClass = const_cast<UClass*>(InClass);
 			ViewModelName = InViewModelName;
 			bExposeInstanceInEditor = GetDefault<UMVVMDeveloperProjectSettings>()->bExposeViewModelInstanceInEditor;
-
-#if WITH_EDITOR
-			Resolver = CreateDefaultResolver();
-#endif
 		}
 	}
 }
@@ -35,7 +31,7 @@ FText FMVVMBlueprintViewModelContext::GetDisplayName() const
 }
 
 #if WITH_EDITOR
-TObjectPtr<UMVVMViewModelContextResolver> FMVVMBlueprintViewModelContext::CreateDefaultResolver() const
+TObjectPtr<UMVVMViewModelContextResolver> FMVVMBlueprintViewModelContext::CreateDefaultResolver(UPackage* Package) const
 {
 	TObjectPtr<UMVVMViewModelContextResolver> DefaultResolver;
 
@@ -47,7 +43,7 @@ TObjectPtr<UMVVMViewModelContextResolver> FMVVMBlueprintViewModelContext::Create
 			{
 				if (DefaultResolverClass->DoesSupportViewModelClass(GetViewModelClass()))
 				{
-					DefaultResolver = NewObject<UMVVMViewModelContextResolver>(GetTransientPackage(), DefaultClass);
+					DefaultResolver = NewObject<UMVVMViewModelContextResolver>(Package, DefaultClass);
 				}
 			}
 		}
diff --git a/Engine/Plugins/Runtime/ModelViewViewModel/Source/ModelViewViewModelBlueprint/Public/MVVMBlueprintViewModelContext.h b/Engine/Plugins/Runtime/ModelViewViewModel/Source/ModelViewViewModelBlueprint/Public/MVVMBlueprintViewModelContext.h
--- a/Engine/Plugins/Runtime/ModelViewViewModel/Source/ModelViewViewModelBlueprint/Public/MVVMBlueprintViewModelContext.h
+++ b/Engine/Plugins/Runtime/ModelViewViewModel/Source/ModelViewViewModelBlueprint/Public/MVVMBlueprintViewModelContext.h
@@ -98,7 +98,7 @@ public:
 	}
 
 #if WITH_EDITOR
-	[[nodiscard]] TObjectPtr<UMVVMViewModelContextResolver> CreateDefaultResolver() const;
+	[[nodiscard]] TObjectPtr<UMVVMViewModelContextResolver> CreateDefaultResolver(UPackage* Package) const;
 #endif
 
 private:
diff --git a/Engine/Plugins/Runtime/ModelViewViewModel/Source/ModelViewViewModelEditor/Private/Customizations/MVVMBlueprintViewModelContextCustomization.cpp b/Engine/Plugins/Runtime/ModelViewViewModel/Source/ModelViewViewModelEditor/Private/Customizations/MVVMBlueprintViewModelContextCustomization.cpp
--- a/Engine/Plugins/Runtime/ModelViewViewModel/Source/ModelViewViewModelEditor/Private/Customizations/MVVMBlueprintViewModelContextCustomization.cpp
+++ b/Engine/Plugins/Runtime/ModelViewViewModel/Source/ModelViewViewModelEditor/Private/Customizations/MVVMBlueprintViewModelContextCustomization.cpp
@@ -526,7 +526,7 @@ void FBlueprintViewModelContextDetailCustomization::HandleCreationTypeChanged()
 				UObject* ExistingResolver = nullptr;
 				if (ResolverHandle->GetValue(ExistingResolver) == FPropertyAccess::Fail || ExistingResolver == nullptr)
 				{
-					TObjectPtr<UMVVMViewModelContextResolver> NewResolver = ContextPtr->CreateDefaultResolver();
+					TObjectPtr<UMVVMViewModelContextResolver> NewResolver = ContextPtr->CreateDefaultResolver(GetTransientPackage());
 
 					// Bypass SetValue, Resolver is set to Instanced which will block it
 					FString PropertyText = NewResolver ? NewResolver->GetPathName() : TEXT("None");
diff --git a/Engine/Plugins/Runtime/ModelViewViewModel/Source/ModelViewViewModelEditor/Private/MVVMEditorSubsystem.cpp b/Engine/Plugins/Runtime/ModelViewViewModel/Source/ModelViewViewModelEditor/Private/MVVMEditorSubsystem.cpp
--- a/Engine/Plugins/Runtime/ModelViewViewModel/Source/ModelViewViewModelEditor/Private/MVVMEditorSubsystem.cpp
+++ b/Engine/Plugins/Runtime/ModelViewViewModel/Source/ModelViewViewModelEditor/Private/MVVMEditorSubsystem.cpp
@@ -184,6 +184,8 @@ FGuid UMVVMEditorSubsystem::AddViewModel(UWidgetBlueprint* WidgetBlueprint, cons
 			FMVVMBlueprintViewModelContext Context = FMVVMBlueprintViewModelContext(ViewModelClass, *ViewModelName);
 			if (Context.IsValid())
 			{
+				Context.Resolver = Context.CreateDefaultResolver(WidgetBlueprint->GetPackage());
+
 				const FScopedTransaction Transaction(LOCTEXT("AddViewModel", "Add viewmodel"));
 				View->Modify();
 				View->AddViewModel(Context);
