MVVM UX: add tooltips for the editor side of viewmodel extensions.

diff --git a/Engine/Plugins/Runtime/ModelViewViewModel/Source/ModelViewViewModelEditor/Private/Customizations/MVVMListViewBaseExtensionCustomizationExtender.cpp b/Engine/Plugins/Runtime/ModelViewViewModel/Source/ModelViewViewModelEditor/Private/Customizations/MVVMListViewBaseExtensionCustomizationExtender.cpp
--- a/Engine/Plugins/Runtime/ModelViewViewModel/Source/ModelViewViewModelEditor/Private/Customizations/MVVMListViewBaseExtensionCustomizationExtender.cpp
+++ b/Engine/Plugins/Runtime/ModelViewViewModel/Source/ModelViewViewModelEditor/Private/Customizations/MVVMListViewBaseExtensionCustomizationExtender.cpp
@@ -69,8 +69,9 @@ void FMVVMListViewBaseExtensionCustomizationExtender::CustomizeDetails(IDetailLa
 					.NameContent()
 					[
 						SNew(STextBlock)
-						.Text(LOCTEXT("VMSupport", "Viewmodel Support"))
+						.Text(LOCTEXT("VMSupport", "Viewmodel Extension"))
 						.Font(IDetailLayoutBuilder::GetDetailFont())
+						.ToolTipText(LOCTEXT("VMExtensionToolTip", "Add or remove a Viewmodel Extension. This extension adds a widget to this list view for each entry that is provided to the Set List Items function via a binding."))
 					]
 					.ValueContent()
 					.HAlign(HAlign_Fill)
@@ -113,6 +114,7 @@ void FMVVMListViewBaseExtensionCustomizationExtender::CustomizeDetails(IDetailLa
 						SNew(STextBlock)
 						.Text(LOCTEXT("EntryVM", "Entry Viewmodel"))
 						.Font(IDetailLayoutBuilder::GetDetailFont())
+						.ToolTipText(LOCTEXT("EntryVMToolTip", "Each entry created by this Viewmodel Extension will be bound to this Viewmodel on the entry widget"))
 					]
 					.ValueContent()
 					.HAlign(HAlign_Fill)
@@ -127,6 +129,7 @@ void FMVVMListViewBaseExtensionCustomizationExtender::CustomizeDetails(IDetailLa
 							[
 								SNew(STextBlock)
 								.Text(this, &FMVVMListViewBaseExtensionCustomizationExtender::OnGetSelectedViewModel)
+								.ToolTipText(this, &FMVVMListViewBaseExtensionCustomizationExtender::OnGetSelectedViewModel)
 							]
 						]
 						+ SHorizontalBox::Slot()
diff --git a/Engine/Plugins/Runtime/ModelViewViewModel/Source/ModelViewViewModelEditor/Private/Customizations/MVVMPanelWidgetExtensionCustomizationExtender.cpp b/Engine/Plugins/Runtime/ModelViewViewModel/Source/ModelViewViewModelEditor/Private/Customizations/MVVMPanelWidgetExtensionCustomizationExtender.cpp
--- a/Engine/Plugins/Runtime/ModelViewViewModel/Source/ModelViewViewModelEditor/Private/Customizations/MVVMPanelWidgetExtensionCustomizationExtender.cpp
+++ b/Engine/Plugins/Runtime/ModelViewViewModel/Source/ModelViewViewModelEditor/Private/Customizations/MVVMPanelWidgetExtensionCustomizationExtender.cpp
@@ -62,8 +62,9 @@ void FMVVMPanelWidgetExtensionCustomizationExtender::CustomizeDetails(IDetailLay
 						.NameContent()
 						[
 							SNew(STextBlock)
-							.Text(LOCTEXT("VMSupport", "Viewmodel Support"))
+							.Text(LOCTEXT("VMExtension", "Viewmodel Extension"))
 							.Font(IDetailLayoutBuilder::GetDetailFont())
+							.ToolTipText(LOCTEXT("VMExtensionToolTip", "Add or remove a Viewmodel Extension. This extension adds a widget to this Panel Widget for each entry that is provided to the Set Items function of the Viewmodel Extension via a binding."))
 						]
 						.ValueContent()
 						.HAlign(HAlign_Fill)
@@ -108,6 +109,7 @@ void FMVVMPanelWidgetExtensionCustomizationExtender::CustomizeDetails(IDetailLay
 							EntryClassHandle->SetOnPropertyValueChanged(FSimpleDelegate::CreateSP(this, &FMVVMPanelWidgetExtensionCustomizationExtender::HandleEntryClassChanged, false));
 							HandleEntryClassChanged(true);
 							IDetailPropertyRow& EntryClassRow = MVVMCategory.AddProperty(EntryClassHandle);
+							EntryClassHandle->SetToolTipText(LOCTEXT("EntryWidgetClassToolTip", "A custom widget that will be added to this Panel Widget for each entry provided to this Viewmodel Extension."));
 
 							// "Entry Viewmodel" property row
 							MVVMCategory.AddCustomRow(FText::FromString(TEXT("Viewmodel")))
@@ -117,6 +119,7 @@ void FMVVMPanelWidgetExtensionCustomizationExtender::CustomizeDetails(IDetailLay
 									SNew(STextBlock)
 									.Text(LOCTEXT("EntryVM", "Entry Viewmodel"))
 									.Font(IDetailLayoutBuilder::GetDetailFont())
+									.ToolTipText(LOCTEXT("EntryVMToolTip", "Each entry created by this Viewmodel Extension will be bound to this Viewmodel on the entry widget"))
 								]
 								.ValueContent()
 								.HAlign(HAlign_Fill)
@@ -131,6 +134,7 @@ void FMVVMPanelWidgetExtensionCustomizationExtender::CustomizeDetails(IDetailLay
 										[
 											SNew(STextBlock)
 											.Text(this, &FMVVMPanelWidgetExtensionCustomizationExtender::OnGetSelectedViewModel)
+											.ToolTipText(this, &FMVVMPanelWidgetExtensionCustomizationExtender::OnGetSelectedViewModel)
 										]
 									]
 									+ SHorizontalBox::Slot()
@@ -159,6 +163,7 @@ void FMVVMPanelWidgetExtensionCustomizationExtender::CustomizeDetails(IDetailLay
 									SNew(STextBlock)
 									.Text(LOCTEXT("SlotTemplate", "Slot Template"))
 									.Font(IDetailLayoutBuilder::GetDetailFont())
+									.ToolTipText(LOCTEXT("SlotTemplateToolTip", "Customize the slot properties used to insert widgets into this Panel Widget."))
 								]
 								.ValueContent()
 								[
@@ -183,6 +188,7 @@ void FMVVMPanelWidgetExtensionCustomizationExtender::CustomizeDetails(IDetailLay
 							TSharedPtr<IPropertyHandle> NumDesignerPreviewEntriesHandle = PanelExtensionObjectHandle->GetChildHandle("NumDesignerPreviewEntries");
 							NumDesignerPreviewEntriesHandle->SetOnPropertyValueChanged(FSimpleDelegate::CreateSP(this, &FMVVMPanelWidgetExtensionCustomizationExtender::HandleNumDesignerPreviewEntriesChanged));
 							MVVMCategory.AddProperty(NumDesignerPreviewEntriesHandle);
+							NumDesignerPreviewEntriesHandle->SetToolTipText(LOCTEXT("NumDesignerPreviewEntriesToolTip", "Set the number of dummy widgets to show in the editor preview of the Panel Widget, to check your layout settings."));
 						}
 					}
 				}
