PR #12292: UMGEditor: Fix inconsistencies in IPropertyBindingExtension by passing the PropertyHandle to all functions

#tests Verified that MVVM's property binding extension still functions as expected
#rnx

diff --git a/Engine/Plugins/Runtime/ModelViewViewModel/Source/ModelViewViewModelEditor/Private/Customizations/MVVMPropertyBindingExtension.cpp b/Engine/Plugins/Runtime/ModelViewViewModel/Source/ModelViewViewModelEditor/Private/Customizations/MVVMPropertyBindingExtension.cpp
--- a/Engine/Plugins/Runtime/ModelViewViewModel/Source/ModelViewViewModelEditor/Private/Customizations/MVVMPropertyBindingExtension.cpp
+++ b/Engine/Plugins/Runtime/ModelViewViewModel/Source/ModelViewViewModelEditor/Private/Customizations/MVVMPropertyBindingExtension.cpp
@@ -213,8 +213,13 @@ EVisibility FMVVMPropertyBindingExtension::GetCheckmarkVisibility(const UWidgetB
 	return EVisibility::Hidden;
 }
 
-TOptional<FName> FMVVMPropertyBindingExtension::GetCurrentValue(const UWidgetBlueprint* WidgetBlueprint, const UWidget* Widget, const FProperty* Property) const
+TOptional<FName> FMVVMPropertyBindingExtension::GetCurrentValue(const UWidgetBlueprint* WidgetBlueprint, const UWidget* Widget, TSharedPtr<IPropertyHandle> PropertyHandle) const
 {
+	if (!PropertyHandle.IsValid())
+	{
+		return TOptional<FName>();
+	}
+
 	const UMVVMWidgetBlueprintExtension_View* MVVMExtensionPtr = UMVVMWidgetBlueprintExtension_View::GetExtension<UMVVMWidgetBlueprintExtension_View>(WidgetBlueprint);
 	if (MVVMExtensionPtr == nullptr)
 	{
@@ -222,7 +227,7 @@ TOptional<FName> FMVVMPropertyBindingExtension::GetCurrentValue(const UWidgetBlu
 	}
 	const UMVVMBlueprintView* MVVMBlueprintView = MVVMExtensionPtr->GetBlueprintView();
 
-	const FMVVMBlueprintViewBinding* Binding = MVVMBlueprintView ? MVVMBlueprintView->FindBinding(Widget, Property) : nullptr;
+	const FMVVMBlueprintViewBinding* Binding = MVVMBlueprintView ? MVVMBlueprintView->FindBinding(Widget, PropertyHandle->GetProperty()) : nullptr;
 	if (Binding == nullptr)
 	{
 		return TOptional<FName>();
@@ -240,8 +245,13 @@ TOptional<FName> FMVVMPropertyBindingExtension::GetCurrentValue(const UWidgetBlu
 	return TOptional<FName>();
 }
 
-const FSlateBrush* FMVVMPropertyBindingExtension::GetCurrentIcon(const UWidgetBlueprint* WidgetBlueprint, const UWidget* Widget, const FProperty* Property) const
+const FSlateBrush* FMVVMPropertyBindingExtension::GetCurrentIcon(const UWidgetBlueprint* WidgetBlueprint, const UWidget* Widget, TSharedPtr<IPropertyHandle> PropertyHandle) const
 {
+	if (!PropertyHandle.IsValid())
+	{
+		return nullptr;
+	}
+
 	const UMVVMWidgetBlueprintExtension_View* MVVMExtensionPtr = UMVVMWidgetBlueprintExtension_View::GetExtension<UMVVMWidgetBlueprintExtension_View>(WidgetBlueprint);
 	if (MVVMExtensionPtr == nullptr)
 	{
@@ -249,7 +259,7 @@ const FSlateBrush* FMVVMPropertyBindingExtension::GetCurrentIcon(const UWidgetBl
 	}
 	const UMVVMBlueprintView* MVVMBlueprintView = MVVMExtensionPtr->GetBlueprintView();
 
-	const FMVVMBlueprintViewBinding* Binding = MVVMBlueprintView ? MVVMBlueprintView->FindBinding(Widget, Property) : nullptr;
+	const FMVVMBlueprintViewBinding* Binding = MVVMBlueprintView ? MVVMBlueprintView->FindBinding(Widget, PropertyHandle->GetProperty()) : nullptr;
 	if (Binding == nullptr)
 	{
 		return nullptr;
@@ -276,8 +286,13 @@ const FSlateBrush* FMVVMPropertyBindingExtension::GetCurrentIcon(const UWidgetBl
 	return nullptr;
 }
 
-TOptional<FLinearColor> FMVVMPropertyBindingExtension::GetCurrentIconColor(const UWidgetBlueprint* WidgetBlueprint, const UWidget* Widget, const FProperty* Property) const
+TOptional<FLinearColor> FMVVMPropertyBindingExtension::GetCurrentIconColor(const UWidgetBlueprint* WidgetBlueprint, const UWidget* Widget, TSharedPtr<IPropertyHandle> PropertyHandle) const
 {
+	if (!PropertyHandle.IsValid())
+	{
+		return TOptional<FLinearColor>();
+	}
+
 	const UMVVMWidgetBlueprintExtension_View* MVVMExtensionPtr = UMVVMWidgetBlueprintExtension_View::GetExtension<UMVVMWidgetBlueprintExtension_View>(WidgetBlueprint);
 	if (MVVMExtensionPtr == nullptr)
 	{
@@ -285,7 +300,7 @@ TOptional<FLinearColor> FMVVMPropertyBindingExtension::GetCurrentIconColor(const
 	}
 	const UMVVMBlueprintView* MVVMBlueprintView = MVVMExtensionPtr->GetBlueprintView();
 
-	const FMVVMBlueprintViewBinding* Binding = MVVMBlueprintView ? MVVMBlueprintView->FindBinding(Widget, Property) : nullptr;
+	const FMVVMBlueprintViewBinding* Binding = MVVMBlueprintView ? MVVMBlueprintView->FindBinding(Widget, PropertyHandle->GetProperty()) : nullptr;
 	if (Binding == nullptr)
 	{
 		return TOptional<FLinearColor>();
@@ -328,17 +343,20 @@ TOptional<FLinearColor> FMVVMPropertyBindingExtension::GetCurrentIconColor(const
 	return TOptional<FLinearColor>();
 }
 
-void FMVVMPropertyBindingExtension::ClearCurrentValue(const UWidgetBlueprint* WidgetBlueprint, const UWidget* Widget, const FProperty* Property)
+void FMVVMPropertyBindingExtension::ClearCurrentValue(const UWidgetBlueprint* WidgetBlueprint, const UWidget* Widget, TSharedPtr<IPropertyHandle> PropertyHandle)
 {
-	if (UMVVMWidgetBlueprintExtension_View* MVVMExtensionPtr = UMVVMWidgetBlueprintExtension_View::GetExtension<UMVVMWidgetBlueprintExtension_View>(WidgetBlueprint))
+	if (PropertyHandle.IsValid())
 	{
-		if (UMVVMBlueprintView* MVVMBlueprintView = MVVMExtensionPtr->GetBlueprintView())
+		if (UMVVMWidgetBlueprintExtension_View* MVVMExtensionPtr = UMVVMWidgetBlueprintExtension_View::GetExtension<UMVVMWidgetBlueprintExtension_View>(WidgetBlueprint))
 		{
-			if (FMVVMBlueprintViewBinding* Binding = MVVMBlueprintView->FindBinding(Widget, Property))
+			if (UMVVMBlueprintView* MVVMBlueprintView = MVVMExtensionPtr->GetBlueprintView())
 			{
-				FScopedTransaction Transaction(LOCTEXT("DeleteBindingsTransaction", "Delete Binding"));
-				MVVMBlueprintView->Modify();
-				MVVMBlueprintView->RemoveBinding(Binding);
+				if (FMVVMBlueprintViewBinding* Binding = MVVMBlueprintView->FindBinding(Widget, PropertyHandle->GetProperty()))
+				{
+					FScopedTransaction Transaction(LOCTEXT("DeleteBindingsTransaction", "Delete Binding"));
+					MVVMBlueprintView->Modify();
+					MVVMBlueprintView->RemoveBinding(Binding);
+				}
 			}
 		}
 	}
@@ -351,7 +369,7 @@ TSharedPtr<FExtender> FMVVMPropertyBindingExtension::CreateMenuExtender(const UW
 	return Extender;
 }
 
-bool FMVVMPropertyBindingExtension::CanExtend(const UWidgetBlueprint* WidgetBlueprint, const UWidget* Widget, const FProperty* Property) const
+bool FMVVMPropertyBindingExtension::CanExtend(const UWidgetBlueprint* WidgetBlueprint, const UWidget* Widget, TSharedPtr<const IPropertyHandle> PropertyHandle) const
 {
 	if (!GetDefault<UMVVMDeveloperProjectSettings>()->bAllowBindingFromDetailView)
 	{
@@ -430,10 +448,5 @@ IPropertyBindingExtension::EDropResult FMVVMPropertyBindingExtension::OnDrop(con
 	}
 	return EDropResult::Unhandled;
 }
-
-TSharedPtr<FExtender> FMVVMPropertyBindingExtension::CreateMenuExtender(const UWidgetBlueprint* WidgetBlueprint, const UWidget* Widget, const FProperty* Property)
-{
-	return MakeShared<FExtender>();
-}
 }
 #undef LOCTEXT_NAMESPACE
diff --git a/Engine/Plugins/Runtime/ModelViewViewModel/Source/ModelViewViewModelEditor/Private/Customizations/MVVMPropertyBindingExtension.h b/Engine/Plugins/Runtime/ModelViewViewModel/Source/ModelViewViewModelEditor/Private/Customizations/MVVMPropertyBindingExtension.h
--- a/Engine/Plugins/Runtime/ModelViewViewModel/Source/ModelViewViewModelEditor/Private/Customizations/MVVMPropertyBindingExtension.h
+++ b/Engine/Plugins/Runtime/ModelViewViewModel/Source/ModelViewViewModelEditor/Private/Customizations/MVVMPropertyBindingExtension.h
@@ -16,17 +16,16 @@ class FMVVMPropertyBindingExtension
 {
 	static void ExtendBindingsMenu(FMenuBuilder& MenuBuilder, TSharedRef<FMVVMPropertyBindingExtension> MVVMPropertyBindingExtension, const UWidgetBlueprint* WidgetBlueprint, UWidget* Widget, TSharedPtr<IPropertyHandle> WidgetPropertyHandle);
 
-	virtual TOptional<FName> GetCurrentValue(const UWidgetBlueprint* WidgetBlueprint, const UWidget* Widget, const FProperty* Property) const override;
-	virtual const FSlateBrush* GetCurrentIcon(const UWidgetBlueprint* WidgetBlueprint, const UWidget* Widget, const FProperty* Property) const override;
-	virtual TOptional<FLinearColor> GetCurrentIconColor(const UWidgetBlueprint* WidgetBlueprint, const UWidget* Widget, const FProperty* Property) const override;
+	virtual TOptional<FName> GetCurrentValue(const UWidgetBlueprint* WidgetBlueprint, const UWidget* Widget, TSharedPtr<IPropertyHandle> PropertyHandle) const override;
+	virtual const FSlateBrush* GetCurrentIcon(const UWidgetBlueprint* WidgetBlueprint, const UWidget* Widget, TSharedPtr<IPropertyHandle> PropertyHandle) const override;
+	virtual TOptional<FLinearColor> GetCurrentIconColor(const UWidgetBlueprint* WidgetBlueprint, const UWidget* Widget, TSharedPtr<IPropertyHandle> PropertyHandle) const override;
 
-	virtual void ClearCurrentValue(const UWidgetBlueprint* WidgetBlueprint, const UWidget* Widget, const FProperty* Property) override;
-	virtual TSharedPtr<FExtender> CreateMenuExtender(const UWidgetBlueprint* WidgetBlueprint, UWidget* Widget, TSharedPtr<IPropertyHandle> WidgetPropertyHandle) override;
-	virtual bool CanExtend(const UWidgetBlueprint* WidgetBlueprint, const UWidget* Widget, const FProperty* Property) const;
-	virtual EDropResult OnDrop(const FGeometry& Geometry, const FDragDropEvent& DragDropEvent, UWidgetBlueprint* WidgetBlueprint, UWidget* Widget, TSharedPtr<IPropertyHandle> WidgetPropertyHandle) override;
-	virtual TSharedPtr<FExtender> CreateMenuExtender(const UWidgetBlueprint* WidgetBlueprint, const UWidget* Widget, const FProperty* Property) override;
+	virtual void ClearCurrentValue(const UWidgetBlueprint* WidgetBlueprint, const UWidget* Widget, TSharedPtr<IPropertyHandle> PropertyHandle) override;
+	virtual TSharedPtr<FExtender> CreateMenuExtender(const UWidgetBlueprint* WidgetBlueprint, UWidget* Widget, TSharedPtr<IPropertyHandle> PropertyHandle) override;
+	virtual bool CanExtend(const UWidgetBlueprint* WidgetBlueprint, const UWidget* Widget, TSharedPtr<const IPropertyHandle> PropertyHandle) const override;
+	virtual EDropResult OnDrop(const FGeometry& Geometry, const FDragDropEvent& DragDropEvent, UWidgetBlueprint* WidgetBlueprint, UWidget* Widget, TSharedPtr<IPropertyHandle> PropertyHandle) override;
 
-	EVisibility GetCheckmarkVisibility(const UWidgetBlueprint* WidgetBlueprint, UWidget* Widget, const FProperty* Property, FGuid OwningViewModelId, TSharedPtr<IPropertyHandle> WidgetPropertyHandle) const;
+	EVisibility GetCheckmarkVisibility(const UWidgetBlueprint* WidgetBlueprint, UWidget* Widget, const FProperty* Property, FGuid OwningViewModelId, TSharedPtr<IPropertyHandle> PropertyHandle) const;
 };
 
 }
diff --git a/Engine/Source/Editor/UMGEditor/Private/Customizations/UMGDetailCustomizations.cpp b/Engine/Source/Editor/UMGEditor/Private/Customizations/UMGDetailCustomizations.cpp
--- a/Engine/Source/Editor/UMGEditor/Private/Customizations/UMGDetailCustomizations.cpp
+++ b/Engine/Source/Editor/UMGEditor/Private/Customizations/UMGDetailCustomizations.cpp
@@ -139,7 +139,7 @@ TSharedRef<SWidget> FBlueprintWidgetCustomization::MakePropertyBindingWidget(TWe
 	IUMGEditorModule& EditorModule = FModuleManager::LoadModuleChecked<IUMGEditorModule>("UMGEditor");
 	for (const TSharedPtr<IPropertyBindingExtension>& Extension : EditorModule.GetPropertyBindingExtensibilityManager()->GetExtensions())
 	{
-		if (Extension->CanExtend(WidgetBlueprint, Widget, InPropertyHandle->GetProperty()))
+		if (Extension->CanExtend(WidgetBlueprint, Widget, InPropertyHandle))
 		{
 			MenuExtenders.Add(Extension->CreateMenuExtender(WidgetBlueprint, Widget, InPropertyHandle));
 			ActiveExtensions.Add(Extension);
@@ -339,7 +339,7 @@ TSharedRef<SWidget> FBlueprintWidgetCustomization::MakePropertyBindingWidget(TWe
 				{
 					for (const TSharedPtr<IPropertyBindingExtension>& Extension : ActiveExtensions)
 					{
-						TOptional<FLinearColor> Color = Extension->GetCurrentIconColor(ThisBlueprint, Widget, InPropertyHandle->GetProperty());
+						TOptional<FLinearColor> Color = Extension->GetCurrentIconColor(ThisBlueprint, Widget, InPropertyHandle);
 						if (Color.IsSet())
 						{
 							return Color.GetValue();
@@ -440,7 +440,7 @@ TSharedRef<SWidget> FBlueprintWidgetCustomization::MakePropertyBindingWidget(TWe
 			{
 				for (const TSharedPtr<IPropertyBindingExtension>& Extension : ActiveExtensions)
 				{
-					Extension->ClearCurrentValue(ThisBlueprint, Widget, InPropertyHandle->GetProperty());
+					Extension->ClearCurrentValue(ThisBlueprint, Widget, InPropertyHandle);
 				}
 			}
 		}
@@ -474,7 +474,7 @@ TSharedRef<SWidget> FBlueprintWidgetCustomization::MakePropertyBindingWidget(TWe
 			{
 				for (const TSharedPtr<IPropertyBindingExtension>& Extension : ActiveExtensions)
 				{
-					TOptional<FName> Name = Extension->GetCurrentValue(ThisBlueprint, Widget, InPropertyHandle->GetProperty());
+					TOptional<FName> Name = Extension->GetCurrentValue(ThisBlueprint, Widget, InPropertyHandle);
 					if (Name.IsSet())
 					{
 						return true;
@@ -573,7 +573,7 @@ TSharedRef<SWidget> FBlueprintWidgetCustomization::MakePropertyBindingWidget(TWe
 			{
 				for (const TSharedPtr<IPropertyBindingExtension>& Extension : ActiveExtensions)
 				{
-					TOptional<FName> Name = Extension->GetCurrentValue(ThisBlueprint, Widget, InPropertyHandle->GetProperty());
+					TOptional<FName> Name = Extension->GetCurrentValue(ThisBlueprint, Widget, InPropertyHandle);
 					if (Name.IsSet())
 					{
 						return FText::FromName(Name.GetValue());
@@ -629,7 +629,7 @@ TSharedRef<SWidget> FBlueprintWidgetCustomization::MakePropertyBindingWidget(TWe
 			{
 				for (const TSharedPtr<IPropertyBindingExtension>& Extension : ActiveExtensions)
 				{
-					const FSlateBrush* Brush = Extension->GetCurrentIcon(ThisBlueprint, Widget, InPropertyHandle->GetProperty());
+					const FSlateBrush* Brush = Extension->GetCurrentIcon(ThisBlueprint, Widget, InPropertyHandle);
 					if (Brush != nullptr)
 					{
 						return Brush;
@@ -698,9 +698,9 @@ bool FBlueprintWidgetCustomization::HasPropertyBindings(TWeakPtr<FWidgetBlueprin
 
 		for (const TSharedPtr<IPropertyBindingExtension>& BindingExtension : EditorModule.GetPropertyBindingExtensibilityManager()->GetExtensions())
 		{
-			if (BindingExtension->CanExtend(ThisBlueprint, Widget, InPropertyHandle->GetProperty()))
+			if (BindingExtension->CanExtend(ThisBlueprint, Widget, InPropertyHandle))
 			{
-				TOptional<FName> CurrentValue = BindingExtension->GetCurrentValue(ThisBlueprint, Widget, InPropertyHandle->GetProperty());
+				TOptional<FName> CurrentValue = BindingExtension->GetCurrentValue(ThisBlueprint, Widget, InPropertyHandle);
 				if (CurrentValue.IsSet())
 				{
 					return true;
diff --git a/Engine/Source/Editor/UMGEditor/Private/Details/DetailWidgetExtensionHandler.cpp b/Engine/Source/Editor/UMGEditor/Private/Details/DetailWidgetExtensionHandler.cpp
--- a/Engine/Source/Editor/UMGEditor/Private/Details/DetailWidgetExtensionHandler.cpp
+++ b/Engine/Source/Editor/UMGEditor/Private/Details/DetailWidgetExtensionHandler.cpp
@@ -112,7 +112,7 @@ bool FDetailWidgetExtensionHandler::IsPropertyExtendable(const UClass* ObjectCla
 	IUMGEditorModule& EditorModule = FModuleManager::LoadModuleChecked<IUMGEditorModule>("UMGEditor");
 	for (const TSharedPtr<IPropertyBindingExtension>& Extension : EditorModule.GetPropertyBindingExtensibilityManager()->GetExtensions())
 	{
-		if (Extension->CanExtend(WidgetBlueprint, Widget, PropertyHandle.GetProperty()))
+		if (Extension->CanExtend(WidgetBlueprint, Widget, PropertyHandle.AsShared()))
 		{
 			return true;
 		}
@@ -162,7 +162,7 @@ void FDetailWidgetExtensionHandler::ExtendWidgetRow(
 		IUMGEditorModule& EditorModule = FModuleManager::LoadModuleChecked<IUMGEditorModule>("UMGEditor");
 		for (const TSharedPtr<IPropertyBindingExtension>& Extension : EditorModule.GetPropertyBindingExtensibilityManager()->GetExtensions())
 		{
-			bShouldShowWidget |= Extension->CanExtend(WidgetBlueprint, Widget, InPropertyHandle->GetProperty());
+			bShouldShowWidget |= Extension->CanExtend(WidgetBlueprint, Widget, InPropertyHandle);
 		}
 	}
 
diff --git a/Engine/Source/Editor/UMGEditor/Public/IHasPropertyBindingExtensibility.h b/Engine/Source/Editor/UMGEditor/Public/IHasPropertyBindingExtensibility.h
--- a/Engine/Source/Editor/UMGEditor/Public/IHasPropertyBindingExtensibility.h
+++ b/Engine/Source/Editor/UMGEditor/Public/IHasPropertyBindingExtensibility.h
@@ -24,16 +24,26 @@ public:
 	};
 
 	/** Does this extension want to extend this property in the widget? */
-	virtual bool CanExtend(const UWidgetBlueprint* WidgetBlueprint, const UWidget* Widget, const FProperty* Property) const = 0;
-	virtual TSharedPtr<FExtender> CreateMenuExtender(const UWidgetBlueprint* WidgetBlueprint, UWidget* Widget, TSharedPtr<IPropertyHandle> WidgetPropertyHandle) = 0;
-	virtual void ClearCurrentValue(const UWidgetBlueprint* WidgetBlueprint, const UWidget* Widget, const FProperty* Property) = 0;
-	virtual TOptional<FName> GetCurrentValue(const UWidgetBlueprint* WidgetBlueprint, const UWidget* Widget, const FProperty* Property) const = 0;
-	virtual const FSlateBrush* GetCurrentIcon(const UWidgetBlueprint* WidgetBlueprint, const UWidget* Widget, const FProperty* Property) const = 0;
-	virtual TOptional<FLinearColor> GetCurrentIconColor(const UWidgetBlueprint* WidgetBlueprint, const UWidget* Widget, const FProperty* Property) const = 0;
-	virtual EDropResult OnDrop(const FGeometry& Geometry, const FDragDropEvent& DragDropEvent, UWidgetBlueprint* WidgetBlueprint, UWidget* Widget, TSharedPtr<IPropertyHandle> WidgetPropertyHandle) = 0;
+	virtual bool CanExtend(const UWidgetBlueprint* WidgetBlueprint, const UWidget* Widget, TSharedPtr<const IPropertyHandle> PropertyHandle) const = 0;
+	virtual TSharedPtr<FExtender> CreateMenuExtender(const UWidgetBlueprint* WidgetBlueprint, UWidget* Widget, TSharedPtr<IPropertyHandle> PropertyHandle) = 0;
+	virtual void ClearCurrentValue(const UWidgetBlueprint* WidgetBlueprint, const UWidget* Widget, TSharedPtr<IPropertyHandle> PropertyHandle) = 0;
+	virtual TOptional<FName> GetCurrentValue(const UWidgetBlueprint* WidgetBlueprint, const UWidget* Widget, TSharedPtr<IPropertyHandle> PropertyHandle) const = 0;
+	virtual const FSlateBrush* GetCurrentIcon(const UWidgetBlueprint* WidgetBlueprint, const UWidget* Widget, TSharedPtr<IPropertyHandle> PropertyHandle) const = 0;
+	virtual TOptional<FLinearColor> GetCurrentIconColor(const UWidgetBlueprint* WidgetBlueprint, const UWidget* Widget, TSharedPtr<IPropertyHandle> PropertyHandle) const = 0;
+	virtual EDropResult OnDrop(const FGeometry& Geometry, const FDragDropEvent& DragDropEvent, UWidgetBlueprint* WidgetBlueprint, UWidget* Widget, TSharedPtr<IPropertyHandle> PropertyHandle) = 0;
 
 	UE_DEPRECATED(5.3, "The function CreateMenuExtender with FProperty parameter is deprecated. Please pass the property handler as the parameter instead.")
-	virtual TSharedPtr<FExtender> CreateMenuExtender(const UWidgetBlueprint* WidgetBlueprint, const UWidget* Widget, const FProperty* Property) = 0;
+	virtual TSharedPtr<FExtender> CreateMenuExtender(const UWidgetBlueprint* WidgetBlueprint, const UWidget* Widget, const FProperty* Property) { return nullptr; }
+	UE_DEPRECATED(5.6, "The function CanExtend with FProperty parameter is deprecated. Please pass the property handler as the parameter instead.")
+	virtual bool CanExtend(const UWidgetBlueprint* WidgetBlueprint, const UWidget* Widget, const FProperty* Property) const { return false; }
+	UE_DEPRECATED(5.6, "The function ClearCurrentValue with FProperty parameter is deprecated. Please pass the property handler as the parameter instead.")
+	virtual void ClearCurrentValue(const UWidgetBlueprint* WidgetBlueprint, const UWidget* Widget, const FProperty* Property) {}
+	UE_DEPRECATED(5.6, "The function GetCurrentValue with FProperty parameter is deprecated. Please pass the property handler as the parameter instead.")
+	virtual TOptional<FName> GetCurrentValue(const UWidgetBlueprint* WidgetBlueprint, const UWidget* Widget, const FProperty* Property) const { return TOptional<FName>(); }
+	UE_DEPRECATED(5.6, "The function GetCurrentIcon with FProperty parameter is deprecated. Please pass the property handler as the parameter instead.")
+	virtual const FSlateBrush* GetCurrentIcon(const UWidgetBlueprint* WidgetBlueprint, const UWidget* Widget, const FProperty* Property) const { return nullptr; }
+	UE_DEPRECATED(5.6, "The function GetCurrentIconColor with FProperty parameter is deprecated. Please pass the property handler as the parameter instead.")
+	virtual TOptional<FLinearColor> GetCurrentIconColor(const UWidgetBlueprint* WidgetBlueprint, const UWidget* Widget, const FProperty* Property) const { return TOptional<FLinearColor>(); };
 };
 
 /**
