MVVM: Improve handling of renamed fields

- Fixup bindings in dependent blueprints when renaming view models
  - Load dependent blueprints when renaming view models so they get fixed too
- Fixup conversion functions when renaming view models
- Fixup components of binding property paths when renaming fields
- Fix case where renaming a field in one blueprint could falsly rename properties in bindings in dependent blueprints

Fixes asserts that can occur when view models are renamed.

[RN] fix creator/ui

diff --git a/Engine/Plugins/Runtime/ModelViewViewModel/Source/ModelViewViewModelBlueprint/Private/MVVMBlueprintView.cpp b/Engine/Plugins/Runtime/ModelViewViewModel/Source/ModelViewViewModelBlueprint/Private/MVVMBlueprintView.cpp
--- a/Engine/Plugins/Runtime/ModelViewViewModel/Source/ModelViewViewModelBlueprint/Private/MVVMBlueprintView.cpp
+++ b/Engine/Plugins/Runtime/ModelViewViewModel/Source/ModelViewViewModelBlueprint/Private/MVVMBlueprintView.cpp
@@ -643,29 +643,26 @@ void UMVVMBlueprintView::AddAssetTags(FAssetRegistryTagsContext Context) const
 	}
 }
 
-void UMVVMBlueprintView::WidgetRenamed(FName OldObjectName, FName NewObjectName)
+void UMVVMBlueprintView::OnFieldRenamed(UClass* FieldOwnerClass, FName OldObjectName, FName NewObjectName)
 {
 	bool bRenamed = false;
 
 	UWidgetBlueprint* WidgetBlueprint = GetOuterUMVVMWidgetBlueprintExtension_View()->GetWidgetBlueprint();
 
-	auto RenameWidget = [WidgetBlueprint , &bRenamed, OldObjectName, NewObjectName](FMVVMBlueprintPropertyPath& PropertyPath) -> bool
+	auto UpdatePaths = [WidgetBlueprint, FieldOwnerClass, &bRenamed, OldObjectName, NewObjectName](FMVVMBlueprintPropertyPath& PropertyPath) -> bool
 	{
-		if (PropertyPath.GetSource(WidgetBlueprint) == EMVVMBlueprintFieldPathSource::Widget && PropertyPath.GetWidgetName() == OldObjectName)
-		{
-			PropertyPath.SetWidgetName(NewObjectName);
-			bRenamed = true;
-			return true;
-		}
-		return false;
+		const bool bDidPathContainField = PropertyPath.OnFieldRenamed(WidgetBlueprint, FieldOwnerClass, OldObjectName, NewObjectName);
+		bRenamed |= bDidPathContainField;
+		return bDidPathContainField;
 	};
 
-	UE::MVVM::Private::ForEachPropertyPath_Update(this, RenameWidget, true);
+	UE::MVVM::Private::ForEachPropertyPath_Update(this, UpdatePaths, true);
 
 	if (bRenamed)
 	{
 		Modify();
 		OnBindingsUpdated.Broadcast();
+		FBlueprintEditorUtils::MarkBlueprintAsStructurallyModified(WidgetBlueprint);
 	}
 }
 #endif
diff --git a/Engine/Plugins/Runtime/ModelViewViewModel/Source/ModelViewViewModelBlueprint/Private/MVVMPropertyPath.cpp b/Engine/Plugins/Runtime/ModelViewViewModel/Source/ModelViewViewModelBlueprint/Private/MVVMPropertyPath.cpp
--- a/Engine/Plugins/Runtime/ModelViewViewModel/Source/ModelViewViewModelBlueprint/Private/MVVMPropertyPath.cpp
+++ b/Engine/Plugins/Runtime/ModelViewViewModel/Source/ModelViewViewModelBlueprint/Private/MVVMPropertyPath.cpp
@@ -344,6 +344,28 @@ bool FMVVMBlueprintPropertyPath::HasFieldInLocalScope() const
 
 namespace UE::MVVM::Private
 {
+UMVVMBlueprintView* GetView(const UBlueprint* Blueprint)
+{
+	const TObjectPtr<UBlueprintExtension>* ExtensionViewPtr = Blueprint->GetExtensions().FindByPredicate([](const UBlueprintExtension* Other)
+	{
+		return Other && Other->GetClass() == UMVVMWidgetBlueprintExtension_View::StaticClass();
+	});
+	if (ExtensionViewPtr && IsValid(*ExtensionViewPtr))
+	{
+		return CastChecked<UMVVMWidgetBlueprintExtension_View>(*ExtensionViewPtr)->GetBlueprintView();
+	}
+	return nullptr;
+}
+
+
+FName GetViewModelName(const UBlueprint* Blueprint, FGuid Id)
+{
+	UMVVMBlueprintView* BlueprintView = GetView(Blueprint);
+	const FMVVMBlueprintViewModelContext* ViewModel = BlueprintView ? BlueprintView->FindViewModel(Id) : nullptr;
+	return ViewModel ? ViewModel->GetViewModelName() : NAME_None;
+}
+
+
 FText GetWidgetDisplayName(const UWidgetBlueprint* WidgetBlueprint, FName WidgetName, bool bUseDisplayName, bool bIncludeMetaData)
 {
 	if (!bUseDisplayName && !bIncludeMetaData)
@@ -393,6 +415,41 @@ FText GetRootName(const UWidgetBlueprint* WidgetBlueprint, const FMVVMBlueprintP
 }//namespace
 
 
+bool FMVVMBlueprintPropertyPath::OnFieldRenamed(const UBlueprint* InContext, UClass* FieldOwnerClass, FName OldObjectName, FName NewObjectName)
+{
+	if (InContext != nullptr)
+	{
+		const EMVVMBlueprintFieldPathSource PathSource = GetSource(InContext);
+		const UClass* SelfContext = InContext->GeneratedClass;
+		if (SelfContext == FieldOwnerClass)
+		{
+			if (PathSource == EMVVMBlueprintFieldPathSource::Widget && GetWidgetName() == OldObjectName)
+			{
+				SetWidgetName(NewObjectName);
+				return true;
+			}
+			// Checking against NewObjectName since the viewmodel would have already updated its name
+			else if (PathSource == EMVVMBlueprintFieldPathSource::ViewModel && UE::MVVM::Private::GetViewModelName(InContext, GetViewModelId()) == NewObjectName)
+			{
+				// Nothing to rename since we reference the VM by ID, but the path does contain the renamed field
+				return true;
+			}
+		}
+
+		// Check the path components for the renamed field
+		for (FMVVMBlueprintFieldPath& Path : Paths)
+		{
+			if (Path.GetParentClass(SelfContext) == FieldOwnerClass && Path.GetFieldName(SelfContext) == OldObjectName)
+			{
+				Path.SetMemberName(NewObjectName);
+				return true;
+			}
+		}
+	}
+
+	return false;
+}
+
 FText FMVVMBlueprintPropertyPath::ToText(const UWidgetBlueprint* WidgetBlueprint, bool bUseDisplayName) const
 {
 	if (WidgetBlueprint == nullptr)
diff --git a/Engine/Plugins/Runtime/ModelViewViewModel/Source/ModelViewViewModelBlueprint/Private/MVVMWidgetBlueprintExtension_View.cpp b/Engine/Plugins/Runtime/ModelViewViewModel/Source/ModelViewViewModelBlueprint/Private/MVVMWidgetBlueprintExtension_View.cpp
--- a/Engine/Plugins/Runtime/ModelViewViewModel/Source/ModelViewViewModelBlueprint/Private/MVVMWidgetBlueprintExtension_View.cpp
+++ b/Engine/Plugins/Runtime/ModelViewViewModel/Source/ModelViewViewModelBlueprint/Private/MVVMWidgetBlueprintExtension_View.cpp
@@ -190,17 +190,26 @@ void UMVVMWidgetBlueprintExtension_View::VerifyWidgetExtensions()
 	}
 }
 
-void UMVVMWidgetBlueprintExtension_View::RenameWidgetExtensions(FName OldName, FName NewName)
+void UMVVMWidgetBlueprintExtension_View::OnFieldRenamed(UClass* FieldOwnerClass, FName OldName, FName NewName)
+{
+	const UWidgetBlueprint* WidgetBlueprint = GetWidgetBlueprint();
+	if (WidgetBlueprint != nullptr && WidgetBlueprint->GeneratedClass == FieldOwnerClass)
+	{
+		RenameWidgetExtensions(OldName, NewName);
+	}
+}
+
+void UMVVMWidgetBlueprintExtension_View::RenameWidgetExtensions(FName OldWidgetName, FName NewWidgetName)
 {
 	for (FMVVMExtensionItem& Extension : BlueprintExtensions)
 	{
-		if (Extension.WidgetName == OldName)
+		if (Extension.WidgetName == OldWidgetName)
 		{
 			Modify();
-			Extension.WidgetName = NewName;
+			Extension.WidgetName = NewWidgetName;
 			if (Extension.ExtensionObj)
 			{
-				Extension.ExtensionObj->WidgetRenamed(OldName, NewName);
+				Extension.ExtensionObj->WidgetRenamed(OldWidgetName, NewWidgetName);
 			}
 		}
 	}
diff --git a/Engine/Plugins/Runtime/ModelViewViewModel/Source/ModelViewViewModelBlueprint/Public/MVVMBlueprintView.h b/Engine/Plugins/Runtime/ModelViewViewModel/Source/ModelViewViewModelBlueprint/Public/MVVMBlueprintView.h
--- a/Engine/Plugins/Runtime/ModelViewViewModel/Source/ModelViewViewModelBlueprint/Public/MVVMBlueprintView.h
+++ b/Engine/Plugins/Runtime/ModelViewViewModel/Source/ModelViewViewModelBlueprint/Public/MVVMBlueprintView.h
@@ -171,7 +171,7 @@ public:
 	void AddAssetTags(FAssetRegistryTagsContext Context) const;
 	UE_DEPRECATED(5.4, "Implement the version that takes FAssetRegistryTagsContext instead.")
 	void AddAssetTags(TArray<FAssetRegistryTag>& OutTags) const;
-	void WidgetRenamed(FName OldObjectName, FName NewObjectName);
+	void OnFieldRenamed(UClass* FieldOwnerClass, FName OldObjectName, FName NewObjectName);
 #endif
 
 	virtual void Serialize(FArchive& Ar) override;
diff --git a/Engine/Plugins/Runtime/ModelViewViewModel/Source/ModelViewViewModelBlueprint/Public/MVVMPropertyPath.h b/Engine/Plugins/Runtime/ModelViewViewModel/Source/ModelViewViewModelBlueprint/Public/MVVMPropertyPath.h
--- a/Engine/Plugins/Runtime/ModelViewViewModel/Source/ModelViewViewModelBlueprint/Public/MVVMPropertyPath.h
+++ b/Engine/Plugins/Runtime/ModelViewViewModel/Source/ModelViewViewModelBlueprint/Public/MVVMPropertyPath.h
@@ -67,6 +67,11 @@ public:
 	void SetDeprecatedBindingReference(const FMemberReference& InBindingReference, EBindingKind InBindingKind);
 	/** */
 	void SetDeprecatedSelfReference(const UBlueprint* Context);
+	
+	void SetMemberName(FName NewName)
+	{
+		BindingReference.SetMemberName(NewName);
+	}
 #endif
 
 public:
@@ -251,6 +256,12 @@ public:
 		Source = EMVVMBlueprintFieldPathSource::SelfContext;
 	}
 
+	/**
+	 * Handle when a field that is part of the path is renamed
+	 * Returns true if the renamed field was part of the path
+	 */
+	bool OnFieldRenamed(const UBlueprint* InContext, UClass* FieldClass, FName OldObjectName, FName NewObjectName);
+
 	bool IsValid() const
 	{
 		bool bHasValidSource = bDeprecatedSource && Source != EMVVMBlueprintFieldPathSource::None;
diff --git a/Engine/Plugins/Runtime/ModelViewViewModel/Source/ModelViewViewModelBlueprint/Public/MVVMWidgetBlueprintExtension_View.h b/Engine/Plugins/Runtime/ModelViewViewModel/Source/ModelViewViewModelBlueprint/Public/MVVMWidgetBlueprintExtension_View.h
--- a/Engine/Plugins/Runtime/ModelViewViewModel/Source/ModelViewViewModelBlueprint/Public/MVVMWidgetBlueprintExtension_View.h
+++ b/Engine/Plugins/Runtime/ModelViewViewModel/Source/ModelViewViewModelBlueprint/Public/MVVMWidgetBlueprintExtension_View.h
@@ -96,7 +96,8 @@ public:
 	//~ End UWidgetBlueprintExtension interface
 
 	void VerifyWidgetExtensions();
-	void RenameWidgetExtensions(FName OldName, FName NewName);
+	void OnFieldRenamed(UClass* FieldOwnerClass, FName OldName, FName NewName);
+	void RenameWidgetExtensions(FName OldWidgetName, FName NewWidgetName);
 	UMVVMBlueprintViewExtension* CreateBlueprintWidgetExtension(TSubclassOf<UMVVMBlueprintViewExtension> ExtensionClass, FName WidgetName);
 	void RemoveBlueprintWidgetExtension(UMVVMBlueprintViewExtension* ExtensionToRemove, FName WidgetName);
 	TArray<UMVVMBlueprintViewExtension*> GetBlueprintExtensionsForWidget(FName WidgetName) const;
diff --git a/Engine/Plugins/Runtime/ModelViewViewModel/Source/ModelViewViewModelEditor/Private/MVVMEditorSubsystem.cpp b/Engine/Plugins/Runtime/ModelViewViewModel/Source/ModelViewViewModelEditor/Private/MVVMEditorSubsystem.cpp
--- a/Engine/Plugins/Runtime/ModelViewViewModel/Source/ModelViewViewModelEditor/Private/MVVMEditorSubsystem.cpp
+++ b/Engine/Plugins/Runtime/ModelViewViewModel/Source/ModelViewViewModelEditor/Private/MVVMEditorSubsystem.cpp
@@ -27,6 +27,9 @@
 #include "MVVMSubsystem.h"
 #include "MVVMWidgetBlueprintExtension_View.h"
 #include "ScopedTransaction.h"
+#include "AssetRegistry/IAssetRegistry.h"
+#include "Misc/MessageDialog.h"
+#include "Misc/ScopedSlowTask.h"
 #include "Types/MVVMAvailableBinding.h"
 #include "Types/MVVMBindingSource.h"
 
@@ -157,6 +160,53 @@ namespace UE::MVVM::Private
 		return ConversionNode;
 	}
 
+	bool LoadBlueprintReferencersForViewModelRename(const UBlueprint* Blueprint)
+	{
+		if (Blueprint)
+		{
+			FARFilter Filter;
+			Filter.bIncludeOnlyOnDiskAssets = true;
+			Filter.bRecursiveClasses = true;
+			Filter.ClassPaths.Add(UBlueprint::StaticClass()->GetClassPathName());
+
+			IAssetRegistry& AssetRegistry = IAssetRegistry::GetChecked();
+			AssetRegistry.GetReferencers(Blueprint->GetPackage()->GetFName(), Filter.PackageNames, UE::AssetRegistry::EDependencyCategory::Package);
+
+			TArray<FAssetData> BlueprintReferencers;
+			AssetRegistry.GetAssets(Filter, BlueprintReferencers);
+
+			BlueprintReferencers = BlueprintReferencers.FilterByPredicate([](const FAssetData& Asset)
+			{
+				constexpr bool bLoad = false;
+				return Asset.FastGetAsset(bLoad) == nullptr;
+			});
+
+			if (!BlueprintReferencers.IsEmpty())
+			{
+				static const FTextFormat LoadingWarningFormat = LOCTEXT("LoadingBlueprintReferencersWarningFormat", "Renaming this view model requires loading {0} {0}|plural(one=asset,other=assets).\nDo you want to continue?");
+				const EAppReturnType::Type DialogResult = FMessageDialog::Open(EAppMsgType::YesNo, FText::Format(LoadingWarningFormat, BlueprintReferencers.Num()));
+				if (DialogResult == EAppReturnType::No)
+				{
+					return false;
+				}
+				
+				FScopedSlowTask LoadReferencersSlowTask = FScopedSlowTask(BlueprintReferencers.Num(), LOCTEXT("LoadingRenamedViewModelBindings", "Loading Dependent Blueprints"));
+				
+				constexpr float TimeBeforeDialog = 1.f;
+				LoadReferencersSlowTask.MakeDialogDelayed(TimeBeforeDialog);
+				for (const FAssetData& Asset : BlueprintReferencers)
+				{
+					LoadReferencersSlowTask.EnterProgressFrame(1.0f, FText::Format(LOCTEXT("LoadingBlueprintReferencerFormat", "Loading {0}"), FText::FromName(Asset.AssetName)));
+					
+					constexpr bool bLoad = true;
+					Asset.FastGetAsset(bLoad);
+				}
+			}
+		}
+
+		return true;
+	}
+
 } //namespace Private
 
 
@@ -315,10 +365,16 @@ bool UMVVMEditorSubsystem::RenameViewModel(UWidgetBlueprint* WidgetBlueprint, FN
 	const FMVVMBlueprintViewModelContext* ViewModelContext = View->FindViewModel(ViewModel);
 	if (ViewModelContext && ViewModelContext->bCanRename)
 	{
-		const FScopedTransaction Transaction(LOCTEXT("RenameViewModel", "Rename Viewmodel"));
-		View->Modify();
-		return View->RenameViewModel(ViewModel, NewViewModel);
+		// Load any dependent BPs so that they can update any bindings that include the renamed view model
+		// Do this before starting the transaction so that any changes during load are excluded
+		if (UE::MVVM::Private::LoadBlueprintReferencersForViewModelRename(WidgetBlueprint))
+		{
+			const FScopedTransaction Transaction(LOCTEXT("RenameViewModel", "Rename Viewmodel"));
+			View->Modify();
+			return View->RenameViewModel(ViewModel, NewViewModel);
+		}
 	}
+
 	return false;
 }
 
diff --git a/Engine/Plugins/Runtime/ModelViewViewModel/Source/ModelViewViewModelEditor/Private/ModelViewViewModelEditorModule.cpp b/Engine/Plugins/Runtime/ModelViewViewModel/Source/ModelViewViewModelEditor/Private/ModelViewViewModelEditorModule.cpp
--- a/Engine/Plugins/Runtime/ModelViewViewModel/Source/ModelViewViewModelEditor/Private/ModelViewViewModelEditorModule.cpp
+++ b/Engine/Plugins/Runtime/ModelViewViewModel/Source/ModelViewViewModelEditor/Private/ModelViewViewModelEditorModule.cpp
@@ -68,7 +68,8 @@ void FModelViewViewModelEditorModule::StartupModule()
 		, IUMGEditorModule::FOnGetInstancePropertyTypeCustomizationInstance::CreateStatic(UE::MVVM::FBlueprintViewModelContextDetailCustomization::MakeInstance)
 		);
 
-	FBlueprintEditorUtils::OnRenameVariableReferencesEvent.AddRaw(this, &FModelViewViewModelEditorModule::HandleRenameVariableReferences);
+	FBlueprintEditorUtils::OnRenameVariableReferencesEvent.AddRaw(this, &FModelViewViewModelEditorModule::HandleRenameFieldReferences);
+	FBlueprintEditorUtils::OnRenameFunctionReferencesEvent.AddRaw(this, &FModelViewViewModelEditorModule::HandleRenameFieldReferences);
 
 	{
 		FMessageLogModule& MessageLogModule = FModuleManager::LoadModuleChecked<FMessageLogModule>("MessageLog");
@@ -101,6 +102,7 @@ void FModelViewViewModelEditorModule::ShutdownModule()
 	}
 
 	FBlueprintEditorUtils::OnRenameVariableReferencesEvent.RemoveAll(this);
+	FBlueprintEditorUtils::OnRenameFunctionReferencesEvent.RemoveAll(this);
 
 	if (IUMGEditorModule* UMGEditorModule = FModuleManager::GetModulePtr<IUMGEditorModule>("UMGEditor"))
 	{
@@ -150,7 +152,7 @@ void FModelViewViewModelEditorModule::HandleRegisterBlueprintEditorTab(const FWi
 }
 
 
-void FModelViewViewModelEditorModule::HandleRenameVariableReferences(UBlueprint* Blueprint, UClass* VariableClass, const FName& OldVarName, const FName& NewVarName)
+void FModelViewViewModelEditorModule::HandleRenameFieldReferences(UBlueprint* Blueprint, UClass* FieldOwnerClass, const FName& OldVarName, const FName& NewVarName)
 {
 	if (UWidgetBlueprint* WidgetBlueprint = Cast<UWidgetBlueprint>(Blueprint))
 	{
@@ -158,10 +160,10 @@ void FModelViewViewModelEditorModule::HandleRenameVariableReferences(UBlueprint*
 		{
 			if (UMVVMBlueprintView* BlueprintView = ViewExtension->GetBlueprintView())
 			{
-				BlueprintView->WidgetRenamed(OldVarName, NewVarName);
+				BlueprintView->OnFieldRenamed(FieldOwnerClass, OldVarName, NewVarName);
 			}
 
-			ViewExtension->RenameWidgetExtensions(OldVarName, NewVarName);
+			ViewExtension->OnFieldRenamed(FieldOwnerClass, OldVarName, NewVarName);
 		}
 	}
 }
diff --git a/Engine/Plugins/Runtime/ModelViewViewModel/Source/ModelViewViewModelEditor/Private/ModelViewViewModelEditorModule.h b/Engine/Plugins/Runtime/ModelViewViewModel/Source/ModelViewViewModelEditor/Private/ModelViewViewModelEditorModule.h
--- a/Engine/Plugins/Runtime/ModelViewViewModel/Source/ModelViewViewModelEditor/Private/ModelViewViewModelEditorModule.h
+++ b/Engine/Plugins/Runtime/ModelViewViewModel/Source/ModelViewViewModelEditor/Private/ModelViewViewModelEditorModule.h
@@ -38,7 +38,7 @@ public:
 
 private:
 	void HandleRegisterBlueprintEditorTab(const FWidgetBlueprintApplicationMode& ApplicationMode, FWorkflowAllowedTabSet& TabFactories);
-	void HandleRenameVariableReferences(UBlueprint* Blueprint, UClass* VariableClass, const FName& OldVarName, const FName& NewVarName);
+	void HandleRenameFieldReferences(UBlueprint* Blueprint, UClass* FieldOwnerClass, const FName& OldVarName, const FName& NewVarName);
 	void HandleDeactiveMode(FWidgetBlueprintApplicationMode& InDesignerMode);
 	void HandleActivateMode(FWidgetBlueprintApplicationMode& InDesignerMode);
 	void HandleWidgetBlueprintAssetTags(const UWidgetBlueprint* Widget, FAssetRegistryTagsContext Context);
