Allow removing condition binding entries
Fix issue with conversion function losing its parameter values when the entry is changing other parameters

#tests Removed conditions, Edit values on condition block, change event target on event entries in binding entries
#rnx

diff --git a/Engine/Plugins/Runtime/ModelViewViewModel/Source/ModelViewViewModelBlueprint/Private/MVVMBlueprintViewCondition.cpp b/Engine/Plugins/Runtime/ModelViewViewModel/Source/ModelViewViewModelBlueprint/Private/MVVMBlueprintViewCondition.cpp
--- a/Engine/Plugins/Runtime/ModelViewViewModel/Source/ModelViewViewModelBlueprint/Private/MVVMBlueprintViewCondition.cpp
+++ b/Engine/Plugins/Runtime/ModelViewViewModel/Source/ModelViewViewModelBlueprint/Private/MVVMBlueprintViewCondition.cpp
@@ -34,7 +34,8 @@ void UMVVMBlueprintViewCondition::SetConditionPath(FMVVMBlueprintPropertyPath In
 		return;
 	}
 
-	RemoveWrapperGraph();
+	UpdatePinValues();
+	RemoveWrapperGraph(LeaveConversionFunctionCurrentValues);
 
 	ConditionPath = MoveTemp(InConditionPath);
 	GraphName = FName();
@@ -50,7 +51,6 @@ void UMVVMBlueprintViewCondition::SetConditionPath(FMVVMBlueprintPropertyPath In
 	bNeedsToRegenerateChildren = true;
 
 	CreateWrapperGraphInternal();
-	SavePinValues();
 }
 
 void UMVVMBlueprintViewCondition::SetDestinationPath(FMVVMBlueprintPropertyPath InDestinationPath)
@@ -60,7 +60,8 @@ void UMVVMBlueprintViewCondition::SetDestinationPath(FMVVMBlueprintPropertyPath
 		return;
 	}
 
-	RemoveWrapperGraph();
+	// when we update the conversion function, we want to reset all the values to their default
+	RemoveWrapperGraph(RemoveConversionFunctionCurrentValues);
 
 	DestinationPath = MoveTemp(InDestinationPath);
 
@@ -77,13 +78,13 @@ void UMVVMBlueprintViewCondition::SetOperation(EMVVMConditionOperation InOperati
 		return;
 	}
 
-	RemoveWrapperGraph();
+	UpdatePinValues();
+	RemoveWrapperGraph(LeaveConversionFunctionCurrentValues);
 	ConditionOperation = InOperation;
 
 	bNeedsToRegenerateChildren = true;
 
 	CreateWrapperGraphInternal();
-	SavePinValues();
 }
 
 void UMVVMBlueprintViewCondition::SetOperationValue(float NewValue)
@@ -94,13 +95,13 @@ void UMVVMBlueprintViewCondition::SetOperationValue(float NewValue)
 		return;
 	}
 
-	RemoveWrapperGraph();
+	UpdatePinValues();
+	RemoveWrapperGraph(LeaveConversionFunctionCurrentValues);
 	Value = NewValue;
 
 	bNeedsToRegenerateChildren = true;
 
 	CreateWrapperGraphInternal();
-	SavePinValues();
 }
 
 void UMVVMBlueprintViewCondition::SetOperationMaxValue(float NewMaxValue) 
@@ -111,13 +112,13 @@ void UMVVMBlueprintViewCondition::SetOperationMaxValue(float NewMaxValue)
 		return;
 	}
 
-	RemoveWrapperGraph();
+	UpdatePinValues();
+	RemoveWrapperGraph(LeaveConversionFunctionCurrentValues);
 	MaxValue = NewMaxValue;
 
 	bNeedsToRegenerateChildren = true;
 
 	CreateWrapperGraphInternal();
-	SavePinValues();
 }
 
 void UMVVMBlueprintViewCondition::SetCachedWrapperGraphInternal(UEdGraph* Graph, UK2Node* Node, UMVVMK2Node_IsConditionValid* SourceNode)
@@ -159,7 +160,7 @@ UEdGraph* UMVVMBlueprintViewCondition::GetOrCreateWrapperGraph()
 	return CachedWrapperGraph;
 }
 
-void UMVVMBlueprintViewCondition::RemoveWrapperGraph()
+void UMVVMBlueprintViewCondition::RemoveWrapperGraph(ERemoveWrapperGraphParam ActionForCurrentValues)
 {
 	if (CachedWrapperGraph)
 	{
@@ -168,7 +169,10 @@ void UMVVMBlueprintViewCondition::RemoveWrapperGraph()
 	}
 
 	Messages.Empty();
-	SavedPins.Empty();
+	if (ActionForCurrentValues == RemoveConversionFunctionCurrentValues)
+	{
+		SavedPins.Empty();
+	}
 }
 
 UEdGraphPin* UMVVMBlueprintViewCondition::GetOrCreateGraphPin(const FMVVMBlueprintPinId& PinId)
diff --git a/Engine/Plugins/Runtime/ModelViewViewModel/Source/ModelViewViewModelBlueprint/Private/MVVMBlueprintViewEvent.cpp b/Engine/Plugins/Runtime/ModelViewViewModel/Source/ModelViewViewModelBlueprint/Private/MVVMBlueprintViewEvent.cpp
--- a/Engine/Plugins/Runtime/ModelViewViewModel/Source/ModelViewViewModelBlueprint/Private/MVVMBlueprintViewEvent.cpp
+++ b/Engine/Plugins/Runtime/ModelViewViewModel/Source/ModelViewViewModelBlueprint/Private/MVVMBlueprintViewEvent.cpp
@@ -34,7 +34,8 @@ void UMVVMBlueprintViewEvent::SetEventPath(FMVVMBlueprintPropertyPath InEventPat
 		return;
 	}
 
-	RemoveWrapperGraph();
+	UpdatePinValues();
+	RemoveWrapperGraph(LeaveConversionFunctionCurrentValues);
 
 	EventPath = MoveTemp(InEventPath);
 	GraphName = FName();
@@ -50,7 +51,6 @@ void UMVVMBlueprintViewEvent::SetEventPath(FMVVMBlueprintPropertyPath InEventPat
 	bNeedsToRegenerateChildren = true;
 
 	CreateWrapperGraphInternal();
-	SavePinValues();
 }
 
 void UMVVMBlueprintViewEvent::SetDestinationPath(FMVVMBlueprintPropertyPath InDestinationPath)
@@ -60,7 +60,7 @@ void UMVVMBlueprintViewEvent::SetDestinationPath(FMVVMBlueprintPropertyPath InDe
 		return;
 	}
 
-	RemoveWrapperGraph();
+	RemoveWrapperGraph(RemoveConversionFunctionCurrentValues);
 
 	DestinationPath = MoveTemp(InDestinationPath);
 
@@ -109,7 +109,7 @@ UEdGraph* UMVVMBlueprintViewEvent::GetOrCreateWrapperGraph()
 	return CachedWrapperGraph;
 }
 
-void UMVVMBlueprintViewEvent::RemoveWrapperGraph()
+void UMVVMBlueprintViewEvent::RemoveWrapperGraph(ERemoveWrapperGraphParam ActionForCurrentValues)
 {
 	if (CachedWrapperGraph)
 	{
@@ -118,7 +118,10 @@ void UMVVMBlueprintViewEvent::RemoveWrapperGraph()
 	}
 
 	Messages.Empty();
-	SavedPins.Empty();
+	if(ActionForCurrentValues == RemoveConversionFunctionCurrentValues)
+	{
+		SavedPins.Empty();
+	}
 }
 
 UEdGraphPin* UMVVMBlueprintViewEvent::GetOrCreateGraphPin(const FMVVMBlueprintPinId& PinId)
diff --git a/Engine/Plugins/Runtime/ModelViewViewModel/Source/ModelViewViewModelBlueprint/Private/Node/MVVMK2Node_IsConditionValid.cpp b/Engine/Plugins/Runtime/ModelViewViewModel/Source/ModelViewViewModelBlueprint/Private/Node/MVVMK2Node_IsConditionValid.cpp
--- a/Engine/Plugins/Runtime/ModelViewViewModel/Source/ModelViewViewModelBlueprint/Private/Node/MVVMK2Node_IsConditionValid.cpp
+++ b/Engine/Plugins/Runtime/ModelViewViewModel/Source/ModelViewViewModelBlueprint/Private/Node/MVVMK2Node_IsConditionValid.cpp
@@ -31,11 +31,10 @@ void UMVVMK2Node_IsConditionValid::AllocateDefaultPins()
 	const UEdGraphSchema_K2* K2Schema = GetDefault<UEdGraphSchema_K2>();
 
 	CreatePin(EGPD_Input, UEdGraphSchema_K2::PC_Exec, UEdGraphSchema_K2::PN_Execute);
-
 	CreatePin(EGPD_Input, UEdGraphSchema_K2::PC_Byte, StaticEnum<EMVVMConditionOperation>(), PN_Operation);
-	CreatePin(EGPD_Input, UEdGraphSchema_K2::PC_Real, PN_Value);
-	CreatePin(EGPD_Input, UEdGraphSchema_K2::PC_Real, PN_CompareValue);
-	CreatePin(EGPD_Input, UEdGraphSchema_K2::PC_Real, PN_CompareMaxValue);
+	CreatePin(EGPD_Input, UEdGraphSchema_K2::PC_Real, UEdGraphSchema_K2::PC_Double, PN_Value);
+	CreatePin(EGPD_Input, UEdGraphSchema_K2::PC_Real, UEdGraphSchema_K2::PC_Double, PN_CompareValue);
+	CreatePin(EGPD_Input, UEdGraphSchema_K2::PC_Real, UEdGraphSchema_K2::PC_Double, PN_CompareMaxValue);
 
 	UEdGraphPin* TruePin = CreatePin(EGPD_Output, UEdGraphSchema_K2::PC_Exec, UEdGraphSchema_K2::PN_Then);
 	TruePin->PinFriendlyName = LOCTEXT("true", "true");
diff --git a/Engine/Plugins/Runtime/ModelViewViewModel/Source/ModelViewViewModelBlueprint/Public/MVVMBlueprintViewCondition.h b/Engine/Plugins/Runtime/ModelViewViewModel/Source/ModelViewViewModelBlueprint/Public/MVVMBlueprintViewCondition.h
--- a/Engine/Plugins/Runtime/ModelViewViewModel/Source/ModelViewViewModelBlueprint/Public/MVVMBlueprintViewCondition.h
+++ b/Engine/Plugins/Runtime/ModelViewViewModel/Source/ModelViewViewModelBlueprint/Public/MVVMBlueprintViewCondition.h
@@ -92,7 +92,12 @@ public:
 		return GraphName;
 	}
 
-	void RemoveWrapperGraph();
+	enum ERemoveWrapperGraphParam
+	{
+		RemoveConversionFunctionCurrentValues, // when removing or changing the conversion function, we want to remove all the conversion function parameters
+		LeaveConversionFunctionCurrentValues // when we remove the wrapper graph because the destination path or condition properties are changed, we want to keep the conversion function parameters
+	};
+	void RemoveWrapperGraph(ERemoveWrapperGraphParam ActionForCurrentValues = RemoveConversionFunctionCurrentValues);
 
 	UK2Node* GetWrapperDestinationNode() const
 	{
diff --git a/Engine/Plugins/Runtime/ModelViewViewModel/Source/ModelViewViewModelBlueprint/Public/MVVMBlueprintViewEvent.h b/Engine/Plugins/Runtime/ModelViewViewModel/Source/ModelViewViewModelBlueprint/Public/MVVMBlueprintViewEvent.h
--- a/Engine/Plugins/Runtime/ModelViewViewModel/Source/ModelViewViewModelBlueprint/Public/MVVMBlueprintViewEvent.h
+++ b/Engine/Plugins/Runtime/ModelViewViewModel/Source/ModelViewViewModelBlueprint/Public/MVVMBlueprintViewEvent.h
@@ -75,7 +75,12 @@ public:
 		return GraphName;
 	}
 
-	void RemoveWrapperGraph();
+	enum ERemoveWrapperGraphParam
+	{
+		RemoveConversionFunctionCurrentValues, // when removing or changing the conversion function, we want to remove all the conversion function parameters
+		LeaveConversionFunctionCurrentValues // when we remove the wrapper graph because the event path has changed, we want to keep the conversion function parameters
+	};
+	void RemoveWrapperGraph(ERemoveWrapperGraphParam ActionForCurrentValues = RemoveConversionFunctionCurrentValues);
 
 	UK2Node* GetWrapperNode() const
 	{
diff --git a/Engine/Plugins/Runtime/ModelViewViewModel/Source/ModelViewViewModelEditor/Private/Framework/MVVMRowHelper.cpp b/Engine/Plugins/Runtime/ModelViewViewModel/Source/ModelViewViewModelEditor/Private/Framework/MVVMRowHelper.cpp
--- a/Engine/Plugins/Runtime/ModelViewViewModel/Source/ModelViewViewModelEditor/Private/Framework/MVVMRowHelper.cpp
+++ b/Engine/Plugins/Runtime/ModelViewViewModel/Source/ModelViewViewModelEditor/Private/Framework/MVVMRowHelper.cpp
@@ -36,7 +36,7 @@ extern UNREALED_API UEditorEngine* GEditor;
 namespace UE::MVVM::BindingEntry
 {
 
-void FRowHelper::GatherAllChildBindings(UMVVMBlueprintView* BlueprintView, const TConstArrayView<TSharedPtr<FBindingEntry>> Entries, TArray<const FMVVMBlueprintViewBinding*>& OutBindings, TArray<UMVVMBlueprintViewEvent*>& OutEvents)
+void FRowHelper::GatherAllChildBindings(UMVVMBlueprintView* BlueprintView, const TConstArrayView<TSharedPtr<FBindingEntry>> Entries, TArray<const FMVVMBlueprintViewBinding*>& OutBindings, TArray<UMVVMBlueprintViewEvent*>& OutEvents, TArray<UMVVMBlueprintViewCondition*>& OutConditions)
 {
 	for (const TSharedPtr<FBindingEntry>& Entry : Entries)
 	{
@@ -57,7 +57,16 @@ void FRowHelper::GatherAllChildBindings(UMVVMBlueprintView* BlueprintView, const
 				OutEvents.AddUnique(Event);
 			}
 		}
-		GatherAllChildBindings(BlueprintView, Entry->GetAllChildren(), OutBindings, OutEvents);
+
+		if (Entry->GetRowType() == FBindingEntry::ERowType::Condition)
+		{
+			UMVVMBlueprintViewCondition* Condition = Entry->GetCondition();
+			if (Condition)
+			{
+				OutConditions.AddUnique(Condition);
+			}
+		}
+		GatherAllChildBindings(BlueprintView, Entry->GetAllChildren(), OutBindings, OutEvents, OutConditions);
 	}
 }
 
@@ -67,9 +76,10 @@ void FRowHelper::DeleteEntries(const UWidgetBlueprint* WidgetBlueprint, UMVVMBlu
 	{
 		TArray<const FMVVMBlueprintViewBinding*> BindingsToRemove;
 		TArray<UMVVMBlueprintViewEvent*> EventsToRemove;
-		GatherAllChildBindings(BlueprintView, Selection, BindingsToRemove, EventsToRemove);
+		TArray<UMVVMBlueprintViewCondition*> ConditionsToRemove;
+		GatherAllChildBindings(BlueprintView, Selection, BindingsToRemove, EventsToRemove, ConditionsToRemove);
 
-		if (BindingsToRemove.Num() == 0 && EventsToRemove.Num() == 0)
+		if (BindingsToRemove.Num() == 0 && EventsToRemove.Num() == 0 && ConditionsToRemove.Num() == 0)
 		{
 			return;
 		}
@@ -83,6 +93,10 @@ void FRowHelper::DeleteEntries(const UWidgetBlueprint* WidgetBlueprint, UMVVMBlu
 		{
 			BindingDisplayNames.Add(Event->GetDisplayName(true));
 		}
+		for (const UMVVMBlueprintViewCondition* Condition : ConditionsToRemove)
+		{
+			BindingDisplayNames.Add(Condition->GetDisplayName(true));
+		}
 
 		const FText Message = FText::Format(BindingDisplayNames.Num() == 1 ?
 			LOCTEXT("ConfirmDeleteSingle", "Are you sure that you want to delete this binding/event?\n\n{1}") :
@@ -105,6 +119,10 @@ void FRowHelper::DeleteEntries(const UWidgetBlueprint* WidgetBlueprint, UMVVMBlu
 			{
 				BlueprintView->RemoveEvent(Event);
 			}
+			for (UMVVMBlueprintViewCondition* Condition : ConditionsToRemove)
+			{
+				BlueprintView->RemoveCondition(Condition);
+			}
 		}
 	}
 }
@@ -390,6 +408,7 @@ FMenuBuilder FRowHelper::CreateContextMenu(UWidgetBlueprint* WidgetBlueprint, UM
 				case FBindingEntry::ERowType::Group:
 				case FBindingEntry::ERowType::Binding:
 				case FBindingEntry::ERowType::Event:
+				case FBindingEntry::ERowType::Condition:
 					break;
 				default:
 					bCanRemoveEntry = false;
diff --git a/Engine/Plugins/Runtime/ModelViewViewModel/Source/ModelViewViewModelEditor/Private/Framework/MVVMRowHelper.h b/Engine/Plugins/Runtime/ModelViewViewModel/Source/ModelViewViewModelEditor/Private/Framework/MVVMRowHelper.h
--- a/Engine/Plugins/Runtime/ModelViewViewModel/Source/ModelViewViewModelEditor/Private/Framework/MVVMRowHelper.h
+++ b/Engine/Plugins/Runtime/ModelViewViewModel/Source/ModelViewViewModelEditor/Private/Framework/MVVMRowHelper.h
@@ -13,6 +13,8 @@ class UWidgetBlueprint;
 class UMVVMBlueprintView;
 struct FMVVMBlueprintViewBinding;
 class UMVVMBlueprintViewEvent;
+class UMVVMBlueprintViewCondition;
+
 namespace UE::MVVM { struct FBindingEntry; }
 
 namespace UE::MVVM::BindingEntry
@@ -20,7 +22,7 @@ namespace UE::MVVM::BindingEntry
 
 struct FRowHelper
 {
-	static void GatherAllChildBindings(UMVVMBlueprintView* BlueprintView, const TConstArrayView<TSharedPtr<FBindingEntry>> Entries, TArray<const FMVVMBlueprintViewBinding*>& OutBindings, TArray<UMVVMBlueprintViewEvent*>& OutEvents);
+	static void GatherAllChildBindings(UMVVMBlueprintView* BlueprintView, const TConstArrayView<TSharedPtr<FBindingEntry>> Entries, TArray<const FMVVMBlueprintViewBinding*>& OutBindings, TArray<UMVVMBlueprintViewEvent*>& OutEvents, TArray<UMVVMBlueprintViewCondition*>& OutConditions);
 
 	static void DeleteEntries(const UWidgetBlueprint* WidgetBlueprint, UMVVMBlueprintView* BlueprintView, TArrayView<const TSharedPtr<FBindingEntry>> Selection);
 	static void ShowBlueprintGraph(FBlueprintEditor* Editor, UWidgetBlueprint* WidgetBlueprint, UMVVMBlueprintView* BlueprintView, TArrayView<const TSharedPtr<FBindingEntry>> Selection);
