MVVM: Refresh rows of parameters that have changed types

- The binding list entries track the PinType of represented parameters which is then used when checking for existing entries during a refresh
- Add some missing logic for the recently added binding data

diff --git a/Engine/Plugins/Runtime/ModelViewViewModel/Source/ModelViewViewModelEditor/Private/Types/MVVMBindingEntry.cpp b/Engine/Plugins/Runtime/ModelViewViewModel/Source/ModelViewViewModelEditor/Private/Types/MVVMBindingEntry.cpp
--- a/Engine/Plugins/Runtime/ModelViewViewModel/Source/ModelViewViewModelEditor/Private/Types/MVVMBindingEntry.cpp
+++ b/Engine/Plugins/Runtime/ModelViewViewModel/Source/ModelViewViewModelEditor/Private/Types/MVVMBindingEntry.cpp
@@ -52,12 +52,13 @@ void FBindingEntry::SetBindingId(FGuid Id)
 	BindingId = Id;
 }
 
-void FBindingEntry::SetBindingParameter(FGuid Id, FMVVMBlueprintPinId Parameter)
+void FBindingEntry::SetBindingParameter(FGuid Id, FMVVMBlueprintPinId Parameter, FEdGraphPinType ParameterType)
 {
 	check(RowType == ERowType::None);
 	RowType = ERowType::BindingParameter;
 	BindingId = Id;
 	PinId = MoveTemp(Parameter);
+	PinType = MoveTemp(ParameterType);
 }
 
 UMVVMBlueprintViewEvent* FBindingEntry::GetEvent() const
@@ -72,12 +73,13 @@ void FBindingEntry::SetEvent(UMVVMBlueprintViewEvent* InEvent)
 	Event = InEvent;
 }
 
-void FBindingEntry::SetEventParameter(UMVVMBlueprintViewEvent* InEvent, FMVVMBlueprintPinId Parameter)
+void FBindingEntry::SetEventParameter(UMVVMBlueprintViewEvent* InEvent, FMVVMBlueprintPinId Parameter, FEdGraphPinType ParameterType)
 {
 	check(RowType == ERowType::None);
 	RowType = ERowType::EventParameter;
 	Event = InEvent;
 	PinId = MoveTemp(Parameter);
+	PinType = MoveTemp(ParameterType);
 }
 
 UMVVMBlueprintViewCondition* FBindingEntry::GetCondition() const
@@ -92,12 +94,13 @@ void FBindingEntry::SetCondition(UMVVMBlueprintViewCondition* InCondition)
 	Condition = InCondition;
 }
 
-void FBindingEntry::SetConditionParameter(UMVVMBlueprintViewCondition* InCondition, FMVVMBlueprintPinId Parameter)
+void FBindingEntry::SetConditionParameter(UMVVMBlueprintViewCondition* InCondition, FMVVMBlueprintPinId Parameter, FEdGraphPinType ParameterType)
 {
 	check(RowType == ERowType::None);
 	RowType = ERowType::ConditionParameter;
 	Condition = InCondition;
 	PinId = MoveTemp(Parameter);
+	PinType = MoveTemp(ParameterType);
 }
 
 
@@ -130,7 +133,9 @@ bool FBindingEntry::operator==(const FBindingEntry& Other) const
 		&& Name == Other.Name
 		&& BindingId == Other.BindingId
 		&& PinId == Other.PinId
-		&& Event == Other.Event;
+		&& PinType == Other.PinType
+		&& Event == Other.Event
+		&& Condition == Other.Condition;
 }
 
 FString FBindingEntry::GetSearchNameString(UMVVMBlueprintView* View, UWidgetBlueprint* WidgetBP)
@@ -145,6 +150,7 @@ FString FBindingEntry::GetSearchNameString(UMVVMBlueprintView* View, UWidgetBlue
 		break;
 	case UE::MVVM::FBindingEntry::ERowType::BindingParameter:
 	case UE::MVVM::FBindingEntry::ERowType::EventParameter:
+	case UE::MVVM::FBindingEntry::ERowType::ConditionParameter:
 		RowToString = PinId.ToString();
 		break;
 	case UE::MVVM::FBindingEntry::ERowType::Binding:
@@ -160,6 +166,12 @@ FString FBindingEntry::GetSearchNameString(UMVVMBlueprintView* View, UWidgetBlue
 			RowToString.Append(EventInRow->GetSearchableString());
 		}
 		break;
+	case UE::MVVM::FBindingEntry::ERowType::Condition:
+		if (UMVVMBlueprintViewCondition* ConditionInRow = GetCondition())
+		{
+			RowToString.Append(ConditionInRow->GetSearchableString());
+		}
+		break;
 	default:
 		break;
 	}
diff --git a/Engine/Plugins/Runtime/ModelViewViewModel/Source/ModelViewViewModelEditor/Private/Types/MVVMBindingEntry.h b/Engine/Plugins/Runtime/ModelViewViewModel/Source/ModelViewViewModelEditor/Private/Types/MVVMBindingEntry.h
--- a/Engine/Plugins/Runtime/ModelViewViewModel/Source/ModelViewViewModelEditor/Private/Types/MVVMBindingEntry.h
+++ b/Engine/Plugins/Runtime/ModelViewViewModel/Source/ModelViewViewModelEditor/Private/Types/MVVMBindingEntry.h
@@ -75,8 +75,13 @@ struct FBindingEntry
 	{
 		return PinId;
 	}
+	
+	const FEdGraphPinType& GetBindingParameterType() const
+	{
+		return PinType;
+	}
 
-	void SetBindingParameter(FGuid Id, FMVVMBlueprintPinId Parameter);
+	void SetBindingParameter(FGuid Id, FMVVMBlueprintPinId Parameter, FEdGraphPinType ParameterType);
 
 	//~ event
 	UMVVMBlueprintViewEvent* GetEvent() const;
@@ -88,19 +93,31 @@ struct FBindingEntry
 	{
 		return PinId;
 	}
+	
+	const FEdGraphPinType& GetEventParameterType() const
+	{
+		return PinType;
+	}
 
-	void SetEventParameter(UMVVMBlueprintViewEvent* Event, FMVVMBlueprintPinId Parameter);
+	void SetEventParameter(UMVVMBlueprintViewEvent* Event, FMVVMBlueprintPinId Parameter, FEdGraphPinType ParameterType);
 
+	//~ condition
 	UMVVMBlueprintViewCondition* GetCondition() const;
 
 	void SetCondition(UMVVMBlueprintViewCondition* InCondition);
 
+	//~ condition parameter
 	const FMVVMBlueprintPinId& GetConditionParameterId() const
 	{
 		return PinId;
 	}
+	
+	const FEdGraphPinType& GetConditionParameterType() const
+	{
+		return PinType;
+	}
 
-	void SetConditionParameter(UMVVMBlueprintViewCondition* Condition, FMVVMBlueprintPinId Parameter);
+	void SetConditionParameter(UMVVMBlueprintViewCondition* Condition, FMVVMBlueprintPinId Parameter, FEdGraphPinType ParameterType);
 
 
 	//~ children
@@ -131,6 +148,7 @@ private:
 	FName Name;
 	FGuid BindingId;
 	FMVVMBlueprintPinId PinId;
+	FEdGraphPinType PinType;
 	TWeakObjectPtr<UMVVMBlueprintViewEvent> Event;
 	TWeakObjectPtr<UMVVMBlueprintViewCondition> Condition;
 	TArray<TSharedPtr<FBindingEntry>> AllChildren;
diff --git a/Engine/Plugins/Runtime/ModelViewViewModel/Source/ModelViewViewModelEditor/Private/Widgets/SMVVMViewBindingListView.cpp b/Engine/Plugins/Runtime/ModelViewViewModel/Source/ModelViewViewModelEditor/Private/Widgets/SMVVMViewBindingListView.cpp
--- a/Engine/Plugins/Runtime/ModelViewViewModel/Source/ModelViewViewModelEditor/Private/Widgets/SMVVMViewBindingListView.cpp
+++ b/Engine/Plugins/Runtime/ModelViewViewModel/Source/ModelViewViewModelEditor/Private/Widgets/SMVVMViewBindingListView.cpp
@@ -419,13 +419,14 @@ void SBindingsList::Refresh()
 						continue;
 					}
 
+					FEdGraphPinType PinType = GraphPin ? GraphPin->PinType : FEdGraphPinType();
 					TSharedPtr<FBindingEntry> ArgumentEntry;
 					if (PreviousGroupEntry)
 					{
 						TSharedPtr<FBindingEntry>* FoundParameter = PreviousGroupEntry->Children.FindByPredicate(
-							[BindingId, ArgumentId = Pin.GetId()](const TSharedPtr<FBindingEntry>& Other)
+							[BindingId, &PinType, ArgumentId = Pin.GetId()](const TSharedPtr<FBindingEntry>& Other)
 							{
-								return Other->GetBindingId() == BindingId && Other->GetRowType() == FBindingEntry::ERowType::BindingParameter && Other->GetBindingParameterId() == ArgumentId;
+								return Other->GetBindingId() == BindingId && Other->GetRowType() == FBindingEntry::ERowType::BindingParameter && Other->GetBindingParameterId() == ArgumentId && Other->GetBindingParameterType() == PinType;
 							});
 						if (FoundParameter)
 						{
@@ -436,7 +437,7 @@ void SBindingsList::Refresh()
 					if (!ArgumentEntry.IsValid())
 					{
 						ArgumentEntry = MakeShared<FBindingEntry>();
-						ArgumentEntry->SetBindingParameter(Binding.BindingId, Pin.GetId());
+						ArgumentEntry->SetBindingParameter(Binding.BindingId, Pin.GetId(), MoveTemp(PinType));
 
 						NewEntries.Add(ArgumentEntry);
 					}
@@ -509,13 +510,14 @@ void SBindingsList::Refresh()
 					continue;
 				}
 
+				FEdGraphPinType PinType = GraphPin ? GraphPin->PinType : FEdGraphPinType();
 				TSharedPtr<FBindingEntry> ArgumentEntry;
 				if (PreviousGroupEntry)
 				{
 					TSharedPtr<FBindingEntry>* FoundParameter = PreviousGroupEntry->Children.FindByPredicate(
-						[Event, ArgumentId = Pin.GetId()](const TSharedPtr<FBindingEntry>& Other)
+						[Event, &PinType, ArgumentId = Pin.GetId()](const TSharedPtr<FBindingEntry>& Other)
 						{
-							return Other->GetRowType() == FBindingEntry::ERowType::EventParameter && Other->GetEvent() == Event && Other->GetEventParameterId() == ArgumentId;
+							return Other->GetRowType() == FBindingEntry::ERowType::EventParameter && Other->GetEvent() == Event && Other->GetEventParameterId() == ArgumentId && Other->GetEventParameterType() == PinType;
 						});
 					if (FoundParameter)
 					{
@@ -526,7 +528,7 @@ void SBindingsList::Refresh()
 				if (!ArgumentEntry.IsValid())
 				{
 					ArgumentEntry = MakeShared<FBindingEntry>();
-					ArgumentEntry->SetEventParameter(Event, Pin.GetId());
+					ArgumentEntry->SetEventParameter(Event, Pin.GetId(), MoveTemp(PinType));
 
 					NewEntries.Add(ArgumentEntry);
 				}
@@ -581,13 +583,14 @@ void SBindingsList::Refresh()
 					continue;
 				}
 
+				FEdGraphPinType PinType = GraphPin ? GraphPin->PinType : FEdGraphPinType();
 				TSharedPtr<FBindingEntry> ArgumentEntry;
 				if (PreviousGroupEntry)
 				{
 					TSharedPtr<FBindingEntry>* FoundParameter = PreviousGroupEntry->Children.FindByPredicate(
-						[Condition, ArgumentId = Pin.GetId()](const TSharedPtr<FBindingEntry>& Other)
+						[Condition, &PinType, ArgumentId = Pin.GetId()](const TSharedPtr<FBindingEntry>& Other)
 						{
-							return Other->GetRowType() == FBindingEntry::ERowType::ConditionParameter && Other->GetCondition() == Condition && Other->GetConditionParameterId() == ArgumentId;
+							return Other->GetRowType() == FBindingEntry::ERowType::ConditionParameter && Other->GetCondition() == Condition && Other->GetConditionParameterId() == ArgumentId && Other->GetConditionParameterType() == PinType;
 						});
 					if (FoundParameter)
 					{
@@ -598,7 +601,7 @@ void SBindingsList::Refresh()
 				if (!ArgumentEntry.IsValid())
 				{
 					ArgumentEntry = MakeShared<FBindingEntry>();
-					ArgumentEntry->SetConditionParameter(Condition, Pin.GetId());
+					ArgumentEntry->SetConditionParameter(Condition, Pin.GetId(), MoveTemp(PinType));
 
 					NewEntries.Add(ArgumentEntry);
 				}
