MVVM: Fix issue where renaming a view model in an unreferenced blueprint would load all blueprints

diff --git a/Engine/Plugins/Runtime/ModelViewViewModel/Source/ModelViewViewModelEditor/Private/MVVMEditorSubsystem.cpp b/Engine/Plugins/Runtime/ModelViewViewModel/Source/ModelViewViewModelEditor/Private/MVVMEditorSubsystem.cpp
--- a/Engine/Plugins/Runtime/ModelViewViewModel/Source/ModelViewViewModelEditor/Private/MVVMEditorSubsystem.cpp
+++ b/Engine/Plugins/Runtime/ModelViewViewModel/Source/ModelViewViewModelEditor/Private/MVVMEditorSubsystem.cpp
@@ -173,6 +173,12 @@ namespace UE::MVVM::Private
 			AssetRegistry.GetReferencers(Blueprint->GetPackage()->GetFName(), Filter.PackageNames, UE::AssetRegistry::EDependencyCategory::Package);
 
 			TArray<FAssetData> BlueprintReferencers;
+			if (Filter.PackageNames.IsEmpty())
+			{
+				// Nothing references the blueprint
+				return true;
+			}
+			
 			AssetRegistry.GetAssets(Filter, BlueprintReferencers);
 
 			BlueprintReferencers = BlueprintReferencers.FilterByPredicate([](const FAssetData& Asset)
