Allow condition bindings by default. It can be removed later if needed

#tests Ran in restrictive mode to ensure conditions get enabled. Wasn't able to verify that validate works with the new allowed classes
#rnx

diff --git a/Engine/Plugins/Runtime/ModelViewViewModel/Source/ModelViewViewModelBlueprint/Public/MVVMDeveloperProjectSettings.h b/Engine/Plugins/Runtime/ModelViewViewModel/Source/ModelViewViewModelBlueprint/Public/MVVMDeveloperProjectSettings.h
--- a/Engine/Plugins/Runtime/ModelViewViewModel/Source/ModelViewViewModelBlueprint/Public/MVVMDeveloperProjectSettings.h
+++ b/Engine/Plugins/Runtime/ModelViewViewModel/Source/ModelViewViewModelBlueprint/Public/MVVMDeveloperProjectSettings.h
@@ -154,7 +154,7 @@ public:
 	bool bAllowBindingEvent = true;
 
 	UPROPERTY(EditAnywhere, config, Category = "Features")
-	bool bAllowConditionBinding = false;
+	bool bAllowConditionBinding = true;
 
 	/** Allow to create an instanced viewmodel directly in the view editor. */
 	UPROPERTY(EditAnywhere, config, Category = "Features", meta=(DisplayName="Experimental - Can Create Viewmodel In View"))
