Added tooltip to explain what checkbox does in MVVM view

diff --git a/Engine/Plugins/Runtime/ModelViewViewModel/Source/ModelViewViewModelEditor/Private/Widgets/BindingEntry/SMVVMBindingRow.cpp b/Engine/Plugins/Runtime/ModelViewViewModel/Source/ModelViewViewModelEditor/Private/Widgets/BindingEntry/SMVVMBindingRow.cpp
--- a/Engine/Plugins/Runtime/ModelViewViewModel/Source/ModelViewViewModelEditor/Private/Widgets/BindingEntry/SMVVMBindingRow.cpp
+++ b/Engine/Plugins/Runtime/ModelViewViewModel/Source/ModelViewViewModelEditor/Private/Widgets/BindingEntry/SMVVMBindingRow.cpp
@@ -107,6 +107,7 @@ TSharedRef<SWidget> SBindingRow::BuildRowWidget()
 				SNew(SCheckBox)
 				.IsChecked(this, &SBindingRow::IsBindingCompiled)
 				.OnCheckStateChanged(this, &SBindingRow::OnIsBindingCompileChanged)
+				.ToolTipText(LOCTEXT("MVVMEnabledTooltip", "Is enabled"))
 			]
 
 			+ SHorizontalBox::Slot()
