Replace check calls with ensure calls since the connections are supposed to be created in Shipping as well.
Therefore, compilation of the assets were failing on production backend.

Some other undetected? cases have been corrected at the same time.

swarm review has been upvoted in another branch : https://p4-swarm.epicgames.net/reviews/36646875

#rnx
#tests Tested with and without fix on shipping build. Saw the bug happening while debugging

diff --git a/Engine/Plugins/Runtime/ModelViewViewModel/Source/ModelViewViewModelBlueprint/Private/Node/MVVMK2Node_AreSourcesValidForBinding.cpp b/Engine/Plugins/Runtime/ModelViewViewModel/Source/ModelViewViewModelBlueprint/Private/Node/MVVMK2Node_AreSourcesValidForBinding.cpp
--- a/Engine/Plugins/Runtime/ModelViewViewModel/Source/ModelViewViewModelBlueprint/Private/Node/MVVMK2Node_AreSourcesValidForBinding.cpp
+++ b/Engine/Plugins/Runtime/ModelViewViewModel/Source/ModelViewViewModelBlueprint/Private/Node/MVVMK2Node_AreSourcesValidForBinding.cpp
@@ -85,13 +85,13 @@ void UMVVMK2Node_AreSourcesValidForBinding::ExpandNode(FKismetCompilerContext& C
 	// Casting subsytem.result into MVVMSubsystem
 	CallGetSubsystemNode->GetReturnValuePin()->PinType = CallGetViewNode->FindPinChecked(UEdGraphSchema_K2::PN_Self)->PinType;
 	// subsystem.result to GetViewFromUserWidget.target
-	check(Schema->TryCreateConnection(CallGetSubsystemNode->GetReturnValuePin(), CallGetViewNode->FindPinChecked(UEdGraphSchema_K2::PN_Self)));
+	ensure(Schema->TryCreateConnection(CallGetSubsystemNode->GetReturnValuePin(), CallGetViewNode->FindPinChecked(UEdGraphSchema_K2::PN_Self)));
 	// self to GetViewFromUserWidget.UserWidget
-	check(Schema->TryCreateConnection(SelfNode->FindPinChecked(UEdGraphSchema_K2::PN_Self), CallGetViewNode->FindPinChecked(FName("UserWidget"))));
+	ensure(Schema->TryCreateConnection(SelfNode->FindPinChecked(UEdGraphSchema_K2::PN_Self), CallGetViewNode->FindPinChecked(FName("UserWidget"))));
 	// GetViewFromUserWidget.Result to AreSourcesValidForBinding.target
-	check(Schema->TryCreateConnection(CallGetViewNode->FindPinChecked(UEdGraphSchema_K2::PN_ReturnValue), CallAreSourcesValidForBindingNode->FindPinChecked(UEdGraphSchema_K2::PN_Self)));
+	ensure(Schema->TryCreateConnection(CallGetViewNode->FindPinChecked(UEdGraphSchema_K2::PN_ReturnValue), CallAreSourcesValidForBindingNode->FindPinChecked(UEdGraphSchema_K2::PN_Self)));
 	// AreSourcesValidForBinding.result to branch.condition
-	check(Schema->TryCreateConnection(CallAreSourcesValidForBindingNode->FindPinChecked(UEdGraphSchema_K2::PN_ReturnValue), BranchNode->GetConditionPin()));
+	ensure(Schema->TryCreateConnection(CallAreSourcesValidForBindingNode->FindPinChecked(UEdGraphSchema_K2::PN_ReturnValue), BranchNode->GetConditionPin()));
 	// move this.exec to branch.exec
 	CompilerContext.MovePinLinksToIntermediate(*GetExecPin(), *BranchNode->GetExecPin());
 	// move this.then to branch.then
diff --git a/Engine/Plugins/Runtime/ModelViewViewModel/Source/ModelViewViewModelBlueprint/Private/Node/MVVMK2Node_AreSourcesValidForEvent.cpp b/Engine/Plugins/Runtime/ModelViewViewModel/Source/ModelViewViewModelBlueprint/Private/Node/MVVMK2Node_AreSourcesValidForEvent.cpp
--- a/Engine/Plugins/Runtime/ModelViewViewModel/Source/ModelViewViewModelBlueprint/Private/Node/MVVMK2Node_AreSourcesValidForEvent.cpp
+++ b/Engine/Plugins/Runtime/ModelViewViewModel/Source/ModelViewViewModelBlueprint/Private/Node/MVVMK2Node_AreSourcesValidForEvent.cpp
@@ -85,13 +85,13 @@ void UMVVMK2Node_AreSourcesValidForEvent::ExpandNode(FKismetCompilerContext& Com
 	// Casting subsytem.result into MVVMSubsystem
 	CallGetSubsystemNode->GetReturnValuePin()->PinType = CallGetViewNode->FindPinChecked(UEdGraphSchema_K2::PN_Self)->PinType;
 	// subsystem.result to GetViewFromUserWidget.target
-	check(Schema->TryCreateConnection(CallGetSubsystemNode->GetReturnValuePin(), CallGetViewNode->FindPinChecked(UEdGraphSchema_K2::PN_Self)));
+	ensure(Schema->TryCreateConnection(CallGetSubsystemNode->GetReturnValuePin(), CallGetViewNode->FindPinChecked(UEdGraphSchema_K2::PN_Self)));
 	// self to GetViewFromUserWidget.UserWidget
-	check(Schema->TryCreateConnection(SelfNode->FindPinChecked(UEdGraphSchema_K2::PN_Self), CallGetViewNode->FindPinChecked(FName("UserWidget"))));
+	ensure(Schema->TryCreateConnection(SelfNode->FindPinChecked(UEdGraphSchema_K2::PN_Self), CallGetViewNode->FindPinChecked(FName("UserWidget"))));
 	// GetViewFromUserWidget.Result to AreSourcesValidForEvent.target
-	check(Schema->TryCreateConnection(CallGetViewNode->FindPinChecked(UEdGraphSchema_K2::PN_ReturnValue), CallAreSourcesValidForEventNode->FindPinChecked(UEdGraphSchema_K2::PN_Self)));
+	ensure(Schema->TryCreateConnection(CallGetViewNode->FindPinChecked(UEdGraphSchema_K2::PN_ReturnValue), CallAreSourcesValidForEventNode->FindPinChecked(UEdGraphSchema_K2::PN_Self)));
 	// AreSourcesValidForEvent.result to branch.condition
-	check(Schema->TryCreateConnection(CallAreSourcesValidForEventNode->FindPinChecked(UEdGraphSchema_K2::PN_ReturnValue), BranchNode->GetConditionPin()));
+	ensure(Schema->TryCreateConnection(CallAreSourcesValidForEventNode->FindPinChecked(UEdGraphSchema_K2::PN_ReturnValue), BranchNode->GetConditionPin()));
 	// move this.exec to branch.exec
 	CompilerContext.MovePinLinksToIntermediate(*GetExecPin(), *BranchNode->GetExecPin());
 	// move this.then to branch.then
diff --git a/Engine/Plugins/Runtime/ModelViewViewModel/Source/ModelViewViewModelBlueprint/Private/Node/MVVMK2Node_IsConditionValid.cpp b/Engine/Plugins/Runtime/ModelViewViewModel/Source/ModelViewViewModelBlueprint/Private/Node/MVVMK2Node_IsConditionValid.cpp
--- a/Engine/Plugins/Runtime/ModelViewViewModel/Source/ModelViewViewModelBlueprint/Private/Node/MVVMK2Node_IsConditionValid.cpp
+++ b/Engine/Plugins/Runtime/ModelViewViewModel/Source/ModelViewViewModelBlueprint/Private/Node/MVVMK2Node_IsConditionValid.cpp
@@ -85,15 +85,15 @@ void UMVVMK2Node_IsConditionValid::ExpandNode(FKismetCompilerContext& CompilerCo
 	// Casting subsytem.result into MVVMSubsystem
 	CallGetSubsystemNode->GetReturnValuePin()->PinType = CallCompareFloatValuesNode->FindPinChecked(UEdGraphSchema_K2::PN_Self)->PinType;
 	// subsystem.result to CompareFloatValues.target
-	check(Schema->TryCreateConnection(CallGetSubsystemNode->GetReturnValuePin(), CallCompareFloatValuesNode->FindPinChecked(UEdGraphSchema_K2::PN_Self)));
+	ensure(Schema->TryCreateConnection(CallGetSubsystemNode->GetReturnValuePin(), CallCompareFloatValuesNode->FindPinChecked(UEdGraphSchema_K2::PN_Self)));
 	// CompareFloatValues.result to branch.condition
-	check(Schema->TryCreateConnection(CallCompareFloatValuesNode->FindPinChecked(UEdGraphSchema_K2::PN_ReturnValue), BranchNode->GetConditionPin()));
+	ensure(Schema->TryCreateConnection(CallCompareFloatValuesNode->FindPinChecked(UEdGraphSchema_K2::PN_ReturnValue), BranchNode->GetConditionPin()));
 	CompilerContext.MovePinLinksToIntermediate(*GetOperationPin(), *CallCompareFloatValuesNode->FindPinChecked(FName("Operation")));
 	CompilerContext.MovePinLinksToIntermediate(*GetValuePin(), *CallCompareFloatValuesNode->FindPinChecked(FName("Value")));
 	CompilerContext.MovePinLinksToIntermediate(*GetCompareValuePin(), *CallCompareFloatValuesNode->FindPinChecked(FName("CompareValue")));
 	CompilerContext.MovePinLinksToIntermediate(*GetCompareMaxValuePin(), *CallCompareFloatValuesNode->FindPinChecked(FName("CompareMaxValue")));
 	// connect compare.then to branch.exec
-	check(Schema->TryCreateConnection(CallCompareFloatValuesNode->GetThenPin(), BranchNode->GetExecPin()));
+	ensure(Schema->TryCreateConnection(CallCompareFloatValuesNode->GetThenPin(), BranchNode->GetExecPin()));
 
 	// move this.exec to compare.exec
 	CompilerContext.MovePinLinksToIntermediate(*GetExecPin(), *CallCompareFloatValuesNode->GetExecPin());
diff --git a/Engine/Plugins/Runtime/ModelViewViewModel/Source/ModelViewViewModelEditor/Private/Node/MVVMK2Node_LoadSoftResource.cpp b/Engine/Plugins/Runtime/ModelViewViewModel/Source/ModelViewViewModelEditor/Private/Node/MVVMK2Node_LoadSoftResource.cpp
--- a/Engine/Plugins/Runtime/ModelViewViewModel/Source/ModelViewViewModelEditor/Private/Node/MVVMK2Node_LoadSoftResource.cpp
+++ b/Engine/Plugins/Runtime/ModelViewViewModel/Source/ModelViewViewModelEditor/Private/Node/MVVMK2Node_LoadSoftResource.cpp
@@ -71,14 +71,14 @@ void UMVVMK2Node_LoadSoftResource::ExpandNode(FKismetCompilerContext& CompilerCo
 	// move this.then to CallLoadAssetNode.then
 	CompilerContext.MovePinLinksToIntermediate(*GetThenPin(), *CallLoadAssetNode->GetThenPin());
 	// CallLoadAssetNode.completed to CastNode.exec
-	check(Schema->TryCreateConnection(CallLoadAssetNode->FindPinChecked(UEdGraphSchema_K2::PN_Completed), CastNode->GetExecPin()));
+	ensure(Schema->TryCreateConnection(CallLoadAssetNode->FindPinChecked(UEdGraphSchema_K2::PN_Completed), CastNode->GetExecPin()));
 	// move this.completed to CastNode.then
 	CompilerContext.MovePinLinksToIntermediate(*GetCompletedPin(), *CastNode->GetThenPin());
 
 	// move this.resource to CallLoadAssetNode.arg
 	CompilerContext.MovePinLinksToIntermediate(*FindPinChecked(GetInputPinName()), *LoadAssetInput);
 	// CallLoadAssetNode.result to CastNode.Input
-	check(Schema->TryCreateConnection(LoadAssetOutput, CastInput));
+	ensure(Schema->TryCreateConnection(LoadAssetOutput, CastInput));
 	// move this.result to CastNode.Output
 	CompilerContext.MovePinLinksToIntermediate(*FindPinChecked(GetOutputPinName()), *CastOutput);
 
diff --git a/Engine/Plugins/Runtime/ModelViewViewModel/Source/ModelViewViewModelEditor/Private/Node/MVVMK2Node_MakeBrushFromSoftResource.cpp b/Engine/Plugins/Runtime/ModelViewViewModel/Source/ModelViewViewModelEditor/Private/Node/MVVMK2Node_MakeBrushFromSoftResource.cpp
--- a/Engine/Plugins/Runtime/ModelViewViewModel/Source/ModelViewViewModelEditor/Private/Node/MVVMK2Node_MakeBrushFromSoftResource.cpp
+++ b/Engine/Plugins/Runtime/ModelViewViewModel/Source/ModelViewViewModelEditor/Private/Node/MVVMK2Node_MakeBrushFromSoftResource.cpp
@@ -93,7 +93,7 @@ void UMVVMK2Node_MakeBrushFromSoftResource::ExpandNode(FKismetCompilerContext& C
 	// move this.then to CallLoadAssetNode.then
 	CompilerContext.MovePinLinksToIntermediate(*GetThenPin(), *CallLoadAssetNode->GetThenPin());
 	// CallLoadAssetNode.completed to CastNode.exec
-	check(Schema->TryCreateConnection(CallLoadAssetNode->FindPinChecked(UEdGraphSchema_K2::PN_Completed), CastNode->GetExecPin()));
+	ensure(Schema->TryCreateConnection(CallLoadAssetNode->FindPinChecked(UEdGraphSchema_K2::PN_Completed), CastNode->GetExecPin()));
 	// move this.completed to CastNode.then
 	CompilerContext.MovePinLinksToIntermediate(*GetCompletedPin(), *CastNode->GetThenPin());
 
@@ -104,9 +104,9 @@ void UMVVMK2Node_MakeBrushFromSoftResource::ExpandNode(FKismetCompilerContext& C
 	// move this.height to CallMakeBrushFrom.height
 	CompilerContext.MovePinLinksToIntermediate(*FindPinChecked(UE::MVVM::Private::PinNames::Height), *CallMakeBrushFromHeightPin);
 	// CallLoadAssetNode.result to CastNode.Input
-	check(Schema->TryCreateConnection(LoadAssetOutput, CastInput));
+	ensure(Schema->TryCreateConnection(LoadAssetOutput, CastInput));
 	// CastNode.Output to CallMakeBrushFrom.Resource
-	check(Schema->TryCreateConnection(CastOutput, CallMakeBrushFromResourcePin));
+	ensure(Schema->TryCreateConnection(CastOutput, CallMakeBrushFromResourcePin));
 	// move this.result to CallLoadAssetNode.result
 	CompilerContext.MovePinLinksToIntermediate(*FindPinChecked(GetOutputPinName()), *CallMakeBrushFromResultPin);
 
