[MVVM] Update GetBindingSourceProperty to return function return type properties
* Fixes warning when binding list/panel items to a function instead of a property

#tests Changed binding from a property to a function

diff --git a/Engine/Plugins/Runtime/ModelViewViewModel/Source/ModelViewViewModelBlueprint/Private/MVVMViewBlueprintCompiler.cpp b/Engine/Plugins/Runtime/ModelViewViewModel/Source/ModelViewViewModelBlueprint/Private/MVVMViewBlueprintCompiler.cpp
--- a/Engine/Plugins/Runtime/ModelViewViewModel/Source/ModelViewViewModelBlueprint/Private/MVVMViewBlueprintCompiler.cpp
+++ b/Engine/Plugins/Runtime/ModelViewViewModel/Source/ModelViewViewModelBlueprint/Private/MVVMViewBlueprintCompiler.cpp
@@ -3909,6 +3909,10 @@ void FMVVMViewBlueprintCompiler::PreCompileViewExtensions(UWidgetBlueprintGenera
 								{
 									return LastSourceField.GetProperty();
 								}
+								else if (LastSourceField.IsFunction())
+								{
+									return BindingHelper::GetReturnProperty(LastSourceField.GetFunction());
+								}
 							}
 						}
 					}
