Fix issue with creating conditions and not seeing them in the binding panel until we create another binding

#rnx
#tests Created every combinations of conditions and bindings addition/removal

diff --git a/Engine/Plugins/Runtime/ModelViewViewModel/Source/ModelViewViewModelBlueprint/Public/MVVMBlueprintView.h b/Engine/Plugins/Runtime/ModelViewViewModel/Source/ModelViewViewModelBlueprint/Public/MVVMBlueprintView.h
--- a/Engine/Plugins/Runtime/ModelViewViewModel/Source/ModelViewViewModelBlueprint/Public/MVVMBlueprintView.h
+++ b/Engine/Plugins/Runtime/ModelViewViewModel/Source/ModelViewViewModelBlueprint/Public/MVVMBlueprintView.h
@@ -152,6 +152,11 @@ public:
 		return Conditions;
 	}
 
+	bool HasAnyTypeOfBinding()
+	{
+		return !Bindings.IsEmpty() || !Events.IsEmpty() || !Conditions.IsEmpty();
+	}
+
 	TArray<FText> GetBindingMessages(FGuid Id, UE::MVVM::EBindingMessageType InMessageType) const;
 	bool HasBindingMessage(FGuid Id, UE::MVVM::EBindingMessageType InMessageType) const;
 	void AddMessageToBinding(FGuid Id, UE::MVVM::FBindingMessage MessageToAdd);
diff --git a/Engine/Plugins/Runtime/ModelViewViewModel/Source/ModelViewViewModelEditor/Private/Widgets/SMVVMViewBindingPanel.cpp b/Engine/Plugins/Runtime/ModelViewViewModel/Source/ModelViewViewModelEditor/Private/Widgets/SMVVMViewBindingPanel.cpp
--- a/Engine/Plugins/Runtime/ModelViewViewModel/Source/ModelViewViewModelEditor/Private/Widgets/SMVVMViewBindingPanel.cpp
+++ b/Engine/Plugins/Runtime/ModelViewViewModel/Source/ModelViewViewModelEditor/Private/Widgets/SMVVMViewBindingPanel.cpp
@@ -921,8 +921,7 @@ EVisibility SBindingsPanel::GetVisibility(bool bVisibleWithBindings) const
 {
 	if (UMVVMWidgetBlueprintExtension_View* MVVMExtensionPtr = MVVMExtension.Get())
 	{
-		if (MVVMExtensionPtr->GetBlueprintView() != nullptr &&
-			(MVVMExtensionPtr->GetBlueprintView()->GetNumBindings() > 0 || MVVMExtensionPtr->GetBlueprintView()->GetEvents().Num() > 0))
+		if (MVVMExtensionPtr->GetBlueprintView() != nullptr && MVVMExtensionPtr->GetBlueprintView()->HasAnyTypeOfBinding())
 		{
 			return bVisibleWithBindings ? EVisibility::Visible : EVisibility::Collapsed;
 		}
@@ -965,7 +964,7 @@ EVisibility SBindingsPanel::GetBindingMessageVisibility() const
 	{
 		if (MVVMExtensionPtr->GetBlueprintView() != nullptr &&
 			MVVMExtensionPtr->GetBlueprintView()->GetViewModels().Num() > 0 &&
-			MVVMExtensionPtr->GetBlueprintView()->GetBindings().Num() == 0)
+			!MVVMExtensionPtr->GetBlueprintView()->HasAnyTypeOfBinding())
 		{
 			return EVisibility::Visible;
 		}
