MVVM: fix event generated graphs failing to connect node if we're binding to a function of a blueprint variable on the widgetblueprint

diff --git a/Engine/Plugins/Runtime/ModelViewViewModel/Source/ModelViewViewModelBlueprint/Private/Bindings/MVVMConversionFunctionHelper.cpp b/Engine/Plugins/Runtime/ModelViewViewModel/Source/ModelViewViewModelBlueprint/Private/Bindings/MVVMConversionFunctionHelper.cpp
--- a/Engine/Plugins/Runtime/ModelViewViewModel/Source/ModelViewViewModelBlueprint/Private/Bindings/MVVMConversionFunctionHelper.cpp
+++ b/Engine/Plugins/Runtime/ModelViewViewModel/Source/ModelViewViewModelBlueprint/Private/Bindings/MVVMConversionFunctionHelper.cpp
@@ -483,6 +483,7 @@ namespace NamedNodes
 				RootGetterCreator.Finalize();
 
 				PreviousDataPin = RootSelfNode->FindPinChecked(UEdGraphSchema_K2::PSC_Self);
+				PreviousDataPin->PinType.PinSubCategoryObject = Blueprint->SkeletonGeneratedClass;
 				PreviousClass = Blueprint->SkeletonGeneratedClass;
 				PreviousProperty = RootProperty;
 			}
