MVVM: Fix ensure caused by creating a CDO during blueprint compilation

- This would get hit when compiling a widget that is bound to a property on another widget that is dirty

#tests Create WidgetA and WidgetB. Place any subwidget in WidgetB. Place WidgetB inside of WidgetA. Bind WidgetA->IsEnabled <- WidgetB->SubWidget->IsEnabled. Save all, close all assets, and close the editor. Relaunch the editor, open WidgetB, and place a new subwidget in its hierarchy without compiling. Open WidgetA, hit the ensure "Class named WidgetB creating its CDO while changing its layout".

diff --git a/Engine/Plugins/Runtime/ModelViewViewModel/Source/ModelViewViewModel/Private/MVVMSubsystem.cpp b/Engine/Plugins/Runtime/ModelViewViewModel/Source/ModelViewViewModel/Private/MVVMSubsystem.cpp
--- a/Engine/Plugins/Runtime/ModelViewViewModel/Source/ModelViewViewModel/Private/MVVMSubsystem.cpp
+++ b/Engine/Plugins/Runtime/ModelViewViewModel/Source/ModelViewViewModel/Private/MVVMSubsystem.cpp
@@ -9,6 +9,10 @@
 #include "Types/MVVMAvailableBinding.h"
 #include "Kismet/KismetMathLibrary.h"
 
+#if WITH_EDITORONLY_DATA
+#include "Engine/Blueprint.h"
+#endif
+
 #include UE_INLINE_GENERATED_CPP_BY_NAME(MVVMSubsystem)
 
 #define LOCTEXT_NAMESPACE "MVVMSubsystem"
@@ -309,6 +313,17 @@ namespace UE::MVVM::Private
 		const UClass* OwnerClass = Cast<UClass>(FieldVariant.GetOwner());
 		if (OwnerClass && OwnerClass->ImplementsInterface(UNotifyFieldValueChanged::StaticClass()))
 		{
+#if WITH_EDITORONLY_DATA
+			// Can't access the CDO while changing layout, fallback to skeleton
+			if (OwnerClass->bLayoutChanging)
+			{
+				const UBlueprint* OwnerBlueprint = Cast<UBlueprint>(OwnerClass->ClassGeneratedBy);
+				if (OwnerBlueprint && OwnerBlueprint->SkeletonGeneratedClass)
+				{
+					OwnerClass = OwnerBlueprint->SkeletonGeneratedClass;
+				}
+			}
+#endif
 			TScriptInterface<INotifyFieldValueChanged> DefaultObject = OwnerClass->GetDefaultObject();
 			const UE::FieldNotification::IClassDescriptor& ClassDescriptor = DefaultObject->GetFieldNotificationDescriptor();
 			return GetAvailableBinding(FieldVariant, InAccessor, &ClassDescriptor, bIsBindingToEvent);
