MVVM: allow binding events to blueprint-defined functions

diff --git a/Engine/Plugins/Runtime/ModelViewViewModel/Source/ModelViewViewModel/Private/Bindings/MVVMBindingHelper.cpp b/Engine/Plugins/Runtime/ModelViewViewModel/Source/ModelViewViewModel/Private/Bindings/MVVMBindingHelper.cpp
--- a/Engine/Plugins/Runtime/ModelViewViewModel/Source/ModelViewViewModel/Private/Bindings/MVVMBindingHelper.cpp
+++ b/Engine/Plugins/Runtime/ModelViewViewModel/Source/ModelViewViewModel/Private/Bindings/MVVMBindingHelper.cpp
@@ -164,7 +164,7 @@ namespace UE::MVVM::BindingHelper
 
 	bool IsValidForEventBinding(const UFunction* InFunction)
 	{
-		return Private::IsValidCommon(InFunction) && !InFunction->HasAnyFunctionFlags(FUNC_Const | FUNC_BlueprintPure | FUNC_BlueprintEvent);
+		return Private::IsValidCommon(InFunction) && !InFunction->HasAnyFunctionFlags(FUNC_Const | FUNC_BlueprintPure);
 	}
 
 
