Fix unconnected conversion function issue

#tests Tested with events and conditions. Verified that the function was still receiving the function entry point when not in one of our cases
#rnx

diff --git a/Engine/Plugins/Runtime/ModelViewViewModel/Source/ModelViewViewModelBlueprint/Private/Bindings/MVVMConversionFunctionHelper.cpp b/Engine/Plugins/Runtime/ModelViewViewModel/Source/ModelViewViewModelBlueprint/Private/Bindings/MVVMConversionFunctionHelper.cpp
--- a/Engine/Plugins/Runtime/ModelViewViewModel/Source/ModelViewViewModelBlueprint/Private/Bindings/MVVMConversionFunctionHelper.cpp
+++ b/Engine/Plugins/Runtime/ModelViewViewModel/Source/ModelViewViewModelBlueprint/Private/Bindings/MVVMConversionFunctionHelper.cpp
@@ -33,7 +33,6 @@ namespace UE::MVVM::ConversionFunctionHelper
 namespace Private
 {
 	static const FLazyName ConversionFunctionMetadataKey = "ConversionFunction";
-	static const FLazyName KeepConnectionsMetadataKey = "KeepConnections";
 	static const FStringView AutoPromoteFunctionMetadataKey = TEXT("AutoPromoteFunction");
 	static const FLazyName ConversionFunctionCategory = "Conversion Functions";
 
@@ -1015,34 +1014,13 @@ namespace NamedNodes
 		}
 	}
 
-	void LinkAllNodes(UEdGraph* FunctionGraph, UK2Node_FunctionEntry* FunctionEntry, UEdGraphNode* Wrapper, UK2Node_FunctionResult* FunctionResult)
+	void LinkAllNodes(UEdGraph* FunctionGraph, UEdGraphPin* EntryPin, UEdGraphNode* Wrapper, UK2Node_FunctionResult* FunctionResult)
 	{
-		check(FunctionGraph && FunctionEntry && Wrapper && FunctionResult);
+		check(FunctionGraph && EntryPin && Wrapper && FunctionResult);
 		const UEdGraphSchema* GraphSchema = GetDefault<UMVVMConversionFunctionGraphSchema>();
 
-		UEdGraphPin* ThenPin = FunctionEntry->FindPinChecked(UEdGraphSchema_K2::PN_Then);
+		UEdGraphPin* ThenPin = EntryPin;
 		check(ThenPin);
-
-		while (ThenPin->LinkedTo.Num() == 1)
-		{
-			bool FoundNodeToConnect = false;
-			if (UK2Node* ConnectedNode = Cast<UK2Node>(ThenPin->LinkedTo[0]->GetOwningNode()))
-			{
-				if (IsNodeMarkedToKeepConnections(ConnectedNode))
-				{
-					if (UEdGraphPin* ExecPin = ConnectedNode->FindPin(UEdGraphSchema_K2::PN_Execute))
-					{
-						ThenPin->BreakAllPinLinks();
-						GraphSchema->TryCreateConnection(ThenPin, ExecPin);
-						ThenPin = ConnectedNode->FindPin(UEdGraphSchema_K2::PN_Then);
-						FoundNodeToConnect = true;
-					}
-				}
-			}
-			if (!FoundNodeToConnect)
-				break;
-		}
-
 		ThenPin->BreakAllPinLinks();
 
 		// Break Pins for every pins
@@ -1530,6 +1508,16 @@ void SetPropertyPathForPin(const UBlueprint* Blueprint, const FMVVMBlueprintProp
 	// Remove previous nodes
 	Private::RemoveNodesFromPin(FunctionGraph, PathPin);
 
+	UEdGraphPin* EntryPin = ConversionFunctionEntry->FindPin(UEdGraphSchema_K2::PN_Then);
+	
+	if (UEdGraphPin* ExecutePin = ConversionNode->FindPin(UEdGraphSchema_K2::PN_Execute))
+	{
+		if (!ExecutePin->LinkedTo.IsEmpty())
+		{
+			EntryPin = ExecutePin->LinkedTo[0];
+		}
+	}
+
 	// Add new nodes
 	if (PropertyPath.IsValid())
 	{
@@ -1553,7 +1541,7 @@ void SetPropertyPathForPin(const UBlueprint* Blueprint, const FMVVMBlueprintProp
 	}
 
 	// Link Then / Exec pin
-	Private::LinkAllNodes(FunctionGraph, ConversionFunctionEntry, ConversionNode, ConversionFunctionResult);
+	Private::LinkAllNodes(FunctionGraph, EntryPin, ConversionNode, ConversionFunctionResult);
 }
 
 FMVVMBlueprintPropertyPath GetPropertyPathForArgument(const UBlueprint* WidgetBlueprint, const UK2Node_CallFunction* FunctionNode, FName ArgumentName, bool bSkipResolve)
@@ -1696,18 +1684,13 @@ bool IsAsyncNode(const TSubclassOf<UK2Node> Node)
 
 void MarkNodeToKeepConnections(const UK2Node* Node)
 {
-	check(Node != nullptr);
-	Node->GetPackage()->GetMetaData()->SetValue(Node, Private::KeepConnectionsMetadataKey.Resolve(), TEXT(""));
 }
 
 bool IsNodeMarkedToKeepConnections(const UK2Node* Node)
 {
-	if (Node == nullptr)
-		return false;
-	return Node->GetPackage()->GetMetaData()->HasValue(Node, Private::KeepConnectionsMetadataKey.Resolve());
+	return false;
 }
 
-
 } // UE::MVVM::ConversionFunctionHelper
 
 #undef LOCTEXT_NAMESPACE
diff --git a/Engine/Plugins/Runtime/ModelViewViewModel/Source/ModelViewViewModelBlueprint/Private/MVVMBlueprintViewCondition.cpp b/Engine/Plugins/Runtime/ModelViewViewModel/Source/ModelViewViewModelBlueprint/Private/MVVMBlueprintViewCondition.cpp
--- a/Engine/Plugins/Runtime/ModelViewViewModel/Source/ModelViewViewModelBlueprint/Private/MVVMBlueprintViewCondition.cpp
+++ b/Engine/Plugins/Runtime/ModelViewViewModel/Source/ModelViewViewModelBlueprint/Private/MVVMBlueprintViewCondition.cpp
@@ -335,8 +335,6 @@ UEdGraph* UMVVMBlueprintViewCondition::CreateWrapperGraphInternal()
 
 	UMVVMK2Node_IsConditionValid* BranchNode = Cast<UMVVMK2Node_IsConditionValid>(UE::MVVM::ConversionFunctionHelper::InsertEarlyExitBranchNode(CreateSetterGraphResult.GetValue().NewGraph, UMVVMK2Node_IsConditionValid::StaticClass()));
 
-	UE::MVVM::ConversionFunctionHelper::MarkNodeToKeepConnections(BranchNode);
-
 	const UEdGraphSchema* GraphSchema = GetDefault<UMVVMConversionFunctionGraphSchema>();
 	{
 		GraphSchema->TrySetDefaultValue(*BranchNode->GetOperationPin(), UEnum::GetValueAsString(ConditionOperation));
diff --git a/Engine/Plugins/Runtime/ModelViewViewModel/Source/ModelViewViewModelBlueprint/Private/MVVMBlueprintViewEvent.cpp b/Engine/Plugins/Runtime/ModelViewViewModel/Source/ModelViewViewModelBlueprint/Private/MVVMBlueprintViewEvent.cpp
--- a/Engine/Plugins/Runtime/ModelViewViewModel/Source/ModelViewViewModelBlueprint/Private/MVVMBlueprintViewEvent.cpp
+++ b/Engine/Plugins/Runtime/ModelViewViewModel/Source/ModelViewViewModelBlueprint/Private/MVVMBlueprintViewEvent.cpp
@@ -289,7 +289,6 @@ UEdGraph* UMVVMBlueprintViewEvent::CreateWrapperGraphInternal()
 		UE::MVVM::ConversionFunctionHelper::SetMetaData(CreateSetterGraphResult.GetValue().NewGraph, NAME_Hidden, FStringView());
 
 		UMVVMK2Node_AreSourcesValidForEvent* BranchNode = Cast<UMVVMK2Node_AreSourcesValidForEvent>(UE::MVVM::ConversionFunctionHelper::InsertEarlyExitBranchNode(CreateSetterGraphResult.GetValue().NewGraph, UMVVMK2Node_AreSourcesValidForEvent::StaticClass()));
-		UE::MVVM::ConversionFunctionHelper::MarkNodeToKeepConnections(BranchNode);
 		SetCachedWrapperGraphInternal(CreateSetterGraphResult.GetValue().NewGraph, CreateSetterGraphResult.GetValue().WrappedNode, BranchNode);
 		LoadPinValuesInternal();
 	}
diff --git a/Engine/Plugins/Runtime/ModelViewViewModel/Source/ModelViewViewModelBlueprint/Public/Bindings/MVVMConversionFunctionHelper.h b/Engine/Plugins/Runtime/ModelViewViewModel/Source/ModelViewViewModelBlueprint/Public/Bindings/MVVMConversionFunctionHelper.h
--- a/Engine/Plugins/Runtime/ModelViewViewModel/Source/ModelViewViewModelBlueprint/Public/Bindings/MVVMConversionFunctionHelper.h
+++ b/Engine/Plugins/Runtime/ModelViewViewModel/Source/ModelViewViewModelBlueprint/Public/Bindings/MVVMConversionFunctionHelper.h
@@ -125,7 +125,9 @@ namespace UE::MVVM::ConversionFunctionHelper
 	/** Is the node an async node. */
 	MODELVIEWVIEWMODELBLUEPRINT_API bool IsAsyncNode(const TSubclassOf<UK2Node> Node);
 
+	UE_DEPRECATED(5.6, "This function is deprecated and will always return false.")
 	MODELVIEWVIEWMODELBLUEPRINT_API bool IsNodeMarkedToKeepConnections(const UK2Node* Node);
+	UE_DEPRECATED(5.6, "This function is deprecated and is not taken into consideration anymore.")
 	MODELVIEWVIEWMODELBLUEPRINT_API void MarkNodeToKeepConnections(const UK2Node* Node);
 
 } //namespace
