Updated FMVVMListViewBaseExtensionCustomizationExtender to find the Entry Widget Class by property name instead of display name. This prevents an assert when the editor is set to a language other than English

diff --git a/Engine/Plugins/Runtime/ModelViewViewModel/Source/ModelViewViewModelEditor/Private/Customizations/MVVMListViewBaseExtensionCustomizationExtender.cpp b/Engine/Plugins/Runtime/ModelViewViewModel/Source/ModelViewViewModelEditor/Private/Customizations/MVVMListViewBaseExtensionCustomizationExtender.cpp
--- a/Engine/Plugins/Runtime/ModelViewViewModel/Source/ModelViewViewModelEditor/Private/Customizations/MVVMListViewBaseExtensionCustomizationExtender.cpp
+++ b/Engine/Plugins/Runtime/ModelViewViewModel/Source/ModelViewViewModelEditor/Private/Customizations/MVVMListViewBaseExtensionCustomizationExtender.cpp
@@ -54,7 +54,11 @@ void FMVVMListViewBaseExtensionCustomizationExtender::CustomizeDetails(IDetailLa
 					MVVMCategory.GetDefaultProperties(ListEntryProperties);
 					TSharedRef<IPropertyHandle>* EntryClassPtr = ListEntryProperties.FindByPredicate( [](const TSharedRef<IPropertyHandle> Property) 
 					{ 
-						return Property->GetPropertyDisplayName().EqualToCaseIgnored(FText::FromString(TEXT("Entry Widget Class"))); 
+						if (Property->GetProperty())
+						{
+      						return Property->GetProperty()->GetName() == TEXT("EntryWidgetClass");
+      					}
+						return false;
 					});
 
 					check(EntryClassPtr);
