MVVM: Fix crash during compilation of a widget that depends on a viewmodel that changed layout

Previously, changing the class layout of a viewmodel bound via conversion function would not fix up the bytecode of the conversion function. Adding the viewmodel BP as a dependency fixes this by doing a bytecode-only compile of the widget when the viewmodel bp is compiled.

diff --git a/Engine/Plugins/Runtime/ModelViewViewModel/Source/ModelViewViewModelBlueprint/Private/MVVMWidgetBlueprintExtension_View.cpp b/Engine/Plugins/Runtime/ModelViewViewModel/Source/ModelViewViewModelBlueprint/Private/MVVMWidgetBlueprintExtension_View.cpp
--- a/Engine/Plugins/Runtime/ModelViewViewModel/Source/ModelViewViewModelBlueprint/Private/MVVMWidgetBlueprintExtension_View.cpp
+++ b/Engine/Plugins/Runtime/ModelViewViewModel/Source/ModelViewViewModelBlueprint/Private/MVVMWidgetBlueprintExtension_View.cpp
@@ -254,6 +254,19 @@ void UMVVMWidgetBlueprintExtension_View::HandleBeginCompilation(FWidgetBlueprint
 
 	for (const FMVVMBlueprintViewModelContext& AvailableViewModel : BlueprintView->GetViewModels())
 	{
+#if WITH_EDITORONLY_DATA
+		UWidgetBlueprint* WidgetBlueprint = GetWidgetBlueprint();
+		UClass* ViewModelClass = AvailableViewModel.GetViewModelClass();
+		if (WidgetBlueprint && ViewModelClass)
+		{
+			if (UBlueprint* ViewModelBP = Cast<UBlueprint>(ViewModelClass->ClassGeneratedBy))
+			{
+				ViewModelBP->CachedDependents.Add(WidgetBlueprint);
+				WidgetBlueprint->CachedDependencies.Add(ViewModelBP);
+			}
+		}
+#endif // WITH_EDITORONLY_DATA
+		
 		if (AvailableViewModel.InstancedViewModel)
 		{
 			AvailableViewModel.InstancedViewModel->GenerateClass(false);
