MVVM: Support fixing up references to renamed view models in unloaded blueprints

- Implement UMVVMWidgetBlueprintExtension_View::HandlePopulateGeneratedVariables to add MVVM properties to the Widget BP's GeneratedVariables list
  - Use the stored view model context IDs as the viewmodel property GUIDs
  - Remove the old CreateVariables method of generating the properties
  - Expose this path to MVVM Blueprint View Extensions
  - Splits CreateVariables step into an addtional GatherGeneratedVariabless step
- Fix issue where duplicating or undoing a widget blueprint with view model setters would prevent compilation of that blueprint until reloading it or restarting the editor
  - Made the setter graphs non-transactional
- When renaming view models, add option to skip loading referencers now that we have a path for fixing up at compile/run time
  - We can't get rid of the load completely since `FMVVMBlueprintFieldPath::BindingReference` was not storing the view model GUIDs until now, so the load is still required to fix up those old bindings
  - Bindings compiled and saved after this change will no longer need to be loaded to get fixed up

diff --git a/Engine/Plugins/Runtime/ModelViewViewModel/Source/ModelViewViewModelBlueprint/Private/MVVMViewBlueprintCompiler.cpp b/Engine/Plugins/Runtime/ModelViewViewModel/Source/ModelViewViewModelBlueprint/Private/MVVMViewBlueprintCompiler.cpp
--- a/Engine/Plugins/Runtime/ModelViewViewModel/Source/ModelViewViewModelBlueprint/Private/MVVMViewBlueprintCompiler.cpp
+++ b/Engine/Plugins/Runtime/ModelViewViewModel/Source/ModelViewViewModelBlueprint/Private/MVVMViewBlueprintCompiler.cpp
@@ -416,7 +416,7 @@ void FMVVMViewBlueprintCompiler::CleanOldData(UWidgetBlueprintGeneratedClass* Cl
 }
 
 
-void FMVVMViewBlueprintCompiler::CreateVariables(const FWidgetBlueprintCompilerContext::FCreateVariableContext& Context)
+void FMVVMViewBlueprintCompiler::GatherGeneratedVariables(const FWidgetBlueprintCompilerContext::FPopulateGeneratedVariablesContext& Context)
 {
 	if (!BlueprintView)
 	{
@@ -427,87 +427,183 @@ void FMVVMViewBlueprintCompiler::CreateVariables(const FWidgetBlueprintCompilerC
 	{
 		return;
 	}
+	
+	CreateWidgetMap(Context);
+	CreateBindingList(Context);
+	CreateEventList(Context);
+	CreateConditionList(Context);
+	CreateExtensionList(Context);
+	CreateRequiredProperties(Context);
 
-	if (Context.GetCompileType() == EKismetCompileType::SkeletonOnly)
-	{
-		CreateWidgetMap(Context);
-		CreateBindingList(Context);
-		CreateEventList(Context);
-		CreateConditionList(Context);
-		CreateExtensionList(Context);
-		CreateRequiredProperties(Context);
-		CreatePublicFunctionsDeclaration(Context);
-	}
-
-	// Create variable
+	// Add Generated Variables
 	{
-		auto CreateObjectVariable = [&Context](const FCompilerUserWidgetProperty& UserWidgetProperty) -> FProperty*
+		auto AddObjectVariable = [&Context](const FCompilerUserWidgetProperty& UserWidgetProperty)
 		{
-			FEdGraphPinType NewPropertyPinType(UEdGraphSchema_K2::PC_Object, NAME_None, UserWidgetProperty.AuthoritativeClass, EPinContainerType::None, false, FEdGraphTerminalType());
-			FProperty* NewProperty = Context.CreateVariable(UserWidgetProperty.Name, NewPropertyPinType);
-			if (NewProperty != nullptr)
+			FBPVariableDescription ObjectVariableDesc;
+			ObjectVariableDesc.VarName = UserWidgetProperty.Name;
+			ObjectVariableDesc.VarGuid = UserWidgetProperty.Guid;
+			ObjectVariableDesc.VarType = FEdGraphPinType(UEdGraphSchema_K2::PC_Object, NAME_None, UserWidgetProperty.AuthoritativeClass, EPinContainerType::None, false, FEdGraphTerminalType());
+			ObjectVariableDesc.FriendlyName = UserWidgetProperty.DisplayName.ToString();
+			ObjectVariableDesc.PropertyFlags = (CPF_BlueprintVisible | CPF_RepSkip);			
+
+			const bool bIsInstanceExposed = UserWidgetProperty.bInstanced && UserWidgetProperty.bInstanceExposed;
+			if (bIsInstanceExposed)
 			{
-				NewProperty->SetPropertyFlags(CPF_BlueprintVisible | CPF_RepSkip);
+				ObjectVariableDesc.PropertyFlags |= (CPF_Edit | CPF_ExportObject | CPF_PersistentInstance | CPF_InstancedReference | CPF_NonNullable | CPF_NoClear);
 
-				const bool bIsInstanceExposed = UserWidgetProperty.bInstanced && UserWidgetProperty.bInstanceExposed;
-				if (bIsInstanceExposed)
+				if (UserWidgetProperty.AuthoritativeClass->HasAnyClassFlags(CLASS_HasInstancedReference))
 				{
-					NewProperty->SetPropertyFlags(CPF_Edit | CPF_ExportObject | CPF_PersistentInstance | CPF_InstancedReference | CPF_NonNullable | CPF_NoClear);
-
-					if (UserWidgetProperty.AuthoritativeClass->HasAnyClassFlags(CLASS_HasInstancedReference))
-					{
-						NewProperty->SetPropertyFlags(CPF_ContainsInstancedReference);
-					}
-				}
-				else
-				{
-					NewProperty->SetPropertyFlags(CPF_Transient | CPF_DuplicateTransient);
+					ObjectVariableDesc.PropertyFlags |= (CPF_ContainsInstancedReference);
 				}
+			}
+			else
+			{
+				ObjectVariableDesc.PropertyFlags |= (CPF_Transient | CPF_DuplicateTransient);
+			}
 
-				if (UserWidgetProperty.bExposeOnSpawn)
-				{
-					NewProperty->SetPropertyFlags(CPF_ExposeOnSpawn);
-				}
-				else if (!bIsInstanceExposed)
-				{
-					NewProperty->SetPropertyFlags(CPF_DisableEditOnInstance);
-				}
+			if (UserWidgetProperty.bExposeOnSpawn)
+			{
+				ObjectVariableDesc.PropertyFlags |= (CPF_ExposeOnSpawn);
+			}
+			else if (!bIsInstanceExposed)
+			{
+				ObjectVariableDesc.PropertyFlags |= (CPF_DisableEditOnInstance);
+			}
 
-				if (UserWidgetProperty.bReadOnly)
-				{
-					NewProperty->SetPropertyFlags(CPF_BlueprintReadOnly);
-				}
+			if (UserWidgetProperty.bReadOnly)
+			{
+				ObjectVariableDesc.PropertyFlags |= (CPF_BlueprintReadOnly);
+			}
 
 #if WITH_EDITOR
-				if (!UserWidgetProperty.BlueprintSetter.IsEmpty())
-				{
-					NewProperty->SetMetaData(FBlueprintMetadata::MD_PropertySetFunction, *UserWidgetProperty.BlueprintSetter);
-				}
-				if (!UserWidgetProperty.DisplayName.IsEmpty())
-				{
-					NewProperty->SetMetaData(FBlueprintMetadata::MD_DisplayName, *UserWidgetProperty.DisplayName.ToString());
-				}
-				if (!UserWidgetProperty.CategoryName.IsEmpty())
+			if (!UserWidgetProperty.BlueprintSetter.IsEmpty())
+			{
+				ObjectVariableDesc.SetMetaData(FBlueprintMetadata::MD_PropertySetFunction, *UserWidgetProperty.BlueprintSetter);
+			}
+			if (!UserWidgetProperty.DisplayName.IsEmpty())
+			{
+				ObjectVariableDesc.SetMetaData(FBlueprintMetadata::MD_DisplayName, *UserWidgetProperty.DisplayName.ToString());
+			}
+			if (!UserWidgetProperty.CategoryName.IsEmpty())
+			{
+				ObjectVariableDesc.SetMetaData(FBlueprintMetadata::MD_FunctionCategory, *UserWidgetProperty.CategoryName);
+			}
+			if (UserWidgetProperty.bExposeOnSpawn)
+			{
+				ObjectVariableDesc.SetMetaData(FBlueprintMetadata::MD_ExposeOnSpawn, TEXT("true"));
+			}
+			if (UserWidgetProperty.bPrivate)
+			{
+				ObjectVariableDesc.SetMetaData(FBlueprintMetadata::MD_Private, TEXT("true"));
+			}
+			if (bIsInstanceExposed)
+			{
+				ObjectVariableDesc.SetMetaData(FName("EditInline"), TEXT("true"));
+			}
+#endif
+
+			Context.AddGeneratedVariable(MoveTemp(ObjectVariableDesc));
+		};
+
+		for (FCompilerUserWidgetProperty& UserWidgetProperty : NeededUserWidgetProperties)
+		{
+			check(UserWidgetProperty.AuthoritativeClass);
+			check(!UserWidgetProperty.Name.IsNone());
+
+			// Check if the property already exists in parent classes
+			FMVVMConstFieldVariant UserWidgetPropertyField = BindingHelper::FindFieldByName(Context.GetWidgetBlueprint()->ParentClass, FMVVMBindingName(UserWidgetProperty.Name));
+			if (UserWidgetPropertyField.IsEmpty())
+			{
+				for (FField* Field = Context.GetWidgetBlueprint()->ParentClass->ChildProperties; Field != nullptr; Field = Field->Next)
 				{
-					NewProperty->SetMetaData(FBlueprintMetadata::MD_FunctionCategory, *UserWidgetProperty.CategoryName);
+					if (Field->GetFName() == UserWidgetProperty.Name)
+					{
+						if (CastField<FProperty>(Field))
+						{
+							UserWidgetPropertyField = FMVVMFieldVariant(CastField<FProperty>(Field));
+						}
+						else
+						{
+							WidgetBlueprintCompilerContext.MessageLog.Error(*FText::Format(LOCTEXT("FieldIsNotProperty", "The field for source '{0}' exists but is not a property."), UserWidgetProperty.DisplayName).ToString());
+							bIsGatherGeneratedVariablesStepValid = false;
+						}
+						break;
+					}
 				}
-				if (UserWidgetProperty.bExposeOnSpawn)
+				for (UField* Field = Context.GetWidgetBlueprint()->ParentClass->Children; Field != nullptr; Field = Field->Next)
 				{
-					NewProperty->SetMetaData(FBlueprintMetadata::MD_ExposeOnSpawn, TEXT("true"));
+					if (Field->GetFName() == UserWidgetProperty.Name)
+					{
+						WidgetBlueprintCompilerContext.MessageLog.Error(*FText::Format(LOCTEXT("FieldIsNotProperty", "The field for source '{0}' exists but is not a property."), UserWidgetProperty.DisplayName).ToString());
+						bIsGatherGeneratedVariablesStepValid = false;
+						break;
+					}
 				}
-				if (UserWidgetProperty.bPrivate)
+			}
+			
+			const FBPVariableDescription* ExistingGeneratedVariablePtr = Context.GetWidgetBlueprint()->GeneratedVariables.FindByPredicate([&UserWidgetProperty](const FBPVariableDescription& VariableDesc)
+			{
+				return VariableDesc.VarName == UserWidgetProperty.Name;
+			});
+
+			// Will always create viewmodel properties.
+			// Will never create properties for animation or other Self.Object
+			// Will create properties for widget when they are not already created.
+			if (UserWidgetProperty.CreationType == FCompilerUserWidgetProperty::ECreationType::CreateOnlyIfDoesntExist)
+			{
+				if (ExistingGeneratedVariablePtr != nullptr)
 				{
-					NewProperty->SetMetaData(FBlueprintMetadata::MD_Private, TEXT("true"));
+					WidgetBlueprintCompilerContext.MessageLog.Error(*FText::Format(LOCTEXT("VariableAlreadyExistsInBP", "There is already a variable named '{0}' in blueprint '{1}'."), UserWidgetProperty.DisplayName, FText::FromString(Context.GetWidgetBlueprint()->GetFriendlyName())).ToString());
+					bIsGatherGeneratedVariablesStepValid = false;
+					continue;
 				}
-				if (bIsInstanceExposed)
+				else if (!UserWidgetPropertyField.IsEmpty())
 				{
-					NewProperty->SetMetaData(FName("EditInline"), TEXT("true"));
+					// Viewmodel property cannot already exist. It will create issues with initialization and with View::SetViewModel.
+					const UClass* OwnerClass = Cast<UClass>(UserWidgetPropertyField.GetOwner());
+					WidgetBlueprintCompilerContext.MessageLog.Error(*FText::Format(LOCTEXT("PropertyAlreadyExistInParent", "There is already a property named '{0}' in scope '{1}'."), UserWidgetProperty.DisplayName, (OwnerClass ? OwnerClass->GetDisplayNameText() : FText::GetEmpty())).ToString());
+					bIsGatherGeneratedVariablesStepValid = false;
+					continue;
 				}
-#endif
 			}
-			return NewProperty;
-		};
 
+			if (ExistingGeneratedVariablePtr != nullptr && (ExistingGeneratedVariablePtr->VarType.PinCategory != UEdGraphSchema_K2::PC_Object || ExistingGeneratedVariablePtr->VarType.PinSubCategoryObject != UserWidgetProperty.AuthoritativeClass))
+			{
+				WidgetBlueprintCompilerContext.MessageLog.Error(*FText::Format(LOCTEXT("VariableExistsWithWrongType", "There is already a variable named '{0}' in blueprint '{1}' that does not have the expected type."), UserWidgetProperty.DisplayName, FText::FromString(Context.GetWidgetBlueprint()->GetFriendlyName())).ToString());
+				bIsGatherGeneratedVariablesStepValid = false;
+				continue;
+			}
+
+			const bool bCreateVariable = UserWidgetProperty.CreationType == FCompilerUserWidgetProperty::ECreationType::CreateOnlyIfDoesntExist
+				|| (UserWidgetProperty.CreationType == FCompilerUserWidgetProperty::ECreationType::CreateIfDoesntExist && UserWidgetPropertyField.IsEmpty());
+			if (bCreateVariable && ExistingGeneratedVariablePtr == nullptr)
+			{
+				AddObjectVariable(UserWidgetProperty);
+			}
+		}
+	}
+}
+
+
+void FMVVMViewBlueprintCompiler::CreateVariables(const FWidgetBlueprintCompilerContext::FCreateVariableContext& Context)
+{
+	if (!BlueprintView)
+	{
+		return;
+	}
+
+	if (!AreStepsValid())
+	{
+		return;
+	}
+
+	if (Context.GetCompileType() == EKismetCompileType::SkeletonOnly)
+	{
+		CreatePublicFunctionsDeclaration(Context);
+	}
+
+	// Create variable
+	{
 		for (FCompilerUserWidgetProperty& UserWidgetProperty : NeededUserWidgetProperties)
 		{
 			check(UserWidgetProperty.AuthoritativeClass);
@@ -555,18 +651,6 @@ void FMVVMViewBlueprintCompiler::CreateVariables(const FWidgetBlueprintCompilerC
 				}
 			}
 
-			// Will always create viewmodel properties.
-			// Will never create properties for animation or other Self.Object
-			// Will create properties for widget when they are not already created.
-			if (UserWidgetProperty.CreationType == FCompilerUserWidgetProperty::ECreationType::CreateOnlyIfDoesntExist && !UserWidgetPropertyField.IsEmpty())
-			{
-				// Viewmodel property cannot already exist. It will creates issue with initialization and with View::SetViewModel.
-				const UClass* OwnerClass = Cast<UClass>(UserWidgetPropertyField.GetOwner());
-				WidgetBlueprintCompilerContext.MessageLog.Error(*FText::Format(LOCTEXT("PropertyAlreadyExistInParent", "There is already a property named '{0}' in scope '{1}'."), UserWidgetProperty.DisplayName, (OwnerClass ? OwnerClass->GetDisplayNameText() : FText::GetEmpty())).ToString());
-				bIsCreateVariableStepValid = false;
-				continue;
-			}
-
 			if (!UserWidgetPropertyField.IsEmpty())
 			{
 				if (UserWidgetPropertyField.IsFunction())
@@ -593,26 +677,9 @@ void FMVVMViewBlueprintCompiler::CreateVariables(const FWidgetBlueprintCompilerC
 					bIsCreateVariableStepValid = false;
 					continue;
 				}
-
-				const bool bIsBindWidget = FWidgetBlueprintEditorUtils::IsBindWidgetProperty(ObjectProperty);
-				if (Context.GetGeneratedClass() != ObjectProperty->GetOwnerStruct() && !bIsBindWidget)
-				{
-					// Widget needs to be BindWidget to be reused as a property.
-					const UClass* OwnerClass = Cast<UClass>(ObjectProperty->GetOwnerStruct());
-					WidgetBlueprintCompilerContext.MessageLog.Error(*FText::Format(LOCTEXT("WidgetPropertyAlreadyExist", "There is already a property named '{0}' in scope '{1}' for the widget. Are you missing a BindWidget?."), UserWidgetProperty.DisplayName, (OwnerClass ? OwnerClass->GetDisplayNameText() : FText::GetEmpty())).ToString());
-					bIsCreateVariableStepValid = false;
-					continue;
-				}
 			}
 
-			// Can we reused the property or we need to create a new one.
-			bool bCreateVariable = UserWidgetProperty.CreationType == FCompilerUserWidgetProperty::ECreationType::CreateOnlyIfDoesntExist
-				|| (UserWidgetProperty.CreationType == FCompilerUserWidgetProperty::ECreationType::CreateIfDoesntExist && UserWidgetPropertyField.IsEmpty());
-			if (bCreateVariable)
-			{
-				UserWidgetProperty.Property = CreateObjectVariable(UserWidgetProperty);
-			}
-			else if (UserWidgetPropertyField.IsProperty())
+			if (UserWidgetPropertyField.IsProperty())
 			{
 				UserWidgetProperty.Property = UserWidgetPropertyField.GetProperty();
 			}
@@ -686,7 +753,7 @@ void FMVVMViewBlueprintCompiler::CreateVariables(const FWidgetBlueprintCompilerC
 }
 
 
-void FMVVMViewBlueprintCompiler::CreateWidgetMap(const FWidgetBlueprintCompilerContext::FCreateVariableContext& Context)
+void FMVVMViewBlueprintCompiler::CreateWidgetMap(const FWidgetBlueprintCompilerContext::FPopulateGeneratedVariablesContext& Context)
 {
 	// The widget tree is not created yet for SKEL class.
 	//Context.GetGeneratedClass()->GetWidgetTreeArchetype()
@@ -711,14 +778,8 @@ void FMVVMViewBlueprintCompiler::CreateWidgetMap(const FWidgetBlueprintCompilerC
 }
 
 
-void FMVVMViewBlueprintCompiler::CreateBindingList(const FWidgetBlueprintCompilerContext::FCreateVariableContext& Context)
+void FMVVMViewBlueprintCompiler::CreateBindingList(const FWidgetBlueprintCompilerContext::FPopulateGeneratedVariablesContext& Context)
 {
-	ensure(Context.GetCompileType() == EKismetCompileType::SkeletonOnly);
-	if (Context.GetCompileType() != EKismetCompileType::SkeletonOnly)
-	{
-		return;
-	}
-
 	ValidBindings.Reset();
 
 	// Build the list of bindings that we should compile.
@@ -728,7 +789,7 @@ void FMVVMViewBlueprintCompiler::CreateBindingList(const FWidgetBlueprintCompile
 		if (BindingPtr == nullptr)
 		{
 			WidgetBlueprintCompilerContext.MessageLog.Error(*FText::Format(LOCTEXT("BindingInvalidIndex", "Internal error: Tried to fetch binding for invalid index {0}."), Index).ToString());
-			bIsCreateVariableStepValid = false;
+			bIsGatherGeneratedVariablesStepValid = false;
 			continue;
 		}
 
@@ -742,7 +803,7 @@ void FMVVMViewBlueprintCompiler::CreateBindingList(const FWidgetBlueprintCompile
 		if (bHasConversionFunction && Binding.BindingType == EMVVMBindingMode::TwoWay)
 		{
 			AddMessageForBinding(Binding, LOCTEXT("TwoWayBindingsWithConversion", "Two-way bindings are not allowed to use conversion functions."), Compiler::EMessageType::Error, FMVVMBlueprintPinId());
-			bIsCreateVariableStepValid = false;
+			bIsGatherGeneratedVariablesStepValid = false;
 			continue;
 		}
 
@@ -768,14 +829,14 @@ void FMVVMViewBlueprintCompiler::CreateBindingList(const FWidgetBlueprintCompile
 }
 
 
-void FMVVMViewBlueprintCompiler::CreateEventList(const FWidgetBlueprintCompilerContext::FCreateVariableContext& Context)
+void FMVVMViewBlueprintCompiler::CreateEventList(const FWidgetBlueprintCompilerContext::FPopulateGeneratedVariablesContext& Context)
 {
 	for (UMVVMBlueprintViewEvent* EventPtr : BlueprintView->GetEvents())
 	{
 		if (EventPtr == nullptr)
 		{
 			WidgetBlueprintCompilerContext.MessageLog.Error(*LOCTEXT("EventInvalid", "Internal error: An event is invalid.").ToString());
-			bIsCreateVariableStepValid = false;
+			bIsGatherGeneratedVariablesStepValid = false;
 			continue;
 		}
 
@@ -787,7 +848,7 @@ void FMVVMViewBlueprintCompiler::CreateEventList(const FWidgetBlueprintCompilerC
 		if (!EventPtr->GetEventPath().HasPaths())
 		{
 			AddMessageForEvent(EventPtr, LOCTEXT("EventInvalidEventPath", "The event path is invalid."), Compiler::EMessageType::Error, FMVVMBlueprintPinId());
-			bIsCreateVariableStepValid = false;
+			bIsGatherGeneratedVariablesStepValid = false;
 			continue;
 		}
 
@@ -798,7 +859,7 @@ void FMVVMViewBlueprintCompiler::CreateEventList(const FWidgetBlueprintCompilerC
 	}
 }
 
-void FMVVMViewBlueprintCompiler::CreateConditionList(const FWidgetBlueprintCompilerContext::FCreateVariableContext& Context)
+void FMVVMViewBlueprintCompiler::CreateConditionList(const FWidgetBlueprintCompilerContext::FPopulateGeneratedVariablesContext& Context)
 {
 	bool bAllowConditionBinding = GetDefault<UMVVMDeveloperProjectSettings>()->bAllowConditionBinding;
 
@@ -807,7 +868,7 @@ void FMVVMViewBlueprintCompiler::CreateConditionList(const FWidgetBlueprintCompi
 		if (ConditionPtr == nullptr)
 		{
 			WidgetBlueprintCompilerContext.MessageLog.Error(*LOCTEXT("ConditionInvalid", "Internal error: A condition is invalid.").ToString());
-			bIsCreateVariableStepValid = false;
+			bIsGatherGeneratedVariablesStepValid = false;
 			continue;
 		}
 
@@ -819,7 +880,7 @@ void FMVVMViewBlueprintCompiler::CreateConditionList(const FWidgetBlueprintCompi
 		if (!ConditionPtr->GetConditionPath().HasPaths())
 		{
 			AddMessageForCondition(ConditionPtr, LOCTEXT("ConditionInvalidConditionPath", "The condition path is invalid."), Compiler::EMessageType::Error, FMVVMBlueprintPinId());
-			bIsCreateVariableStepValid = false;
+			bIsGatherGeneratedVariablesStepValid = false;
 			continue;
 		}
 
@@ -832,7 +893,7 @@ void FMVVMViewBlueprintCompiler::CreateConditionList(const FWidgetBlueprintCompi
 
 
 
-void FMVVMViewBlueprintCompiler::CreateExtensionList(const FWidgetBlueprintCompilerContext::FCreateVariableContext& Context)
+void FMVVMViewBlueprintCompiler::CreateExtensionList(const FWidgetBlueprintCompilerContext::FPopulateGeneratedVariablesContext& Context)
 {
 	for (int32 Index = BlueprintView->GetOuterUMVVMWidgetBlueprintExtension_View()->BlueprintExtensions.Num() - 1; Index >= 0; --Index)
 	{
@@ -851,14 +912,8 @@ void FMVVMViewBlueprintCompiler::CreateExtensionList(const FWidgetBlueprintCompi
 }
 
 
-void FMVVMViewBlueprintCompiler::CreateRequiredProperties(const FWidgetBlueprintCompilerContext::FCreateVariableContext& Context)
+void FMVVMViewBlueprintCompiler::CreateRequiredProperties(const FWidgetBlueprintCompilerContext::FPopulateGeneratedVariablesContext& Context)
 {
-	ensure(Context.GetCompileType() == EKismetCompileType::SkeletonOnly);
-	if (Context.GetCompileType() != EKismetCompileType::SkeletonOnly)
-	{
-		return;
-	}
-
 	NeededBindingSources.Reset();
 	NeededUserWidgetProperties.Reset();
 	ViewModelCreatorContexts.Reset();
@@ -870,14 +925,14 @@ void FMVVMViewBlueprintCompiler::CreateRequiredProperties(const FWidgetBlueprint
 		if (!ViewModelContext.GetViewModelId().IsValid())
 		{
 			AddMessageForViewModel(ViewModelContext, LOCTEXT("ViewmodelInvalidGuid", "GUID is invalid."), Compiler::EMessageType::Error);
-			bIsCreateVariableStepValid = false;
+			bIsGatherGeneratedVariablesStepValid = false;
 			continue;
 		}
 
 		if (ViewModelGuids.Contains(ViewModelContext.GetViewModelId()))
 		{
 			AddMessageForViewModel(ViewModelContext, LOCTEXT("ViewmodelAlreadyAdded", "Identical viewmodel has already been added."), Compiler::EMessageType::Error);
-			bIsCreateVariableStepValid = false;
+			bIsGatherGeneratedVariablesStepValid = false;
 			continue;
 		}
 
@@ -886,7 +941,7 @@ void FMVVMViewBlueprintCompiler::CreateRequiredProperties(const FWidgetBlueprint
 		if (ViewModelContext.GetViewModelClass() == nullptr || !ViewModelContext.IsValid())
 		{
 			AddMessageForViewModel(ViewModelContext, LOCTEXT("ViewmodelInvalidClass", "Invalid class."), Compiler::EMessageType::Error);
-			bIsCreateVariableStepValid = false;
+			bIsGatherGeneratedVariablesStepValid = false;
 			continue;
 		}
 
@@ -933,6 +988,7 @@ void FMVVMViewBlueprintCompiler::CreateRequiredProperties(const FWidgetBlueprint
 			SourceVariable.Name = PropertyName;
 			SourceVariable.DisplayName = ViewModelContext.GetDisplayName();
 			SourceVariable.CategoryName = TEXT("Viewmodel");
+			SourceVariable.Guid = ViewModelContext.GetViewModelId();
 			SourceVariable.BlueprintSetter = SetterFunctionName;
 			SourceVariable.CreationType = FCompilerUserWidgetProperty::ECreationType::CreateOnlyIfDoesntExist;
 			SourceVariable.bExposeOnSpawn = bCreateSetterFunction;
@@ -951,13 +1007,12 @@ void FMVVMViewBlueprintCompiler::CreateRequiredProperties(const FWidgetBlueprint
 
 	TSet<FName> WidgetSourcesCreated;
 	TSet<FName> WidgetUserPropertyCreated;
-	const FName DefaultWidgetCategory = Context.GetWidgetBlueprint()->GetFName();
 	bool bSelfBindingSourceCreated = NeededBindingSources.ContainsByPredicate([](const TSharedRef<FCompilerBindingSource>& Other)
 		{
 			return Other->Type == FCompilerBindingSource::EType::Self;
 		});
 
-	auto GenerateCompilerContext = [Self = this, DefaultWidgetCategory, Class = Context.GetGeneratedClass(), &ViewModelGuids, &WidgetSourcesCreated, &WidgetUserPropertyCreated, &bSelfBindingSourceCreated](bool bInCreateSource, const FMVVMBlueprintPropertyPath& PropertyPath) -> TValueOrError<void, FText>
+	auto GenerateCompilerContext = [Self = this, WidgetBP = Context.GetWidgetBlueprint(), &ViewModelGuids, &WidgetSourcesCreated, &WidgetUserPropertyCreated, &bSelfBindingSourceCreated](bool bInCreateSource, const FMVVMBlueprintPropertyPath& PropertyPath) -> TValueOrError<void, FText>
 	{
 		FName PropertyPathWidgetName = PropertyPath.GetWidgetName();
 		FGuid PropertyPathViewModelId = PropertyPath.GetViewModelId();
@@ -1047,7 +1102,8 @@ void FMVVMViewBlueprintCompiler::CreateRequiredProperties(const FWidgetBlueprint
 					FCompilerUserWidgetProperty& SourceVariable = Self->NeededUserWidgetProperties.AddDefaulted_GetRef();
 					SourceVariable.AuthoritativeClass = Widget->GetClass();
 					SourceVariable.Name = PropertyPathWidgetName;
-					SourceVariable.DisplayName = FText::FromString(Widget->GetDisplayLabel());
+					SourceVariable.Guid = WidgetBP->WidgetVariableNameToGuidMap.FindRef(Widget->GetFName());
+					SourceVariable.DisplayName = Widget->IsGeneratedName() ? FText::FromString(Widget->GetName()) : Widget->GetLabelText();
 					SourceVariable.CategoryName = TEXT("Widget");
 					SourceVariable.CreationType = FCompilerUserWidgetProperty::ECreationType::CreateIfDoesntExist;
 					SourceVariable.bExposeOnSpawn = false;
@@ -1093,7 +1149,7 @@ void FMVVMViewBlueprintCompiler::CreateRequiredProperties(const FWidgetBlueprint
 				if (SourceContextResult.HasError())
 				{
 					Self->AddMessageForBinding(Binding, SourceContextResult.StealError(), Compiler::EMessageType::Error, PinId);
-					Self->bIsCreateVariableStepValid = false;
+					Self->bIsGatherGeneratedVariablesStepValid = false;
 				}
 			};
 
@@ -1112,7 +1168,7 @@ void FMVVMViewBlueprintCompiler::CreateRequiredProperties(const FWidgetBlueprint
 							, Compiler::EMessageType::Error
 							, FMVVMBlueprintPinId()
 						);
-						bIsCreateVariableStepValid = false;
+						bIsGatherGeneratedVariablesStepValid = false;
 						continue;
 					}
 					RunGenerateCompilerSourceContext(true, Pin.GetPath(), Pin.GetId());
@@ -1143,7 +1199,7 @@ void FMVVMViewBlueprintCompiler::CreateRequiredProperties(const FWidgetBlueprint
 			if (SourceContextResult.HasError())
 			{
 				Self->AddMessageForEvent(EventPtr, SourceContextResult.StealError(), Compiler::EMessageType::Error, PinId);
-				Self->bIsCreateVariableStepValid = false;
+				Self->bIsGatherGeneratedVariablesStepValid = false;
 			}
 		};
 
@@ -1160,7 +1216,7 @@ void FMVVMViewBlueprintCompiler::CreateRequiredProperties(const FWidgetBlueprint
 						, Compiler::EMessageType::Error
 						, FMVVMBlueprintPinId()
 					);
-					bIsCreateVariableStepValid = false;
+					bIsGatherGeneratedVariablesStepValid = false;
 					continue;
 				}
 				RunGenerateCompilerSourceContext(true, Pin.GetPath(), Pin.GetId());
@@ -1178,7 +1234,7 @@ void FMVVMViewBlueprintCompiler::CreateRequiredProperties(const FWidgetBlueprint
 				if (SourceContextResult.HasError())
 				{
 					Self->AddMessageForCondition(ConditionPtr, SourceContextResult.StealError(), Compiler::EMessageType::Error, PinId);
-					Self->bIsCreateVariableStepValid = false;
+					Self->bIsGatherGeneratedVariablesStepValid = false;
 				}
 			};
 
@@ -1194,7 +1250,7 @@ void FMVVMViewBlueprintCompiler::CreateRequiredProperties(const FWidgetBlueprint
 						, Compiler::EMessageType::Error
 						, FMVVMBlueprintPinId()
 					);
-					bIsCreateVariableStepValid = false;
+					bIsGatherGeneratedVariablesStepValid = false;
 					continue;
 				}
 				RunGenerateCompilerSourceContext(true, Pin.GetPath(), Pin.GetId());
@@ -1220,6 +1276,7 @@ void FMVVMViewBlueprintCompiler::CreateRequiredProperties(const FWidgetBlueprint
 				CompilerUserWidgetProperty.bPrivate = Property.bPrivate;
 				CompilerUserWidgetProperty.bReadOnly = Property.bReadOnly;
 				CompilerUserWidgetProperty.CategoryName = Property.CategoryName;
+				CompilerUserWidgetProperty.Guid = Property.Guid;
 				CompilerUserWidgetProperty.CreationType = FCompilerUserWidgetProperty::ECreationType::CreateIfDoesntExist;
 			}
 
@@ -1234,13 +1291,13 @@ void FMVVMViewBlueprintCompiler::CreateRequiredProperties(const FWidgetBlueprint
 					if (SourceContextResult.HasError())
 					{
 						AddMessage(SourceContextResult.StealError(), Compiler::EMessageType::Error);
-						bIsCreateVariableStepValid = false;
+						bIsGatherGeneratedVariablesStepValid = false;
 					}
 				}
 				else
 				{
 					AddMessage(FText::Format(LOCTEXT("InvalidWidgetPropertyRequested", "Could not find the widget : {0} in the widget tree to make the property requested by viewmodel extension."), FText::FromName(WidgetProperty.WidgetName)), Compiler::EMessageType::Error);
-					bIsCreateVariableStepValid = false;
+					bIsGatherGeneratedVariablesStepValid = false;
 				}
 			}
 		}
@@ -1291,6 +1348,7 @@ void FMVVMViewBlueprintCompiler::CreatePublicFunctionsDeclaration(const FWidgetB
 				continue;
 			}
 
+			Setter.SetterGraph->ClearFlags(RF_Transactional);
 			BlueprintView->TemporaryGraph.Add(Setter.SetterGraph);
 			FunctionPermissionsToAdd.AddUnique(Setter.SetterGraph->GetFName());
 			GeneratedFunctions.AddUnique(Setter.SetterGraph->GetFName());
diff --git a/Engine/Plugins/Runtime/ModelViewViewModel/Source/ModelViewViewModelBlueprint/Private/MVVMViewBlueprintCompiler.h b/Engine/Plugins/Runtime/ModelViewViewModel/Source/ModelViewViewModelBlueprint/Private/MVVMViewBlueprintCompiler.h
--- a/Engine/Plugins/Runtime/ModelViewViewModel/Source/ModelViewViewModelBlueprint/Private/MVVMViewBlueprintCompiler.h
+++ b/Engine/Plugins/Runtime/ModelViewViewModel/Source/ModelViewViewModelBlueprint/Private/MVVMViewBlueprintCompiler.h
@@ -60,6 +60,9 @@ public:
 	void AddExtension(UWidgetBlueprintGeneratedClass* Class, UMVVMViewClass* ViewExtension);
 	void CleanOldData(UWidgetBlueprintGeneratedClass* ClassToClean, UObject* OldCDO);
 
+	/** Update the list of compiler generated variables to be created by the kismet compiler */
+	void GatherGeneratedVariables(const FWidgetBlueprintCompilerContext::FPopulateGeneratedVariablesContext& Context);
+
 	/** Generate variable and public function in the Skeleton class and the generated class */
 	void CreateVariables(const FWidgetBlueprintCompilerContext::FCreateVariableContext& Context);
 	/** Generate function that are hidden from the user (not on the Skeleton class). */
@@ -78,16 +81,18 @@ public:
 private:
 	bool AreStepsValid() const
 	{
-		return bIsCreateVariableStepValid && bIsCreateFunctionsStepValid && bIsPreCompileStepValid && bIsCompileStepValid;
+		return bIsGatherGeneratedVariablesStepValid && bIsCreateVariableStepValid && bIsCreateFunctionsStepValid && bIsPreCompileStepValid && bIsCompileStepValid;
 	}
 
+	// GatherGeneratedVariables
+	void CreateWidgetMap(const FWidgetBlueprintCompilerContext::FPopulateGeneratedVariablesContext& Context);
+	void CreateBindingList(const FWidgetBlueprintCompilerContext::FPopulateGeneratedVariablesContext& Context);
+	void CreateEventList(const FWidgetBlueprintCompilerContext::FPopulateGeneratedVariablesContext& Context);
+	void CreateConditionList(const FWidgetBlueprintCompilerContext::FPopulateGeneratedVariablesContext& Context);
+	void CreateExtensionList(const FWidgetBlueprintCompilerContext::FPopulateGeneratedVariablesContext& Context);
+	void CreateRequiredProperties(const FWidgetBlueprintCompilerContext::FPopulateGeneratedVariablesContext& Context);
+
 	// CreateVariables
-	void CreateWidgetMap(const FWidgetBlueprintCompilerContext::FCreateVariableContext& Context);
-	void CreateBindingList(const FWidgetBlueprintCompilerContext::FCreateVariableContext& Context);
-	void CreateEventList(const FWidgetBlueprintCompilerContext::FCreateVariableContext& Context);
-	void CreateConditionList(const FWidgetBlueprintCompilerContext::FCreateVariableContext& Context);
-	void CreateExtensionList(const FWidgetBlueprintCompilerContext::FCreateVariableContext& Context);
-	void CreateRequiredProperties(const FWidgetBlueprintCompilerContext::FCreateVariableContext& Context);
 	void CreatePublicFunctionsDeclaration(const FWidgetBlueprintCompilerContext::FCreateVariableContext& Context);
 
 	// CreateFuntions
@@ -382,6 +387,7 @@ private:
 	FWidgetBlueprintCompilerContext& WidgetBlueprintCompilerContext;
 	TObjectPtr<UMVVMBlueprintView> BlueprintView = nullptr;
 	FCompiledBindingLibraryCompiler BindingLibraryCompiler;
+	bool bIsGatherGeneratedVariablesStepValid = true;
 	bool bIsCreateVariableStepValid = true;
 	bool bIsCreateFunctionsStepValid = true;
 	bool bIsPreCompileStepValid = true;
diff --git a/Engine/Plugins/Runtime/ModelViewViewModel/Source/ModelViewViewModelBlueprint/Private/MVVMWidgetBlueprintExtension_View.cpp b/Engine/Plugins/Runtime/ModelViewViewModel/Source/ModelViewViewModelBlueprint/Private/MVVMWidgetBlueprintExtension_View.cpp
--- a/Engine/Plugins/Runtime/ModelViewViewModel/Source/ModelViewViewModelBlueprint/Private/MVVMWidgetBlueprintExtension_View.cpp
+++ b/Engine/Plugins/Runtime/ModelViewViewModel/Source/ModelViewViewModelBlueprint/Private/MVVMWidgetBlueprintExtension_View.cpp
@@ -304,6 +304,14 @@ void UMVVMWidgetBlueprintExtension_View::HandleCleanAndSanitizeClass(UWidgetBlue
 }
 
 
+void UMVVMWidgetBlueprintExtension_View::HandlePopulateGeneratedVariables(const FWidgetBlueprintCompilerContext::FPopulateGeneratedVariablesContext& Context)
+{
+	Super::HandlePopulateGeneratedVariables(Context);
+
+	CurrentCompilerContext->GatherGeneratedVariables(Context);
+}
+
+
 void UMVVMWidgetBlueprintExtension_View::HandleCreateClassVariablesFromBlueprint(const FWidgetBlueprintCompilerContext::FCreateVariableContext& Context)
 {
 	Super::HandleCreateClassVariablesFromBlueprint(Context);
diff --git a/Engine/Plugins/Runtime/ModelViewViewModel/Source/ModelViewViewModelBlueprint/Public/MVVMBlueprintViewCompilerInterface.h b/Engine/Plugins/Runtime/ModelViewViewModel/Source/ModelViewViewModelBlueprint/Public/MVVMBlueprintViewCompilerInterface.h
--- a/Engine/Plugins/Runtime/ModelViewViewModel/Source/ModelViewViewModelBlueprint/Public/MVVMBlueprintViewCompilerInterface.h
+++ b/Engine/Plugins/Runtime/ModelViewViewModel/Source/ModelViewViewModelBlueprint/Public/MVVMBlueprintViewCompilerInterface.h
@@ -78,6 +78,7 @@ struct FBlueprintViewUserWidgetProperty
 	FName Name;
 	FText DisplayName;
 	FString CategoryName;
+	FGuid Guid;
 	bool bExposeOnSpawn = false;
 	bool bPrivate = false;
 	bool bReadOnly = false;
diff --git a/Engine/Plugins/Runtime/ModelViewViewModel/Source/ModelViewViewModelBlueprint/Public/MVVMWidgetBlueprintExtension_View.h b/Engine/Plugins/Runtime/ModelViewViewModel/Source/ModelViewViewModelBlueprint/Public/MVVMWidgetBlueprintExtension_View.h
--- a/Engine/Plugins/Runtime/ModelViewViewModel/Source/ModelViewViewModelBlueprint/Public/MVVMWidgetBlueprintExtension_View.h
+++ b/Engine/Plugins/Runtime/ModelViewViewModel/Source/ModelViewViewModelBlueprint/Public/MVVMWidgetBlueprintExtension_View.h
@@ -88,6 +88,7 @@ public:
 	virtual void HandlePreloadObjectsForCompilation(UBlueprint* OwningBlueprint) override;
 	virtual void HandleBeginCompilation(FWidgetBlueprintCompilerContext& InCreationContext) override;
 	virtual void HandleCleanAndSanitizeClass(UWidgetBlueprintGeneratedClass* ClassToClean, UObject* OldCDO) override;
+	virtual void HandlePopulateGeneratedVariables(const FWidgetBlueprintCompilerContext::FPopulateGeneratedVariablesContext& Context) override;
 	virtual void HandleCreateClassVariablesFromBlueprint(const FWidgetBlueprintCompilerContext::FCreateVariableContext& Context) override;
 	virtual void HandleCreateFunctionList(const FWidgetBlueprintCompilerContext::FCreateFunctionContext& Context) override;
 	virtual void HandleFinishCompilingClass(UWidgetBlueprintGeneratedClass* Class) override;
diff --git a/Engine/Plugins/Runtime/ModelViewViewModel/Source/ModelViewViewModelEditor/Private/MVVMEditorSubsystem.cpp b/Engine/Plugins/Runtime/ModelViewViewModel/Source/ModelViewViewModelEditor/Private/MVVMEditorSubsystem.cpp
--- a/Engine/Plugins/Runtime/ModelViewViewModel/Source/ModelViewViewModelEditor/Private/MVVMEditorSubsystem.cpp
+++ b/Engine/Plugins/Runtime/ModelViewViewModel/Source/ModelViewViewModelEditor/Private/MVVMEditorSubsystem.cpp
@@ -190,23 +190,26 @@ namespace UE::MVVM::Private
 
 			if (!BlueprintReferencers.IsEmpty())
 			{
-				static const FTextFormat LoadingWarningFormat = LOCTEXT("LoadingBlueprintReferencersWarningFormat", "Renaming this view model requires loading {0} {0}|plural(one=asset,other=assets).\nDo you want to continue?");
-				const EAppReturnType::Type DialogResult = FMessageDialog::Open(EAppMsgType::YesNo, FText::Format(LoadingWarningFormat, BlueprintReferencers.Num()));
-				if (DialogResult == EAppReturnType::No)
+				static const FTextFormat LoadingWarningFormat = LOCTEXT("LoadingBlueprintReferencersWarningFormat", "There {0}|plural(one=is,other=are) {0} unloaded {0}|plural(one=asset,other=assets) that might reference this view model.\nDo you want to load the assets to automatically update them now?");
+				const EAppReturnType::Type DialogResult = FMessageDialog::Open(EAppMsgType::YesNoCancel, FText::Format(LoadingWarningFormat, BlueprintReferencers.Num()));
+				if (DialogResult == EAppReturnType::Cancel)
 				{
 					return false;
 				}
-				
-				FScopedSlowTask LoadReferencersSlowTask = FScopedSlowTask(BlueprintReferencers.Num(), LOCTEXT("LoadingRenamedViewModelBindings", "Loading Dependent Blueprints"));
-				
-				constexpr float TimeBeforeDialog = 1.f;
-				LoadReferencersSlowTask.MakeDialogDelayed(TimeBeforeDialog);
-				for (const FAssetData& Asset : BlueprintReferencers)
+
+				if (DialogResult == EAppReturnType::Yes)
 				{
-					LoadReferencersSlowTask.EnterProgressFrame(1.0f, FText::Format(LOCTEXT("LoadingBlueprintReferencerFormat", "Loading {0}"), FText::FromName(Asset.AssetName)));
+					FScopedSlowTask LoadReferencersSlowTask = FScopedSlowTask(BlueprintReferencers.Num(), LOCTEXT("LoadingRenamedViewModelBindings", "Loading Dependent Blueprints"));
+				
+					constexpr float TimeBeforeDialog = 1.f;
+					LoadReferencersSlowTask.MakeDialogDelayed(TimeBeforeDialog);
+					for (const FAssetData& Asset : BlueprintReferencers)
+					{
+						LoadReferencersSlowTask.EnterProgressFrame(1.0f, FText::Format(LOCTEXT("LoadingBlueprintReferencerFormat", "Loading {0}"), FText::FromName(Asset.AssetName)));
 					
-					constexpr bool bLoad = true;
-					Asset.FastGetAsset(bLoad);
+						constexpr bool bLoad = true;
+						Asset.FastGetAsset(bLoad);
+					}
 				}
 			}
 		}
@@ -373,6 +376,7 @@ bool UMVVMEditorSubsystem::RenameViewModel(UWidgetBlueprint* WidgetBlueprint, FN
 	if (ViewModelContext && ViewModelContext->bCanRename)
 	{
 		// Load any dependent BPs so that they can update any bindings that include the renamed view model
+		// This is required because FMVVMBlueprintFieldPath::BindingReference did not always serialize the GUID for the generated view model properties, so it could not resolve the field if it was renamed
 		// Do this before starting the transaction so that any changes during load are excluded
 		if (UE::MVVM::Private::LoadBlueprintReferencersForViewModelRename(WidgetBlueprint))
 		{
