MVVM: Make field icon colors consistent with those set in the graph editor settings

diff --git a/Engine/Plugins/Runtime/ModelViewViewModel/Source/ModelViewViewModelEditor/Private/Widgets/SMVVMViewModelBindingListWidget.cpp b/Engine/Plugins/Runtime/ModelViewViewModel/Source/ModelViewViewModelEditor/Private/Widgets/SMVVMViewModelBindingListWidget.cpp
--- a/Engine/Plugins/Runtime/ModelViewViewModel/Source/ModelViewViewModelEditor/Private/Widgets/SMVVMViewModelBindingListWidget.cpp
+++ b/Engine/Plugins/Runtime/ModelViewViewModel/Source/ModelViewViewModelEditor/Private/Widgets/SMVVMViewModelBindingListWidget.cpp
@@ -6,6 +6,7 @@
 #include "Bindings/MVVMFieldPathHelper.h"
 #include "Blueprint/WidgetTree.h"
 #include "BlueprintEditorSettings.h"
+#include "GraphEditorSettings.h"
 #include "MVVMBlueprintViewModelContext.h"
 #include "MVVMDeveloperProjectSettings.h"
 #include "MVVMSubsystem.h"
@@ -371,11 +372,31 @@ void SSourceBindingList::Construct(const FArguments& InArgs, const UWidgetBluepr
 	bIsBindingToEvent = InArgs._IsBindingToEvent;
 	OnDoubleClicked = InArgs._OnDoubleClicked;
 
+	const UGraphEditorSettings* GraphEditorSettings = GetDefault<UGraphEditorSettings>();
+	UE::PropertyViewer::FFieldColorSettings OverrideIconColorSettings;
+	OverrideIconColorSettings.DefaultTypeColor = GraphEditorSettings->DefaultPinTypeColor;
+	OverrideIconColorSettings.BooleanTypeColor = GraphEditorSettings->BooleanPinTypeColor;
+	OverrideIconColorSettings.ClassTypeColor = GraphEditorSettings->ClassPinTypeColor;
+	OverrideIconColorSettings.EnumTypeColor = GraphEditorSettings->BytePinTypeColor;
+	OverrideIconColorSettings.IntTypeColor = GraphEditorSettings->IntPinTypeColor;
+	OverrideIconColorSettings.FloatTypeColor = GraphEditorSettings->FloatPinTypeColor;
+	OverrideIconColorSettings.NameTypeColor = GraphEditorSettings->NamePinTypeColor;
+	OverrideIconColorSettings.DelegateTypeColor = GraphEditorSettings->DelegatePinTypeColor;
+	OverrideIconColorSettings.ObjectTypeColor = GraphEditorSettings->ObjectPinTypeColor;
+	OverrideIconColorSettings.InterfaceTypeColor = GraphEditorSettings->InterfacePinTypeColor;
+	OverrideIconColorSettings.StringTypeColor = GraphEditorSettings->StringPinTypeColor;
+	OverrideIconColorSettings.TextTypeColor = GraphEditorSettings->TextPinTypeColor;
+	OverrideIconColorSettings.DefaultStructTypeColor = GraphEditorSettings->StructPinTypeColor;
+	OverrideIconColorSettings.StructColors.Add(TBaseStructure<FVector>::Get()->GetStructPathName().ToString(), GraphEditorSettings->VectorPinTypeColor);
+	OverrideIconColorSettings.StructColors.Add(TBaseStructure<FRotator>::Get()->GetStructPathName().ToString(), GraphEditorSettings->RotatorPinTypeColor);
+	OverrideIconColorSettings.StructColors.Add(TBaseStructure<FTransform>::Get()->GetStructPathName().ToString(), GraphEditorSettings->TransformPinTypeColor);
+
 	PropertyViewer = SNew(SPropertyViewer)
 		.FieldIterator(FieldIterator.Get())
 		.FieldExpander(FieldExpander.Get())
 		.PropertyVisibility(SPropertyViewer::EPropertyVisibility::Hidden)
 		.bShowFieldIcon(true)
+		.OverrideIconColorSettings(OverrideIconColorSettings)
 		.bSanitizeName(true)
 		.SelectionMode(InArgs._EnableSelection ? ESelectionMode::Single : ESelectionMode::None)
 		.bShowSearchBox(InArgs._ShowSearchBox)
diff --git a/Engine/Source/Runtime/AdvancedWidgets/Private/Widgets/PropertyViewer/SFieldName.cpp b/Engine/Source/Runtime/AdvancedWidgets/Private/Widgets/PropertyViewer/SFieldName.cpp
--- a/Engine/Source/Runtime/AdvancedWidgets/Private/Widgets/PropertyViewer/SFieldName.cpp
+++ b/Engine/Source/Runtime/AdvancedWidgets/Private/Widgets/PropertyViewer/SFieldName.cpp
@@ -22,7 +22,8 @@ void SFieldName::Construct(const FArguments& InArgs, const UClass* Class)
 	TSharedPtr<SWidget> Icon;
 	if (InArgs._bShowIcon)
 	{
-		Icon = SNew(SFieldIcon, Class);
+		Icon = SNew(SFieldIcon, Class)
+			.OverrideColorSettings(InArgs._OverrideIconColorSettings);
 	}
 
 	FText DisplayName;
@@ -51,7 +52,8 @@ void SFieldName::Construct(const FArguments& InArgs, const UScriptStruct* Struct
 	TSharedPtr<SWidget> Icon;
 	if (InArgs._bShowIcon)
 	{
-		Icon = SNew(SFieldIcon, Struct);
+		Icon = SNew(SFieldIcon, Struct)
+			.OverrideColorSettings(InArgs._OverrideIconColorSettings);
 	}
 
 	FText DisplayName;
@@ -80,7 +82,8 @@ void SFieldName::Construct(const FArguments& InArgs, const FProperty* Property)
 	TSharedPtr<SWidget> Icon;
 	if (InArgs._bShowIcon)
 	{
-		Icon = SNew(SFieldIcon, Property);
+		Icon = SNew(SFieldIcon, Property)
+			.OverrideColorSettings(InArgs._OverrideIconColorSettings);
 	}
 	
 	FText DisplayName;
@@ -109,7 +112,8 @@ void SFieldName::Construct(const FArguments& InArgs, const UFunction* Function)
 	TSharedPtr<SWidget> Icon;
 	if (InArgs._bShowIcon)
 	{
-		Icon = SNew(SFieldIcon, Function);
+		Icon = SNew(SFieldIcon, Function)
+			.OverrideColorSettings(InArgs._OverrideIconColorSettings);
 	}
 
 	FText DisplayName;
diff --git a/Engine/Source/Runtime/AdvancedWidgets/Private/Widgets/PropertyViewer/SPropertyViewerImpl.cpp b/Engine/Source/Runtime/AdvancedWidgets/Private/Widgets/PropertyViewer/SPropertyViewerImpl.cpp
--- a/Engine/Source/Runtime/AdvancedWidgets/Private/Widgets/PropertyViewer/SPropertyViewerImpl.cpp
+++ b/Engine/Source/Runtime/AdvancedWidgets/Private/Widgets/PropertyViewer/SPropertyViewerImpl.cpp
@@ -518,6 +518,8 @@ TSharedRef<SWidget> FPropertyViewerImpl::Construct(const SPropertyViewer::FArgum
 	FilterHandler->SetRootItems(&TreeSource, &FilteredTreeSource);
 	FilterHandler->SetGetChildrenDelegate(FTreeFilter::FOnGetChildren::CreateRaw(this, &FPropertyViewerImpl::HandleGetChildren));
 
+	OverrideIconColorSettings = InArgs._OverrideIconColorSettings;
+
 	TSharedPtr<SHorizontalBox> SearchBox;
 	if (InArgs._bShowSearchBox || InArgs._SearchBoxPreSlot.Widget != SNullWidget::NullWidget || InArgs._SearchBoxPostSlot.Widget != SNullWidget::NullWidget)
 	{
@@ -969,14 +971,16 @@ TSharedRef<ITableRow> FPropertyViewerImpl::HandleGenerateRow(TSharedPtr<FTreeNod
 				ItemWidget = SNew(SFieldName, Class)
 					.bShowIcon(true)
 					.bSanitizeName(bSanitizeName)
-					.OverrideDisplayName(Item->GetOverrideDisplayName());
+					.OverrideDisplayName(Item->GetOverrideDisplayName())
+					.OverrideIconColorSettings(OverrideIconColorSettings);
 			}
 			if (const UScriptStruct* Struct = Cast<const UScriptStruct>(ContainerPin->GetStruct()))
 			{
 				ItemWidget = SNew(SFieldName, Struct)
 					.bShowIcon(true)
 					.bSanitizeName(bSanitizeName)
-					.OverrideDisplayName(Item->GetOverrideDisplayName());
+					.OverrideDisplayName(Item->GetOverrideDisplayName())
+					.OverrideIconColorSettings(OverrideIconColorSettings);
 			}
 		}
 	}
@@ -992,6 +996,7 @@ TSharedRef<ITableRow> FPropertyViewerImpl::HandleGenerateRow(TSharedPtr<FTreeNod
 						.bShowIcon(bShowFieldIcon)
 						.bSanitizeName(bSanitizeName)
 						.OverrideDisplayName(Item->GetOverrideDisplayName())
+						.OverrideIconColorSettings(OverrideIconColorSettings)
 						.HighlightText(HighlightText);
 					Item->PropertyWidget = FieldName;
 					ItemWidget = FieldName;
@@ -1002,6 +1007,7 @@ TSharedRef<ITableRow> FPropertyViewerImpl::HandleGenerateRow(TSharedPtr<FTreeNod
 						.bShowIcon(bShowFieldIcon)
 						.bSanitizeName(bSanitizeName)
 						.OverrideDisplayName(Item->GetOverrideDisplayName())
+						.OverrideIconColorSettings(OverrideIconColorSettings)
 						.HighlightText(HighlightText);
 					Item->PropertyWidget = FieldName;
 					ItemWidget = FieldName;
diff --git a/Engine/Source/Runtime/AdvancedWidgets/Private/Widgets/PropertyViewer/SPropertyViewerImpl.h b/Engine/Source/Runtime/AdvancedWidgets/Private/Widgets/PropertyViewer/SPropertyViewerImpl.h
--- a/Engine/Source/Runtime/AdvancedWidgets/Private/Widgets/PropertyViewer/SPropertyViewerImpl.h
+++ b/Engine/Source/Runtime/AdvancedWidgets/Private/Widgets/PropertyViewer/SPropertyViewerImpl.h
@@ -2,6 +2,7 @@
 
 #pragma once
 
+#include "Framework/PropertyViewer/FieldIconFinder.h"
 #include "Framework/PropertyViewer/PropertyPath.h"
 #include "Widgets/PropertyViewer/SPropertyViewer.h"
 #include "Types/SlateEnums.h"
@@ -224,6 +225,7 @@ private:
 	SPropertyViewer::FOnDragDetected OnDragDetected;
 	SPropertyViewer::FOnGenerateContainer OnGenerateContainer;
 	SPropertyViewer::EPropertyVisibility PropertyVisibility;
+	TOptional<const FFieldColorSettings> OverrideIconColorSettings;
 	bool bSanitizeName = false;
 	bool bShowFieldIcon = false;
 	bool bUseRows = false;
diff --git a/Engine/Source/Runtime/AdvancedWidgets/Public/Widgets/PropertyViewer/SFieldName.h b/Engine/Source/Runtime/AdvancedWidgets/Public/Widgets/PropertyViewer/SFieldName.h
--- a/Engine/Source/Runtime/AdvancedWidgets/Public/Widgets/PropertyViewer/SFieldName.h
+++ b/Engine/Source/Runtime/AdvancedWidgets/Public/Widgets/PropertyViewer/SFieldName.h
@@ -3,6 +3,7 @@
 #pragma once
 
 #include "Widgets/SCompoundWidget.h"
+#include "Framework/PropertyViewer/FieldIconFinder.h"
 #include "UObject/Field.h"
 
 class IToolTip;
@@ -23,6 +24,8 @@ public:
 		SLATE_ARGUMENT_DEFAULT(bool, bSanitizeName) = false;
 		/** Override the DisplayName of the container name. */
 		SLATE_ARGUMENT(TOptional<FText>, OverrideDisplayName);
+		/** Override the color settings of the icon. */
+		SLATE_ARGUMENT(TOptional<const FFieldColorSettings>, OverrideIconColorSettings);
 		/** The current highlighted text. */
 		SLATE_ATTRIBUTE(FText, HighlightText)
 	SLATE_END_ARGS()
diff --git a/Engine/Source/Runtime/AdvancedWidgets/Public/Widgets/PropertyViewer/SPropertyViewer.h b/Engine/Source/Runtime/AdvancedWidgets/Public/Widgets/PropertyViewer/SPropertyViewer.h
--- a/Engine/Source/Runtime/AdvancedWidgets/Public/Widgets/PropertyViewer/SPropertyViewer.h
+++ b/Engine/Source/Runtime/AdvancedWidgets/Public/Widgets/PropertyViewer/SPropertyViewer.h
@@ -2,6 +2,7 @@
 
 #pragma once
 
+#include "Framework/PropertyViewer/FieldIconFinder.h"
 #include "Framework/Views/ITypedTableView.h"
 #include "Templates/TypeHash.h"
 #include "Widgets/SCompoundWidget.h"
@@ -94,6 +95,8 @@ public:
 		SLATE_ARGUMENT_DEFAULT(IFieldExpander*, FieldExpander) = nullptr;
 		/** Callback when a property is modified. */
 		SLATE_ARGUMENT_DEFAULT(INotifyHook*, NotifyHook) = nullptr;
+		/** Override the color settings of the icons. */
+		SLATE_ARGUMENT(TOptional<const FFieldColorSettings>, OverrideIconColorSettings);
 		/** Slot for additional widget to go before the search box. */
 		SLATE_NAMED_SLOT(FArguments, SearchBoxPreSlot)
 		/** Slot for additional widget to go after the search box. */
