[MVVM] Protect against check/crash during blurprint compile in MVVM system

#tests Ran into this while converting properties to functions, which broke a lot of connections, and running asset validation

diff --git a/Engine/Plugins/Runtime/ModelViewViewModel/Source/ModelViewViewModelBlueprint/Private/Bindings/MVVMConversionFunctionHelper.cpp b/Engine/Plugins/Runtime/ModelViewViewModel/Source/ModelViewViewModelBlueprint/Private/Bindings/MVVMConversionFunctionHelper.cpp
--- a/Engine/Plugins/Runtime/ModelViewViewModel/Source/ModelViewViewModelBlueprint/Private/Bindings/MVVMConversionFunctionHelper.cpp
+++ b/Engine/Plugins/Runtime/ModelViewViewModel/Source/ModelViewViewModelBlueprint/Private/Bindings/MVVMConversionFunctionHelper.cpp
@@ -593,17 +593,23 @@ namespace NamedNodes
 					return MakeError();
 				}
 
-				const UFunction* Function = Field.GetFunction();
-
-				FGraphNodeCreator<UK2Node_CallFunction> CallFunctionCreator(*FunctionGraph);
-				UK2Node_CallFunction* FunctionNode = CallFunctionCreator.CreateNode();
-				FunctionNode->SetFromFunction(Function);
-				FunctionNode->AllocateDefaultPins();
-				CallFunctionCreator.Finalize();
-
-				NewNode = FunctionNode;
-				NewPreviousProperty = UE::MVVM::BindingHelper::GetReturnProperty(Function);
-				NewPreviousDataPin = FindNewOutputPin(NewNode);
+				if (const UFunction* Function = Field.GetFunction())
+				{
+					FGraphNodeCreator<UK2Node_CallFunction> CallFunctionCreator(*FunctionGraph);
+					UK2Node_CallFunction* FunctionNode = CallFunctionCreator.CreateNode();
+					FunctionNode->SetFromFunction(Function);
+					FunctionNode->AllocateDefaultPins();
+					CallFunctionCreator.Finalize();
+
+					NewNode = FunctionNode;
+					NewPreviousProperty = UE::MVVM::BindingHelper::GetReturnProperty(Function);
+					NewPreviousDataPin = FindNewOutputPin(NewNode);
+				}
+				else
+				{
+					ensure(false);
+					return MakeError();
+				}
 			}
 			else
 			{
