More log to debug

#tests Tested the output by setting the next statement inside of the if check
#rnx

diff --git a/Engine/Plugins/Runtime/ModelViewViewModel/Source/ModelViewViewModel/Private/Bindings/MVVMCompiledBindingLibrary.cpp b/Engine/Plugins/Runtime/ModelViewViewModel/Source/ModelViewViewModel/Private/Bindings/MVVMCompiledBindingLibrary.cpp
--- a/Engine/Plugins/Runtime/ModelViewViewModel/Source/ModelViewViewModel/Private/Bindings/MVVMCompiledBindingLibrary.cpp
+++ b/Engine/Plugins/Runtime/ModelViewViewModel/Source/ModelViewViewModel/Private/Bindings/MVVMCompiledBindingLibrary.cpp
@@ -362,7 +362,16 @@ TValueOrError<UE::MVVM::FFieldContext, void> FMVVMCompiledBindingLibrary::Evalua
 		FMVVMCompiledLoadedPropertyOrFunctionIndex PathIndex = FieldPaths[Index];
 		if (PathIndex.bIsProperty)
 		{
-			check(LoadedProperties.IsValidIndex(PathIndex.Index));
+			if (!LoadedProperties.IsValidIndex(PathIndex.Index))
+			{
+				FString LoadedPropertiesString;
+				for (int32 i = 0; i < LoadedProperties.Num(); ++i)
+				{
+					LoadedPropertiesString.Append(FString::Format(TEXT("{0}, "), { LoadedProperties[i]->GetNameCPP() }));
+				}
+				FString DebugLoadedProperties = FString::Format(TEXT("{0} on {1}. LoadedProperties : {2}"), { PathIndex.Index, LoadedProperties.Num(), LoadedPropertiesString });
+				checkf(false, TEXT("LoadedProperties.IsValidIndex(PathIndex.Index) failed. '%s'"), *DebugLoadedProperties);
+			}
 			const FProperty* Property = LoadedProperties[PathIndex.Index];
 			if (!Property)
 			{
