Fix issue with cleaning compiled graphs on undo/redo

review was done here in another branch : https://p4-swarm.epicgames.net/reviews/38180780?v=0,1

#tests Issue reproable with undo/redo of 2 AddViewModel. Followed the graph under the blueprint as we modify and clean up the generated graphs.
#rnx

diff --git a/Engine/Plugins/Runtime/ModelViewViewModel/Source/ModelViewViewModelBlueprint/Private/MVVMViewBlueprintCompiler.cpp b/Engine/Plugins/Runtime/ModelViewViewModel/Source/ModelViewViewModelBlueprint/Private/MVVMViewBlueprintCompiler.cpp
--- a/Engine/Plugins/Runtime/ModelViewViewModel/Source/ModelViewViewModelBlueprint/Private/MVVMViewBlueprintCompiler.cpp
+++ b/Engine/Plugins/Runtime/ModelViewViewModel/Source/ModelViewViewModelBlueprint/Private/MVVMViewBlueprintCompiler.cpp
@@ -1308,16 +1308,25 @@ void FMVVMViewBlueprintCompiler::CreateRequiredProperties(const FWidgetBlueprint
 void FMVVMViewBlueprintCompiler::CreatePublicFunctionsDeclaration(const FWidgetBlueprintCompilerContext::FCreateVariableContext& Context)
 {
 	// Clean all previous intermediate function graph. It should stay alive. The graph lives on the Blueprint not on the class and it's used to generate the UFunction.
+	for (UEdGraph* OldGraph : BlueprintView->TemporaryGraph)
 	{
-		for (UEdGraph* OldGraph : BlueprintView->TemporaryGraph)
+		if (OldGraph)
 		{
-			if (OldGraph)
-			{
-				RenameObjectToTransientPackage(OldGraph);
-			}
+			RenameObjectToTransientPackage(OldGraph);
+		}
+	}
+	// During deserialization for undo/redo, the existing graphs can get renamed to leave room for the serialized objects with the same name. When this happens, we basically don't have
+	// any pointers left on the graphs and nobody can clean them anymore.
+	for (FName OldGraphName : BlueprintView->TemporaryGraphNames)
+	{
+		UEdGraph* OldGraph = static_cast<UEdGraph*>(StaticFindObjectFastSafe(UEdGraph::StaticClass(), WidgetBlueprintCompilerContext.Blueprint, OldGraphName));
+		if (OldGraph)
+		{
+			RenameObjectToTransientPackage(OldGraph);
 		}
-		BlueprintView->TemporaryGraph.Reset();
 	}
+	BlueprintView->TemporaryGraph.Reset();
+	BlueprintView->TemporaryGraphNames.Reset();
 
 	if (GetDefault<UMVVMDeveloperProjectSettings>()->bAllowGeneratedViewModelSetter)
 	{
@@ -1340,7 +1349,7 @@ void FMVVMViewBlueprintCompiler::CreatePublicFunctionsDeclaration(const FWidgetB
 				);
 
 				if (Setter.SetterGraph != nullptr)
-				{ 
+				{
 					RenameObjectToTransientPackage(Setter.SetterGraph);
 					Setter.SetterGraph = nullptr;
 				}
@@ -1350,15 +1359,16 @@ void FMVVMViewBlueprintCompiler::CreatePublicFunctionsDeclaration(const FWidgetB
 
 			Setter.SetterGraph->ClearFlags(RF_Transactional);
 			BlueprintView->TemporaryGraph.Add(Setter.SetterGraph);
-			FunctionPermissionsToAdd.AddUnique(Setter.SetterGraph->GetFName());
-			GeneratedFunctions.AddUnique(Setter.SetterGraph->GetFName());
+			FName GraphName = Setter.SetterGraph->GetFName();
+			BlueprintView->TemporaryGraphNames.Add(GraphName);
+			FunctionPermissionsToAdd.AddUnique(GraphName);
+			GeneratedFunctions.AddUnique(GraphName);
 
 			UE::MVVM::FunctionGraphHelper::AddFunctionArgument(Setter.SetterGraph, const_cast<UClass*>(Setter.Class), "Viewmodel");
 		}
 	}
 }
 
-
 void FMVVMViewBlueprintCompiler::CreateFunctions(const FWidgetBlueprintCompilerContext::FCreateFunctionContext& Context)
 {
 	if (!AreStepsValid())
diff --git a/Engine/Plugins/Runtime/ModelViewViewModel/Source/ModelViewViewModelBlueprint/Public/MVVMBlueprintView.h b/Engine/Plugins/Runtime/ModelViewViewModel/Source/ModelViewViewModelBlueprint/Public/MVVMBlueprintView.h
--- a/Engine/Plugins/Runtime/ModelViewViewModel/Source/ModelViewViewModelBlueprint/Public/MVVMBlueprintView.h
+++ b/Engine/Plugins/Runtime/ModelViewViewModel/Source/ModelViewViewModelBlueprint/Public/MVVMBlueprintView.h
@@ -206,6 +206,10 @@ public:
 	UPROPERTY(Transient, NonTransactional)
 	TArray<TObjectPtr<UEdGraph>> TemporaryGraph;
 
+	// Use during compilation to clean the automatically generated graph.
+	UPROPERTY(Transient)
+	TArray<FName> TemporaryGraphNames;
+
 private:
 	UPROPERTY(Instanced)
 	TObjectPtr<UMVVMBlueprintViewSettings> Settings;
