View Models no longer need view bindings or events, to become available.
This is so we can use them in blueprints only.

diff --git a/Engine/Plugins/Runtime/ModelViewViewModel/Source/ModelViewViewModel/Public/View/MVVMViewClass.h b/Engine/Plugins/Runtime/ModelViewViewModel/Source/ModelViewViewModel/Public/View/MVVMViewClass.h
--- a/Engine/Plugins/Runtime/ModelViewViewModel/Source/ModelViewViewModel/Public/View/MVVMViewClass.h
+++ b/Engine/Plugins/Runtime/ModelViewViewModel/Source/ModelViewViewModel/Public/View/MVVMViewClass.h
@@ -639,6 +639,12 @@ public:
 	{
 		return bListenToViewModelCollectionChanged;
 	}
+	
+	/** Should the view be created even when there are no bindings or events. */
+	[[nodiscard]] bool IsCreatedWithoutBindings() const
+	{
+		return bCreateViewWithoutBindings;
+	}
 
 	/** Get the container of all the bindings. */
 	[[nodiscard]] const FMVVMCompiledBindingLibrary& GetBindingLibrary() const
@@ -779,6 +785,9 @@ private:
 	
 	UPROPERTY()
 	bool bListenToViewModelCollectionChanged = false;
+	
+	UPROPERTY()
+	bool bCreateViewWithoutBindings = false;
 
 #if WITH_EDITORONLY_DATA
 	FDelegateHandle BlueprintCompiledHandle;
diff --git a/Engine/Plugins/Runtime/ModelViewViewModel/Source/ModelViewViewModelBlueprint/Private/MVVMViewBlueprintCompiler.cpp b/Engine/Plugins/Runtime/ModelViewViewModel/Source/ModelViewViewModelBlueprint/Private/MVVMViewBlueprintCompiler.cpp
--- a/Engine/Plugins/Runtime/ModelViewViewModel/Source/ModelViewViewModelBlueprint/Private/MVVMViewBlueprintCompiler.cpp
+++ b/Engine/Plugins/Runtime/ModelViewViewModel/Source/ModelViewViewModelBlueprint/Private/MVVMViewBlueprintCompiler.cpp
@@ -1925,7 +1925,7 @@ bool FMVVMViewBlueprintCompiler::Compile(UWidgetBlueprintGeneratedClass* Class,
 		ViewExtension->bInitializeBindingsOnConstruct = ViewExtension->bInitializeSourcesOnConstruct ? BlueprintView->GetSettings()->bInitializeBindingsOnConstruct : false;
 		ViewExtension->bInitializeEventsOnConstruct = BlueprintView->GetSettings()->bInitializeEventsOnConstruct;
 		ViewExtension->bListenToViewModelCollectionChanged = Algo::AnyOf(ViewExtension->Sources, [](const FMVVMViewClass_Source& Source){ return Source.RequireGlobalViewModelCollectionUpdate(); });
-
+		ViewExtension->bCreateViewWithoutBindings = BlueprintView->GetSettings()->bCreateViewWithoutBindings;
 	}
 	{
 		ViewExtension->OptionalSources = 0;
diff --git a/Engine/Plugins/Runtime/ModelViewViewModel/Source/ModelViewViewModelBlueprint/Private/MVVMWidgetBlueprintExtension_View.cpp b/Engine/Plugins/Runtime/ModelViewViewModel/Source/ModelViewViewModelBlueprint/Private/MVVMWidgetBlueprintExtension_View.cpp
--- a/Engine/Plugins/Runtime/ModelViewViewModel/Source/ModelViewViewModelBlueprint/Private/MVVMWidgetBlueprintExtension_View.cpp
+++ b/Engine/Plugins/Runtime/ModelViewViewModel/Source/ModelViewViewModelBlueprint/Private/MVVMWidgetBlueprintExtension_View.cpp
@@ -360,12 +360,12 @@ void UMVVMWidgetBlueprintExtension_View::HandleFinishCompilingClass(UWidgetBluep
 			check(ViewExtension);
 
 			// Does it have any bindings
-			if (ViewExtension->GetBindings().Num() > 0 || ViewExtension->GetEvents().Num() > 0 || ViewExtension->GetConditions().Num() > 0)
+			if (ViewExtension->GetBindings().Num() > 0 || ViewExtension->GetEvents().Num() > 0 || ViewExtension->GetConditions().Num() > 0 || (GetBlueprintView()->GetViewModels().Num() > 0 && ViewExtension->IsCreatedWithoutBindings()))
 			{
 				// Test if parent also has a view
 				if (Class->GetExtension<UMVVMViewClass>(true))
 				{
-					CurrentCompilerContext->GetCompilerContext().MessageLog.Warning(*LOCTEXT("MoreThanOneViewWarning", "There is more than one view.").ToString());
+					CurrentCompilerContext->GetCompilerContext().MessageLog.Warning(*LOCTEXT("MoreThanOneViewWarning", "There is more than one view. This can happen when inheriting from a widget blueprint that also contains a view. Please be aware that the parent view model is not available to the child widget.").ToString());
 				}
 
 				// If we are not allowed to add the view class
diff --git a/Engine/Plugins/Runtime/ModelViewViewModel/Source/ModelViewViewModelBlueprint/Public/MVVMBlueprintView.h b/Engine/Plugins/Runtime/ModelViewViewModel/Source/ModelViewViewModelBlueprint/Public/MVVMBlueprintView.h
--- a/Engine/Plugins/Runtime/ModelViewViewModel/Source/ModelViewViewModelBlueprint/Public/MVVMBlueprintView.h
+++ b/Engine/Plugins/Runtime/ModelViewViewModel/Source/ModelViewViewModelBlueprint/Public/MVVMBlueprintView.h
@@ -64,6 +64,13 @@ public:
 	 */
 	UPROPERTY(EditAnywhere, Category = "View")
 	bool bInitializeEventsOnConstruct = true;
+	
+	/**
+	 * Create the view even when there are no view bindings or events.
+	 * If false, the view models will not be automatically available for use in blueprints if there are no bindings.
+	 */
+	UPROPERTY(EditAnywhere, Category = "View")
+	bool bCreateViewWithoutBindings = false;
 };
 
 /**
diff --git a/Engine/Plugins/Runtime/ModelViewViewModel/Source/ModelViewViewModelEditor/Private/Widgets/SMVVMViewModelPanel.cpp b/Engine/Plugins/Runtime/ModelViewViewModel/Source/ModelViewViewModelEditor/Private/Widgets/SMVVMViewModelPanel.cpp
--- a/Engine/Plugins/Runtime/ModelViewViewModel/Source/ModelViewViewModelEditor/Private/Widgets/SMVVMViewModelPanel.cpp
+++ b/Engine/Plugins/Runtime/ModelViewViewModel/Source/ModelViewViewModelEditor/Private/Widgets/SMVVMViewModelPanel.cpp
@@ -881,22 +881,11 @@ bool SMVVMViewModelPanel::RenameViewModelProperty(FGuid ViewModelGuid, const FTe
 
 EVisibility SMVVMViewModelPanel::GetWarningPanelVisibility() const
 {
-	if (!bDisableWarningPanel)
+	if (GetWarningMessage().IsEmpty())
 	{
-		if (UMVVMBlueprintView* WidgetBlueprint = WeakBlueprintView.Get())
-		{
-			if (WidgetBlueprint->GetNumBindings() == 0 && WidgetBlueprint->GetEvents().Num() == 0 && WidgetBlueprint->GetViewModels().Num() != 0)
-			{
-				return EVisibility::Visible;
-			}
-			if (!WidgetBlueprint->GetSettings()->bInitializeSourcesOnConstruct || !WidgetBlueprint->GetSettings()->bInitializeBindingsOnConstruct)
-			{
-				return EVisibility::Visible;
-			}
-		}
+		return EVisibility::Collapsed;
 	}
-
-	return EVisibility::Collapsed;
+	return EVisibility::Visible;
 }
 
 
@@ -906,9 +895,9 @@ FText SMVVMViewModelPanel::GetWarningMessage() const
 	{
 		if (UMVVMBlueprintView* WidgetBlueprint = WeakBlueprintView.Get())
 		{
-			if (WidgetBlueprint->GetNumBindings() == 0 && WidgetBlueprint->GetViewModels().Num() != 0)
+			if (WidgetBlueprint->GetNumBindings() == 0 && WidgetBlueprint->GetEvents().Num() == 0 && WidgetBlueprint->GetViewModels().Num() != 0 && !WidgetBlueprint->GetSettings()->bCreateViewWithoutBindings)
 			{
-				return LOCTEXT("NoBindingWarningDescription", "No view will be created for this widget. There are no bindings and there are no events. Your viewmodels won't be initialized.");
+				return LOCTEXT("NoBindingWarningDescription", "No view will be created for this widget. There are no bindings and there are no events. Your viewmodels won't be initialized. If you want to create the view anyway, choose the Create View Without Bindings option from the View Settings.");
 			}
 			if (!WidgetBlueprint->GetSettings()->bInitializeSourcesOnConstruct || !WidgetBlueprint->GetSettings()->bInitializeBindingsOnConstruct)
 			{
