Defensive code + logging when LoadedProperties is empty

#rnx
#tests Checked the 2 branches of code that calls the new function. Verified the result if index is invalid

diff --git a/Engine/Plugins/Runtime/ModelViewViewModel/Source/ModelViewViewModel/Private/Bindings/MVVMCompiledBindingLibrary.cpp b/Engine/Plugins/Runtime/ModelViewViewModel/Source/ModelViewViewModel/Private/Bindings/MVVMCompiledBindingLibrary.cpp
--- a/Engine/Plugins/Runtime/ModelViewViewModel/Source/ModelViewViewModel/Private/Bindings/MVVMCompiledBindingLibrary.cpp
+++ b/Engine/Plugins/Runtime/ModelViewViewModel/Source/ModelViewViewModel/Private/Bindings/MVVMCompiledBindingLibrary.cpp
@@ -6,6 +6,7 @@
 #include "IFieldNotificationClassDescriptor.h"
 #include "INotifyFieldValueChanged.h"
 #include "ModelViewViewModelModule.h"
+#include "MVVMMessageLog.h"
 #include "Templates/ValueOrError.h"
 #include "Types/MVVMFieldContext.h"
 #include "Types/MVVMFunctionContext.h"
@@ -326,6 +327,22 @@ namespace UE::MVVM::Private
 	}
 }
 
+bool FMVVMCompiledBindingLibrary::CanAccessLoadedProperiesAtIndex(int16 index) const
+{
+	if (LoadedProperties.IsValidIndex(index))
+	{
+		return true;
+	}
+	FString LoadedPropertiesString;
+	for (int32 i = 0; i < LoadedProperties.Num(); ++i)
+	{
+		LoadedPropertiesString.Append(FString::Format(TEXT("{0}, "), { LoadedProperties[i]->GetNameCPP() }));
+	}
+	FString DebugLoadedProperties = FString::Format(TEXT("{0} on {1}. LoadedProperties : {2}"), { index, LoadedProperties.Num(), LoadedPropertiesString });
+	checkf(false, TEXT("LoadedProperties.IsValidIndex(PathIndex.Index) failed. '%s'"), *DebugLoadedProperties);
+	UE_LOG(LogMVVM, Error, TEXT("LoadedProperties.IsValidIndex(PathIndex.Index) failed. '%s'"), *DebugLoadedProperties);
+	return false;
+}
 
 TValueOrError<UE::MVVM::FFieldContext, void> FMVVMCompiledBindingLibrary::EvaluateFieldPath(UObject* InExecutionSource, const FMVVMVCompiledFieldPath& InFieldPath) const
 {
@@ -362,15 +379,9 @@ TValueOrError<UE::MVVM::FFieldContext, void> FMVVMCompiledBindingLibrary::Evalua
 		FMVVMCompiledLoadedPropertyOrFunctionIndex PathIndex = FieldPaths[Index];
 		if (PathIndex.bIsProperty)
 		{
-			if (!LoadedProperties.IsValidIndex(PathIndex.Index))
+			if (!CanAccessLoadedProperiesAtIndex(PathIndex.Index))
 			{
-				FString LoadedPropertiesString;
-				for (int32 i = 0; i < LoadedProperties.Num(); ++i)
-				{
-					LoadedPropertiesString.Append(FString::Format(TEXT("{0}, "), { LoadedProperties[i]->GetNameCPP() }));
-				}
-				FString DebugLoadedProperties = FString::Format(TEXT("{0} on {1}. LoadedProperties : {2}"), { PathIndex.Index, LoadedProperties.Num(), LoadedPropertiesString });
-				checkf(false, TEXT("LoadedProperties.IsValidIndex(PathIndex.Index) failed. '%s'"), *DebugLoadedProperties);
+				return MakeError();
 			}
 			const FProperty* Property = LoadedProperties[PathIndex.Index];
 			if (!Property)
@@ -459,7 +470,10 @@ TValueOrError<UE::MVVM::FMVVMFieldVariant, void> FMVVMCompiledBindingLibrary::Ge
 	FMVVMCompiledLoadedPropertyOrFunctionIndex FinalPathIndex = FieldPaths[IterationMax];
 	if (FinalPathIndex.bIsProperty)
 	{
-		check(LoadedProperties.IsValidIndex(FinalPathIndex.Index));
+		if (!CanAccessLoadedProperiesAtIndex(FinalPathIndex.Index))
+		{
+			return MakeError();
+		}
 		FProperty* Property = LoadedProperties[FinalPathIndex.Index];
 		if (!Property)
 		{
diff --git a/Engine/Plugins/Runtime/ModelViewViewModel/Source/ModelViewViewModel/Public/Bindings/MVVMCompiledBindingLibrary.h b/Engine/Plugins/Runtime/ModelViewViewModel/Source/ModelViewViewModel/Public/Bindings/MVVMCompiledBindingLibrary.h
--- a/Engine/Plugins/Runtime/ModelViewViewModel/Source/ModelViewViewModel/Public/Bindings/MVVMCompiledBindingLibrary.h
+++ b/Engine/Plugins/Runtime/ModelViewViewModel/Source/ModelViewViewModel/Public/Bindings/MVVMCompiledBindingLibrary.h
@@ -313,6 +313,7 @@ private:
 
 	TValueOrError<UE::MVVM::FMVVMFieldVariant, void> GetFinalFieldFromPathImpl(UE::MVVM::FObjectVariant CurrentContainer, const FMVVMVCompiledFieldPath& FieldPath) const;
 
+	bool CanAccessLoadedProperiesAtIndex(int16 index) const;
 private:
 	struct FLoadedFunction
 	{
