MVVM: Fix up MemberParents in binding paths when a field in the path has changed type

#rnx

diff --git a/Engine/Plugins/Runtime/ModelViewViewModel/Source/ModelViewViewModelBlueprint/Private/MVVMPropertyPath.cpp b/Engine/Plugins/Runtime/ModelViewViewModel/Source/ModelViewViewModelBlueprint/Private/MVVMPropertyPath.cpp
--- a/Engine/Plugins/Runtime/ModelViewViewModel/Source/ModelViewViewModelBlueprint/Private/MVVMPropertyPath.cpp
+++ b/Engine/Plugins/Runtime/ModelViewViewModel/Source/ModelViewViewModelBlueprint/Private/MVVMPropertyPath.cpp
@@ -4,6 +4,7 @@
 #include "BlueprintCompilationManager.h"
 #include "Blueprint/WidgetTree.h"
 #include "Engine/Blueprint.h"
+#include "IPlatformFilePak.h"
 #include "StructUtils/UserDefinedStruct.h"
 #include "Kismet2/BlueprintEditorUtils.h"
 #include "Kismet2/StructureEditorUtils.h"
@@ -412,42 +413,106 @@ FText GetRootName(const UWidgetBlueprint* WidgetBlueprint, const FMVVMBlueprintP
 
 	return FText::GetEmpty();
 }
+	
+const UStruct* GetMostUpToDateStruct(const UStruct* Struct)
+{
+#if WITH_EDITOR
+	if (const UClass* Class = Cast<UClass>(Struct))
+	{
+		return FBlueprintEditorUtils::GetMostUpToDateClass(Class);
+	}
+	return Struct;
+#else
+	return Struct;
+#endif
+}
 }//namespace
 
 
 bool FMVVMBlueprintPropertyPath::OnFieldRenamed(const UBlueprint* InContext, UClass* FieldOwnerClass, FName OldObjectName, FName NewObjectName)
 {
+	bool bFoundRenamedField = false;
 	if (InContext != nullptr)
 	{
 		const EMVVMBlueprintFieldPathSource PathSource = GetSource(InContext);
 		const UClass* SelfContext = InContext->GeneratedClass;
+		
 		if (SelfContext == FieldOwnerClass)
 		{
 			if (PathSource == EMVVMBlueprintFieldPathSource::Widget && GetWidgetName() == OldObjectName)
 			{
 				SetWidgetName(NewObjectName);
-				return true;
+				bFoundRenamedField = true;
 			}
 			// Checking against NewObjectName since the viewmodel would have already updated its name
 			else if (PathSource == EMVVMBlueprintFieldPathSource::ViewModel && UE::MVVM::Private::GetViewModelName(InContext, GetViewModelId()) == NewObjectName)
 			{
 				// Nothing to rename since we reference the VM by ID, but the path does contain the renamed field
-				return true;
+				bFoundRenamedField = true;
+			}
+		}
+
+		const UStruct* OwnerStruct = nullptr;
+
+		if (PathSource == EMVVMBlueprintFieldPathSource::SelfContext)
+		{
+			OwnerStruct = SelfContext;
+		}
+		else if (PathSource == EMVVMBlueprintFieldPathSource::ViewModel || PathSource == EMVVMBlueprintFieldPathSource::Widget)
+		{
+			const FName OwnerName = (PathSource == EMVVMBlueprintFieldPathSource::ViewModel) ? UE::MVVM::Private::GetViewModelName(InContext, GetViewModelId()) : GetWidgetName();
+			const FProperty* OwnerProperty = FBlueprintEditorUtils::GetMostUpToDateClass(SelfContext)->FindPropertyByName(OwnerName);
+			if (const FObjectPropertyBase* ObjectProperty = CastField<FObjectPropertyBase>(OwnerProperty))
+			{
+				OwnerStruct = ObjectProperty->PropertyClass;
+			}
+			else if (const FStructProperty* StructProperty = CastField<FStructProperty>(OwnerProperty))
+			{
+				OwnerStruct = StructProperty->Struct;
 			}
 		}
 
 		// Check the path components for the renamed field
 		for (FMVVMBlueprintFieldPath& Path : Paths)
 		{
+			if (bFoundRenamedField)
+			{
+				if (!Path.IsFieldSelfContext() && OwnerStruct)
+				{
+					// We need to update the stored parent in the path in case the parent changed type
+					FProperty* UpdatedProperty = OwnerStruct->FindPropertyByName(Path.GetFieldName(nullptr));
+					if (ensure(UpdatedProperty))
+					{
+						Path = { InContext, UE::MVVM::FMVVMConstFieldVariant(UpdatedProperty) };
+					}
+				}
+				
+				break;
+			}
+
+			if (OwnerStruct)
+			{
+				const UStruct* UpToDateStruct = UE::MVVM::Private::GetMostUpToDateStruct(OwnerStruct);
+				const FProperty* OwnerProperty = UpToDateStruct->FindPropertyByName(Path.GetFieldName(SelfContext));
+				if (const FObjectPropertyBase* ObjectProperty = CastField<FObjectPropertyBase>(OwnerProperty))
+				{
+					OwnerStruct = ObjectProperty->PropertyClass;
+				}
+				else if (const FStructProperty* StructProperty = CastField<FStructProperty>(OwnerProperty))
+				{
+					OwnerStruct = StructProperty->Struct;
+				}
+			}
+			
 			if (Path.GetParentClass(SelfContext) == FieldOwnerClass && Path.GetFieldName(SelfContext) == OldObjectName)
 			{
 				Path.SetMemberName(NewObjectName);
-				return true;
+				bFoundRenamedField = true;
 			}
 		}
 	}
 
-	return false;
+	return bFoundRenamedField;
 }
 
 FText FMVVMBlueprintPropertyPath::ToText(const UWidgetBlueprint* WidgetBlueprint, bool bUseDisplayName) const
