UndoRedo of view model operation updates. Added transaction to RemoveViewModel

#tests Undo/Redo of Add viewmodel and Remove viewmodel
#rnx

diff --git a/Engine/Plugins/Runtime/ModelViewViewModel/Source/ModelViewViewModelEditor/Private/Widgets/SMVVMViewModelPanel.cpp b/Engine/Plugins/Runtime/ModelViewViewModel/Source/ModelViewViewModelEditor/Private/Widgets/SMVVMViewModelPanel.cpp
--- a/Engine/Plugins/Runtime/ModelViewViewModel/Source/ModelViewViewModelEditor/Private/Widgets/SMVVMViewModelPanel.cpp
+++ b/Engine/Plugins/Runtime/ModelViewViewModel/Source/ModelViewViewModelEditor/Private/Widgets/SMVVMViewModelPanel.cpp
@@ -11,6 +11,7 @@
 #include "MVVMEditorSubsystem.h"
 #include "MVVMWidgetBlueprintExtension_View.h"
 #include "PropertyEditorModule.h"
+#include "ScopedTransaction.h"
 #include "ToolMenus.h"
 #include "View/MVVMViewModelContextResolver.h"
 #include "ViewModelFieldDragDropOp.h"
@@ -217,6 +218,9 @@ void SMVVMViewModelPanel::Construct(const FArguments& InArgs, TSharedPtr<FWidget
 		ExtensionAddeddHandle = WidgetBlueprint->OnExtensionAdded.AddSP(this, &SMVVMViewModelPanel::HandleViewUpdated);
 	}
 
+	WidgetBlueprintTransactionHandle = WidgetBlueprintEditor->GetOnWidgetBlueprintTransaction().AddSP(this, &SMVVMViewModelPanel::FillViewModel);
+
+
 	CreateCommandList();
 
 	FToolMenuContext GenerateWidgetContext;
@@ -322,6 +326,14 @@ SMVVMViewModelPanel::~SMVVMViewModelPanel()
 			CurrentBlueprintView->OnViewModelsUpdated.Remove(ViewModelsUpdatedHandle);
 		}
 	}
+
+	if (WidgetBlueprintTransactionHandle.IsValid())
+	{
+		if (TSharedPtr<FWidgetBlueprintEditor> WidgetBlueprintEditor = WeakBlueprintEditor.Pin())
+		{
+			WidgetBlueprintEditor->GetOnWidgetBlueprintTransaction().Remove(WidgetBlueprintTransactionHandle);
+		}
+	}
 }
 
 
@@ -677,6 +689,18 @@ void SMVVMViewModelPanel::HandleDeleteViewModel()
 	}
 
 	TArray<UE::PropertyViewer::SPropertyViewer::FSelectedItem> Items = ViewModelTreeView->GetSelectedItems();
+	if (Items.Num() == 0)
+	{
+		return;
+	}
+
+	const FScopedTransaction Transaction(LOCTEXT("RemoveViewModel", "Remove viewmodel"));
+	UMVVMBlueprintView* BpView = GEditor->GetEditorSubsystem<UMVVMEditorSubsystem>()->GetView(WidgetBP);
+	if (BpView != nullptr)
+	{
+		BpView->Modify();
+	}
+
 	for (UE::PropertyViewer::SPropertyViewer::FSelectedItem& SelectedItem : Items)
 	{
 		if (SelectedItem.bIsContainerSelected)
diff --git a/Engine/Plugins/Runtime/ModelViewViewModel/Source/ModelViewViewModelEditor/Private/Widgets/SMVVMViewModelPanel.h b/Engine/Plugins/Runtime/ModelViewViewModel/Source/ModelViewViewModelEditor/Private/Widgets/SMVVMViewModelPanel.h
--- a/Engine/Plugins/Runtime/ModelViewViewModel/Source/ModelViewViewModelEditor/Private/Widgets/SMVVMViewModelPanel.h
+++ b/Engine/Plugins/Runtime/ModelViewViewModel/Source/ModelViewViewModelEditor/Private/Widgets/SMVVMViewModelPanel.h
@@ -124,6 +124,7 @@ private:
 	TWeakObjectPtr<UMVVMBlueprintView> WeakBlueprintView;
 	FDelegateHandle ViewModelsUpdatedHandle;
 	FDelegateHandle ExtensionAddeddHandle;
+	FDelegateHandle WidgetBlueprintTransactionHandle;
 
 	TStrongObjectPtr<UMVVMBlueprintViewModelContextWrapper> ModelContextWrapper;
 
