Fix for UE-235098 Error message when setting vm to null and then setting it back to something valid

#tests made a graph with delay setting the vm, then to null, then to the vm again. Saw that the bindings were evaluated and animations were triggered
#rnx

diff --git a/Engine/Plugins/Runtime/ModelViewViewModel/Source/ModelViewViewModel/Private/View/MVVMView.cpp b/Engine/Plugins/Runtime/ModelViewViewModel/Source/ModelViewViewModel/Private/View/MVVMView.cpp
--- a/Engine/Plugins/Runtime/ModelViewViewModel/Source/ModelViewViewModel/Private/View/MVVMView.cpp
+++ b/Engine/Plugins/Runtime/ModelViewViewModel/Source/ModelViewViewModel/Private/View/MVVMView.cpp
@@ -928,7 +928,13 @@ bool UMVVMView::SetSourceInternal(FMVVMViewClass_SourceKey ClassSourceKey, TScri
 
 		if (ViewSource.bBindingsInitialized)
 		{
-			UninitializeSourceBindings(ClassSourceKey, ClassSource, ViewSource);
+			// if we had bindings init, then the vm gets set to null, then the vm gets set to a valid ptr, we want to update the binding values as it gets set
+			// the only issue is that UninitializeSourceBindings expect the source to be non null
+			// if the previous source was null, we don't need to do the work anyway
+			if (ViewSource.Source != nullptr)
+			{
+				UninitializeSourceBindings(ClassSourceKey, ClassSource, ViewSource);
+			}
 			GEngine->GetEngineSubsystem<UMVVMBindingSubsystem>()->RemoveDelayedBindings(this, ClassSourceKey);
 		}
 		UninitializeSource(ViewSourceKey);
