MVVM extension: update the warning text on panel widget extension to be more clear on how to create the array binding.

diff --git a/Engine/Plugins/Runtime/ModelViewViewModel/Source/ModelViewViewModelBlueprint/Private/Extensions/MVVMViewBlueprintPanelWidgetExtension.cpp b/Engine/Plugins/Runtime/ModelViewViewModel/Source/ModelViewViewModelBlueprint/Private/Extensions/MVVMViewBlueprintPanelWidgetExtension.cpp
--- a/Engine/Plugins/Runtime/ModelViewViewModel/Source/ModelViewViewModelBlueprint/Private/Extensions/MVVMViewBlueprintPanelWidgetExtension.cpp
+++ b/Engine/Plugins/Runtime/ModelViewViewModel/Source/ModelViewViewModelBlueprint/Private/Extensions/MVVMViewBlueprintPanelWidgetExtension.cpp
@@ -113,7 +113,7 @@ void UMVVMBlueprintViewExtension_PanelWidget::Precompile(UE::MVVM::Compiler::IMV
 
 		if (!bFoundSetterBinding && EntryViewModelId.IsValid())
 		{
-			Compiler->AddMessage(FText::Format(LOCTEXT("PreCompileMVVMWidgetExtensionNoSetterBindingFound", "No binding found from an array of viewmodels to {0} -> SetItems. Please add this binding or remove the Viewmodel extension on widget {0} from its details panel."), FText::FromName(WidgetName)), UE::MVVM::Compiler::EMessageType::Warning);
+			Compiler->AddMessage(FText::Format(LOCTEXT("PreCompileMVVMWidgetExtensionNoSetterBindingFound", "No binding found from an array of viewmodels to {0} -> SetItems. Please find {0} on the root widget and add this binding or remove the Viewmodel extension on widget {1} from its details panel."), FText::FromName(PanelPropertyName), FText::FromName(WidgetName)), UE::MVVM::Compiler::EMessageType::Warning);
 		}
 	};
 	 
