Support Property Interface as an actual interface for UMG widgets.

#tests Different mix and match of property interface as interface or property in a umg widget. Played with the settings in the editor to ensure that all flows were working correctly. On the verse side, played with property interface to ensure that there was no regression. Tested the new pattern possible to play with property interface:
if (MyPI := PlayerInfo_interface[mywidget]):
set MyPI.PlayerInfo.Value = newValue

diff --git a/Engine/Plugins/Runtime/ModelViewViewModel/Source/ModelViewViewModelBlueprint/Private/MVVMBlueprintViewModelContext.cpp b/Engine/Plugins/Runtime/ModelViewViewModel/Source/ModelViewViewModelBlueprint/Private/MVVMBlueprintViewModelContext.cpp
--- a/Engine/Plugins/Runtime/ModelViewViewModel/Source/ModelViewViewModelBlueprint/Private/MVVMBlueprintViewModelContext.cpp
+++ b/Engine/Plugins/Runtime/ModelViewViewModel/Source/ModelViewViewModelBlueprint/Private/MVVMBlueprintViewModelContext.cpp
@@ -2,6 +2,7 @@
 
 #include "MVVMBlueprintViewModelContext.h"
 #include "MVVMDeveloperProjectSettings.h"
+#include "HAL/IConsoleManager.h"
 #include "View/MVVMViewModelContextResolver.h"
 
 #include UE_INLINE_GENERATED_CPP_BY_NAME(MVVMBlueprintViewModelContext)
@@ -59,6 +60,14 @@ namespace Private
 static const FName NAME_MVVMAllowedContextCreationType = TEXT("MVVMAllowedContextCreationType");
 static const FName NAME_MVVMDisallowedContextCreationType = TEXT("MVVMDisallowedContextCreationType");
 
+bool GSupportUseAsInterfaceSetting = false;
+static FAutoConsoleVariableRef CVarSupportUseAsInterfaceSetting(
+	TEXT("MVVM.SupportUseAsInterfaceSetting"),
+	GSupportUseAsInterfaceSetting,
+	TEXT("Enable in order to have the UseAsInterface option available in UMG's viewmodel."),
+	ECVF_ReadOnly
+);
+
 // Find the Allowed and Disallowed from first class in the hierarchy that defines it.
 void GenerateAllowedList(const UClass* Class, TArray<FString>& AllowedStrings, TArray<FString>& DisallowedStrings)
 {
diff --git a/Engine/Plugins/Runtime/ModelViewViewModel/Source/ModelViewViewModelBlueprint/Public/MVVMBlueprintViewModelContext.h b/Engine/Plugins/Runtime/ModelViewViewModel/Source/ModelViewViewModelBlueprint/Public/MVVMBlueprintViewModelContext.h
--- a/Engine/Plugins/Runtime/ModelViewViewModel/Source/ModelViewViewModelBlueprint/Public/MVVMBlueprintViewModelContext.h
+++ b/Engine/Plugins/Runtime/ModelViewViewModel/Source/ModelViewViewModelBlueprint/Public/MVVMBlueprintViewModelContext.h
@@ -97,6 +97,11 @@ public:
 		return NotifyFieldValueClass != nullptr;
 	}
 
+	bool CanRename() const
+	{
+		return bCanRename && !bUseAsInterface;
+	}
+
 #if WITH_EDITOR
 	[[nodiscard]] TObjectPtr<UMVVMViewModelContextResolver> CreateDefaultResolver(UPackage* Package) const;
 #endif
@@ -143,14 +148,14 @@ public:
 	 * Generate a public setter for this viewmodel.
 	 * @note Always true when the Creation Type is Manual.
 	 */
-	UPROPERTY(EditAnywhere, Category = "Viewmodel", AdvancedDisplay, meta = (DisplanName="Create Public Setter"))
+	UPROPERTY(EditAnywhere, Category = "Viewmodel", AdvancedDisplay, meta = (DisplayName="Create Public Setter"))
 	bool bCreateSetterFunction = false;
 
 	/**
 	 * Generate a public getter for this viewmodel.
 	 * @note Always false when using a Instanced viewmodel.
 	 */
-	UPROPERTY(EditAnywhere, Category = "Viewmodel", AdvancedDisplay, meta = (DisplanName = "Create Public Getter"))
+	UPROPERTY(EditAnywhere, Category = "Viewmodel", AdvancedDisplay, meta = (DisplayName = "Create Public Getter"))
 	bool bCreateGetterFunction = true;
 
 	/**
@@ -191,6 +196,10 @@ public:
 	UPROPERTY()
 	bool bCanRemove = true;
 
+	/** Will the viewmodel be handled as a property or as an interface in Verse. */
+	UPROPERTY(EditAnywhere, Category = "View")
+	bool bUseAsInterface = false;
+
 private:
 	UPROPERTY()
 	bool bCreateSetterFunction_Deprecation = false;
diff --git a/Engine/Plugins/Runtime/ModelViewViewModel/Source/ModelViewViewModelEditor/Private/Customizations/MVVMBlueprintViewModelContextCustomization.cpp b/Engine/Plugins/Runtime/ModelViewViewModel/Source/ModelViewViewModelEditor/Private/Customizations/MVVMBlueprintViewModelContextCustomization.cpp
--- a/Engine/Plugins/Runtime/ModelViewViewModel/Source/ModelViewViewModelEditor/Private/Customizations/MVVMBlueprintViewModelContextCustomization.cpp
+++ b/Engine/Plugins/Runtime/ModelViewViewModel/Source/ModelViewViewModelEditor/Private/Customizations/MVVMBlueprintViewModelContextCustomization.cpp
@@ -205,6 +205,7 @@ void FBlueprintViewModelContextDetailCustomization::CustomizeChildren(TSharedRef
 	CreationTypeHandle = PropertyHandle->GetChildHandle(GET_MEMBER_NAME_CHECKED(FMVVMBlueprintViewModelContext, CreationType), false);
 	ViewModelNameHandle = PropertyHandle->GetChildHandle(GET_MEMBER_NAME_CHECKED(FMVVMBlueprintViewModelContext, ViewModelName), false);
 	OptionalHandle = PropertyHandle->GetChildHandle(GET_MEMBER_NAME_CHECKED(FMVVMBlueprintViewModelContext, bOptional), false);
+	UseAsInterfaceHandle = PropertyHandle->GetChildHandle(GET_MEMBER_NAME_CHECKED(FMVVMBlueprintViewModelContext, bUseAsInterface), false);
 	CreateSetterFunctionHandle = PropertyHandle->GetChildHandle(GET_MEMBER_NAME_CHECKED(FMVVMBlueprintViewModelContext, bCreateSetterFunction), false);
 	ForceExecuteBindingsOnSetSourceHandle = PropertyHandle->GetChildHandle(GET_MEMBER_NAME_CHECKED(FMVVMBlueprintViewModelContext, bForceExecuteBindingsOnSetSource), false);
 	ResolverHandle = PropertyHandle->GetChildHandle(GET_MEMBER_NAME_CHECKED(FMVVMBlueprintViewModelContext, Resolver), false);
@@ -442,6 +443,24 @@ void FBlueprintViewModelContextDetailCustomization::CustomizeChildren(TSharedRef
 					}));
 		}
 
+		if (ensure(UseAsInterfaceHandle))
+		{
+			ChildBuilder.AddProperty(UseAsInterfaceHandle.ToSharedRef())
+				.IsEnabled(TAttribute<bool>(this, &FBlueprintViewModelContextDetailCustomization::IsUseAsInterfaceAvailable))
+				.ToolTip(LOCTEXT("UseAsInterfaceToolTip", "True to use as interface, False to use as property.\nIt can only be enabled if there is no other interface of the same type, and if the default name for the viewmodel is available.\nIt is not possible to rename the viewmodel if the option is enabled."))
+				.Visibility(MakeAttributeLambda([]()
+					{
+						IConsoleVariable* CVarSupportUseAsInterfaceSetting = IConsoleManager::Get().FindConsoleVariable(TEXT("MVVM.SupportUseAsInterfaceSetting")); ensure(CVarSupportUseAsInterfaceSetting);
+						if (CVarSupportUseAsInterfaceSetting)
+						{
+							return CVarSupportUseAsInterfaceSetting->GetBool() ? EVisibility::Visible : EVisibility::Collapsed;
+						}
+						return EVisibility::Collapsed;
+					}));
+
+			UseAsInterfaceHandle->SetOnPropertyValueChanged(FSimpleDelegate::CreateSP(this, &FBlueprintViewModelContextDetailCustomization::HandleUseAsInterfaceChanged));
+		}
+
 		if (ensure(ForceExecuteBindingsOnSetSourceHandle))
 		{
 			ChildBuilder.AddProperty(ForceExecuteBindingsOnSetSourceHandle.ToSharedRef())
@@ -569,10 +588,10 @@ TSharedRef<SWidget> FBlueprintViewModelContextDetailCustomization::CreateExecuti
 			FUIAction
 			(
 				FExecuteAction::CreateLambda([this, Type]()
-				{
-					uint8 Value = static_cast<uint8>(Type);
-					CreationTypeHandle->SetValue(Value);
-				})
+					{
+						uint8 Value = static_cast<uint8>(Type);
+						CreationTypeHandle->SetValue(Value);
+					})
 			)
 		);
 	}
@@ -601,6 +620,47 @@ FText FBlueprintViewModelContextDetailCustomization::GetExecutionTypeValueToolTi
 	return FText::GetEmpty();
 }
 
+bool FBlueprintViewModelContextDetailCustomization::IsUseAsInterfaceAvailable() const
+{
+	UWidgetBlueprint* WidgetBlueprint = WidgetBlueprintEditor.Pin()->GetWidgetBlueprintObj();
+	FMVVMBlueprintViewModelContext* ContextPtr = Private::GetViewModelContext(ContextHandle.ToSharedRef());
+	if (ContextPtr == nullptr)
+	{
+		return false;
+	}
+	const UClass* ViewModelClass = ContextPtr->GetViewModelClass();
+	if (ViewModelClass == nullptr)
+	{
+		return false;
+	}
+
+	UMVVMEditorSubsystem* EditorSubsystem = GEditor->GetEditorSubsystem<UMVVMEditorSubsystem>();
+	check(EditorSubsystem);
+	UMVVMBlueprintView* View = EditorSubsystem->GetView(WidgetBlueprint);
+	if (View == nullptr)
+	{
+		return false;
+	}
+
+	const FGuid ViewModelContextId = ContextPtr->GetViewModelId();
+	const TArrayView<const FMVVMBlueprintViewModelContext> ViewModels = View->GetViewModels();
+
+	bool bHasInterfaceOfSameType = false;
+	if (ViewModels.FindByPredicate([&ViewModelClass, &ViewModelContextId](FMVVMBlueprintViewModelContext ViewModel) { return ViewModel.GetViewModelClass() == ViewModelClass && ViewModel.bUseAsInterface && ViewModel.GetViewModelId() != ViewModelContextId; }))
+	{
+		bHasInterfaceOfSameType = true;
+	}
+
+	bool bHasAnotherViewModelUsingTheName = false;
+	const FString DefaultViewModelName = UMVVMEditorSubsystem::GetDefaultViewModelName(ViewModelClass);
+	if (ViewModels.FindByPredicate([&DefaultViewModelName, &ViewModelContextId](FMVVMBlueprintViewModelContext ViewModel) { return ViewModel.GetViewModelName() == DefaultViewModelName && ViewModel.GetViewModelId() != ViewModelContextId; }))
+	{
+		bHasAnotherViewModelUsingTheName = true;
+	}
+
+	return !bHasInterfaceOfSameType && !bHasAnotherViewModelUsingTheName;
+}
+
 FText FBlueprintViewModelContextDetailCustomization::GetViewModalNameValueAsText() const
 {
 	check(ViewModelNameHandle.IsValid());
@@ -740,6 +800,28 @@ bool FBlueprintViewModelContextDetailCustomization::HandleNameVerifyTextChanged(
 	return Private::VerifyViewModelName(WidgetBlueprintEditor.Pin(), ViewModelNameHandle, NewText, false, OutError);
 }
 
+void FBlueprintViewModelContextDetailCustomization::HandleUseAsInterfaceChanged()
+{
+	bool bUseAsInterfaceHandle = false;
+	if (UseAsInterfaceHandle->GetValue(bUseAsInterfaceHandle) == FPropertyAccess::Success)
+	{
+		if (FMVVMBlueprintViewModelContext* ContextPtr = Private::GetViewModelContext(ContextHandle.ToSharedRef()))
+		{
+			const FString DefaultViewModelName = UMVVMEditorSubsystem::GetDefaultViewModelName(ContextPtr->GetViewModelClass());
+			if (ContextPtr->GetViewModelName() != DefaultViewModelName && bUseAsInterfaceHandle)
+			{
+				FText OutErrorMessage;
+				FText NewName = FText::FromString(DefaultViewModelName);
+				if (!Private::VerifyViewModelName(WidgetBlueprintEditor.Pin(), ViewModelNameHandle, NewName, true, OutErrorMessage))
+				{
+					//Unable to set the name as required, revert the change.
+					UseAsInterfaceHandle->SetValue(false);
+				}
+			}
+		}
+	}
+}
+
 } // namespace UE::MVVM
 
 #undef LOCTEXT_NAMESPACE
diff --git a/Engine/Plugins/Runtime/ModelViewViewModel/Source/ModelViewViewModelEditor/Private/Customizations/MVVMBlueprintViewModelContextCustomization.h b/Engine/Plugins/Runtime/ModelViewViewModel/Source/ModelViewViewModelEditor/Private/Customizations/MVVMBlueprintViewModelContextCustomization.h
--- a/Engine/Plugins/Runtime/ModelViewViewModel/Source/ModelViewViewModelEditor/Private/Customizations/MVVMBlueprintViewModelContextCustomization.h
+++ b/Engine/Plugins/Runtime/ModelViewViewModel/Source/ModelViewViewModelEditor/Private/Customizations/MVVMBlueprintViewModelContextCustomization.h
@@ -55,12 +55,14 @@ private:
 	void HandleClassChanged();
 	void HandleCreationTypeChanged();
 	void HandleCreateSetterFunctionChanged();
+	void HandleUseAsInterfaceChanged();
 	TSharedRef<SWidget> HandleClassGetMenuContent();
 	void HandleClassCancelMenu();
 	void HandleClassCommitted(const UClass* SelectedClass);
 	FText GetViewModalNameValueAsText() const;
 	void HandleNameTextCommitted(const FText& NewText, ETextCommit::Type /*CommitInfo*/);
 	bool HandleNameVerifyTextChanged(const FText& Text, FText& OutError) const;
+	bool IsUseAsInterfaceAvailable() const;
 
 	FViewModelPropertyAccessEditor PropertyAccessEditor;
 	TWeakPtr<FWidgetBlueprintEditor> WidgetBlueprintEditor;
@@ -71,6 +73,7 @@ private:
 	TSharedPtr<IPropertyHandle> CreationTypeHandle;
 	TSharedPtr<IPropertyHandle> ViewModelNameHandle;
 	TSharedPtr<IPropertyHandle> OptionalHandle;
+	TSharedPtr<IPropertyHandle> UseAsInterfaceHandle;
 	TSharedPtr<IPropertyHandle> CreateSetterFunctionHandle;
 	TSharedPtr<IPropertyHandle> ForceExecuteBindingsOnSetSourceHandle;
 	TSharedPtr<IPropertyHandle> ResolverHandle;
diff --git a/Engine/Plugins/Runtime/ModelViewViewModel/Source/ModelViewViewModelEditor/Private/MVVMEditorSubsystem.cpp b/Engine/Plugins/Runtime/ModelViewViewModel/Source/ModelViewViewModelEditor/Private/MVVMEditorSubsystem.cpp
--- a/Engine/Plugins/Runtime/ModelViewViewModel/Source/ModelViewViewModelEditor/Private/MVVMEditorSubsystem.cpp
+++ b/Engine/Plugins/Runtime/ModelViewViewModel/Source/ModelViewViewModelEditor/Private/MVVMEditorSubsystem.cpp
@@ -251,27 +251,40 @@ FGuid UMVVMEditorSubsystem::AddViewModel(UWidgetBlueprint* WidgetBlueprint, cons
 	{
 		if (UMVVMBlueprintView* View = GetView(WidgetBlueprint))
 		{
-			FString ClassName = ViewModelClass->ClassGeneratedBy != nullptr ? ViewModelClass->ClassGeneratedBy->GetName() : ViewModelClass->GetAuthoredName();
-			if (Cast<UBlueprintGeneratedClass>(ViewModelClass) && ClassName.EndsWith(TEXT("_C")))
+			FString DefaultViewModelName = GetDefaultViewModelName(ViewModelClass);
+			FString ViewModelName = DefaultViewModelName;
+			FKismetNameValidator NameValidator(WidgetBlueprint);
+
+			bool bHasInterfaceOfSameType = false;
+			const TArrayView<const FMVVMBlueprintViewModelContext> ViewModels = View->GetViewModels();
+			if (const FMVVMBlueprintViewModelContext* ViewModel = ViewModels.FindByPredicate([&ViewModelClass](FMVVMBlueprintViewModelContext ViewModel) { return ViewModel.GetViewModelClass() == ViewModelClass && ViewModel.bUseAsInterface; }))
 			{
-				ClassName.RemoveAt(ClassName.Len() - 2, 2);
+				bHasInterfaceOfSameType = true;
 			}
-			FString ViewModelName = ClassName;
-			FKismetNameValidator NameValidator(WidgetBlueprint);
 
 			int32 Index = 1;
 			while (NameValidator.IsValid(ViewModelName) != EValidatorResult::Ok)
 			{
-				ViewModelName = ClassName + "_";
+				ViewModelName = DefaultViewModelName + "_";
 				ViewModelName.AppendInt(Index);
 
 				++Index;
 			}
+			const bool bHasViewModelWithDefaultName = ViewModelName != DefaultViewModelName;
+
+			bool bSupportUseAsInterfaceSetting = false;
+			IConsoleVariable* CVarSupportUseAsInterfaceSetting = IConsoleManager::Get().FindConsoleVariable(TEXT("MVVM.SupportUseAsInterfaceSetting")); ensure(CVarSupportUseAsInterfaceSetting);
+			if (CVarSupportUseAsInterfaceSetting)
+			{
+				bSupportUseAsInterfaceSetting = CVarSupportUseAsInterfaceSetting->GetBool();
+			}
+			const bool bCanAddAsInterface = bSupportUseAsInterfaceSetting && !bHasInterfaceOfSameType && !bHasViewModelWithDefaultName;
 
 			FMVVMBlueprintViewModelContext Context = FMVVMBlueprintViewModelContext(ViewModelClass, *ViewModelName);
 			if (Context.IsValid())
 			{
 				Context.Resolver = Context.CreateDefaultResolver(WidgetBlueprint->GetPackage());
+				Context.bUseAsInterface = bCanAddAsInterface;
 
 				const FScopedTransaction Transaction(LOCTEXT("AddViewModel", "Add viewmodel"));
 				View->Modify();
@@ -283,6 +296,22 @@ FGuid UMVVMEditorSubsystem::AddViewModel(UWidgetBlueprint* WidgetBlueprint, cons
 	return Result;
 }
 
+FString UMVVMEditorSubsystem::GetDefaultViewModelName(const UClass* ViewModelClass)
+{
+	if (ViewModelClass == nullptr)
+	{
+		return "";
+	}
+
+	FString ClassName = ViewModelClass->ClassGeneratedBy != nullptr ? ViewModelClass->ClassGeneratedBy->GetName() : ViewModelClass->GetAuthoredName();
+	if (Cast<UBlueprintGeneratedClass>(ViewModelClass) && ClassName.EndsWith(TEXT("_C")))
+	{
+		ClassName.RemoveAt(ClassName.Len() - 2, 2);
+	}
+
+	return ClassName;
+}
+
 FGuid UMVVMEditorSubsystem::AddInstancedViewModel(UWidgetBlueprint* WidgetBlueprint)
 {
 	UMVVMWidgetBlueprintExtension_View* ExtensionView = UMVVMWidgetBlueprintExtension_View::GetExtension<UMVVMWidgetBlueprintExtension_View>(WidgetBlueprint);
diff --git a/Engine/Plugins/Runtime/ModelViewViewModel/Source/ModelViewViewModelEditor/Private/Widgets/SMVVMViewModelPanel.cpp b/Engine/Plugins/Runtime/ModelViewViewModel/Source/ModelViewViewModelEditor/Private/Widgets/SMVVMViewModelPanel.cpp
--- a/Engine/Plugins/Runtime/ModelViewViewModel/Source/ModelViewViewModelEditor/Private/Widgets/SMVVMViewModelPanel.cpp
+++ b/Engine/Plugins/Runtime/ModelViewViewModel/Source/ModelViewViewModelEditor/Private/Widgets/SMVVMViewModelPanel.cpp
@@ -577,7 +577,7 @@ bool SMVVMViewModelPanel::HandleCanRename(FGuid ViewModelGuid) const
 	const UMVVMBlueprintView* BlueprintView = WeakBlueprintView.Get();
 	const FMVVMBlueprintViewModelContext* ViewModelContext = BlueprintView ? BlueprintView->FindViewModel(ViewModelGuid) : nullptr;
 	TSharedPtr<FWidgetBlueprintEditor> BlueprintEditor = WeakBlueprintEditor.Pin();
-	return ViewModelContext && ViewModelContext->bCanRename && BlueprintEditor && BlueprintEditor->InEditingMode();
+	return ViewModelContext && ViewModelContext->CanRename() && BlueprintEditor && BlueprintEditor->InEditingMode();
 }
 
 
@@ -884,7 +884,7 @@ bool SMVVMViewModelPanel::RenameViewModelProperty(FGuid ViewModelGuid, const FTe
 			{
 				if (const FMVVMBlueprintViewModelContext* ViewModelContext = View->FindViewModel(ViewModelGuid))
 				{
-					if (ViewModelContext->bCanRename)
+					if (ViewModelContext->CanRename())
 					{
 						if (bCommit)
 						{
diff --git a/Engine/Plugins/Runtime/ModelViewViewModel/Source/ModelViewViewModelEditor/Public/MVVMEditorSubsystem.h b/Engine/Plugins/Runtime/ModelViewViewModel/Source/ModelViewViewModelEditor/Public/MVVMEditorSubsystem.h
--- a/Engine/Plugins/Runtime/ModelViewViewModel/Source/ModelViewViewModelEditor/Public/MVVMEditorSubsystem.h
+++ b/Engine/Plugins/Runtime/ModelViewViewModel/Source/ModelViewViewModelEditor/Public/MVVMEditorSubsystem.h
@@ -163,6 +163,7 @@ public:
 	TArray<UE::MVVM::FBindingSource> GetAllViewModels(const UWidgetBlueprint* WidgetBlueprint) const;
 
 	FGuid GetFirstBindingThatUsesViewModel(const UWidgetBlueprint* WidgetBlueprint, FGuid ViewModelId) const;
+	static FString GetDefaultViewModelName(const UClass* ViewModelClass);
 
 private:
 	mutable TUniquePtr<UE::MVVM::ConversionFunctionLibrary::FCollection> ConversionFunctionCollection;
