Collapse/Expand options for ViewBindings Panel

[RN] minor creator/ui

diff --git a/Engine/Plugins/Runtime/ModelViewViewModel/Source/ModelViewViewModelEditor/Private/MVVMViewBindingMenuContext.h b/Engine/Plugins/Runtime/ModelViewViewModel/Source/ModelViewViewModelEditor/Private/MVVMViewBindingMenuContext.h
new file mode 100644
--- /dev/null
+++ b/Engine/Plugins/Runtime/ModelViewViewModel/Source/ModelViewViewModelEditor/Private/MVVMViewBindingMenuContext.h
@@ -0,0 +1,20 @@
+// Copyright Epic Games, Inc. All Rights Reserved.
+
+#pragma once
+
+#include "CoreMinimal.h"
+
+#include "MVVMViewBindingMenuContext.generated.h"
+
+class FWidgetBlueprintEditor;
+namespace UE::MVVM { class SBindingsPanel; }
+
+UCLASS()
+	class UMVVMViewBindingMenuContext : public UObject
+{
+	GENERATED_BODY()
+
+public:
+	TWeakPtr<FWidgetBlueprintEditor> WidgetBlueprintEditor;
+	TWeakPtr<UE::MVVM::SBindingsPanel> BindingsPanel;
+};
diff --git a/Engine/Plugins/Runtime/ModelViewViewModel/Source/ModelViewViewModelEditor/Private/Widgets/SMVVMViewBindingListView.cpp b/Engine/Plugins/Runtime/ModelViewViewModel/Source/ModelViewViewModelEditor/Private/Widgets/SMVVMViewBindingListView.cpp
--- a/Engine/Plugins/Runtime/ModelViewViewModel/Source/ModelViewViewModelEditor/Private/Widgets/SMVVMViewBindingListView.cpp
+++ b/Engine/Plugins/Runtime/ModelViewViewModel/Source/ModelViewViewModelEditor/Private/Widgets/SMVVMViewBindingListView.cpp
@@ -734,6 +734,18 @@ void SBindingsList::RequestNavigateToCondition(UMVVMBlueprintViewCondition* Cond
 	}
 }
 
+void SBindingsList::SetExpansion(bool InIsExpanded)
+{
+	if (TreeView.IsValid())
+	{
+		// Set expansion only to the binding widgets
+		for (const TSharedPtr<FBindingEntry>& Entry : TreeView->GetRootItems())
+		{
+			TreeView->SetItemExpansion(Entry, InIsExpanded);
+		}	
+	}
+}
+
 FReply SBindingsList::OnKeyDown(const FGeometry& MyGeometry, const FKeyEvent& InKeyEvent)
 {
 	if (InKeyEvent.GetKey() == EKeys::Delete)
diff --git a/Engine/Plugins/Runtime/ModelViewViewModel/Source/ModelViewViewModelEditor/Private/Widgets/SMVVMViewBindingListView.h b/Engine/Plugins/Runtime/ModelViewViewModel/Source/ModelViewViewModelEditor/Private/Widgets/SMVVMViewBindingListView.h
--- a/Engine/Plugins/Runtime/ModelViewViewModel/Source/ModelViewViewModelEditor/Private/Widgets/SMVVMViewBindingListView.h
+++ b/Engine/Plugins/Runtime/ModelViewViewModel/Source/ModelViewViewModelEditor/Private/Widgets/SMVVMViewBindingListView.h
@@ -45,6 +45,8 @@ public:
 	void RequestNavigateToEvent(UMVVMBlueprintViewEvent* Event);
 	void RequestNavigateToCondition(UMVVMBlueprintViewCondition* Condition);
 
+	void SetExpansion(bool InIsExpanded);
+
 private:
 	TSharedRef<ITableRow> GenerateEntryRow(TSharedPtr<FBindingEntry> Entry, const TSharedRef<STableViewBase>& OwnerTable) const;
 	void OnSourceListSelectionChanged(TSharedPtr<FBindingEntry> Entry, ESelectInfo::Type SelectionType) const;
diff --git a/Engine/Plugins/Runtime/ModelViewViewModel/Source/ModelViewViewModelEditor/Private/Widgets/SMVVMViewBindingPanel.cpp b/Engine/Plugins/Runtime/ModelViewViewModel/Source/ModelViewViewModelEditor/Private/Widgets/SMVVMViewBindingPanel.cpp
--- a/Engine/Plugins/Runtime/ModelViewViewModel/Source/ModelViewViewModelEditor/Private/Widgets/SMVVMViewBindingPanel.cpp
+++ b/Engine/Plugins/Runtime/ModelViewViewModel/Source/ModelViewViewModelEditor/Private/Widgets/SMVVMViewBindingPanel.cpp
@@ -8,6 +8,7 @@
 #include "MVVMBlueprintView.h"
 #include "MVVMDeveloperProjectSettings.h"
 #include "MVVMEditorSubsystem.h"
+#include "MVVMViewBindingMenuContext.h"
 #include "MVVMWidgetBlueprintExtension_View.h"
 #include "WidgetBlueprintEditor.h"
 #include "MVVMBlueprintViewCondition.h"
@@ -540,10 +541,11 @@ void SBindingsPanel::HandleExtensionAdded(UBlueprintExtension* NewExtension)
 
 TSharedRef<SWidget> SBindingsPanel::GenerateSettingsMenu()
 {
-	UWidgetBlueprintToolMenuContext* WidgetBlueprintMenuContext = NewObject<UWidgetBlueprintToolMenuContext>();
-	WidgetBlueprintMenuContext->WidgetBlueprintEditor = WeakBlueprintEditor;
+	UMVVMViewBindingMenuContext* BindingMenuContext = NewObject<UMVVMViewBindingMenuContext>();
+	BindingMenuContext->WidgetBlueprintEditor = WeakBlueprintEditor;
+	BindingMenuContext->BindingsPanel = SharedThis(this);
 
-	FToolMenuContext MenuContext(WidgetBlueprintMenuContext);
+	FToolMenuContext MenuContext(BindingMenuContext);
 	return UToolMenus::Get()->GenerateWidget("MVVM.ViewBindings.Toolbar", MenuContext);
 }
 
@@ -554,7 +556,7 @@ void SBindingsPanel::RegisterSettingsMenu()
 	FToolMenuSection& Section = Menu->FindOrAddSection("Settings");
 	Section.AddDynamicEntry("Settings", FNewToolMenuSectionDelegate::CreateLambda([](FToolMenuSection& InSection)
 		{
-			if (const UWidgetBlueprintToolMenuContext* Context = InSection.FindContext<UWidgetBlueprintToolMenuContext>())
+			if (const UMVVMViewBindingMenuContext* Context = InSection.FindContext<UMVVMViewBindingMenuContext>())
 			{
 				if (GetDefault<UMVVMDeveloperProjectSettings>()->bShowViewSettings)
 				{
@@ -567,6 +569,29 @@ void SBindingsPanel::RegisterSettingsMenu()
 						, EUserInterfaceActionType::Button
 					);
 				}
+
+				if (TSharedPtr<SBindingsPanel> BindingsPanel = Context->BindingsPanel.Pin())
+				{
+					TSharedRef<SBindingsList> BindingsListRef = BindingsPanel->BindingsList.ToSharedRef();
+
+					InSection.AddMenuEntry(
+						"ExpandAllBindings"
+						, LOCTEXT("ExpandAll", "Expand All")
+						, LOCTEXT("ExpandAllTooltip", "Expand All View Bindings")
+						, FSlateIcon()
+						, FUIAction(FExecuteAction::CreateSP(BindingsListRef, &SBindingsList::SetExpansion, true))
+						, EUserInterfaceActionType::Button
+					);
+
+					InSection.AddMenuEntry(
+						"CollapseAllBindings"
+						, LOCTEXT("CollapseAll", "Collapse All")
+						, LOCTEXT("CollapseAllTooltip", "Collapse All View Bindings")
+						, FSlateIcon()
+						, FUIAction(FExecuteAction::CreateSP(BindingsListRef, &SBindingsList::SetExpansion, false))
+						, EUserInterfaceActionType::Button
+					);
+				}
 			}
 		}));
 }
