Prevent assertion from shared pointer to shared ref when SBindingsPanel's  BindingList is null

#rnx

diff --git a/Engine/Plugins/Runtime/ModelViewViewModel/Source/ModelViewViewModelEditor/Private/Widgets/SMVVMViewBindingPanel.cpp b/Engine/Plugins/Runtime/ModelViewViewModel/Source/ModelViewViewModelEditor/Private/Widgets/SMVVMViewBindingPanel.cpp
--- a/Engine/Plugins/Runtime/ModelViewViewModel/Source/ModelViewViewModelEditor/Private/Widgets/SMVVMViewBindingPanel.cpp
+++ b/Engine/Plugins/Runtime/ModelViewViewModel/Source/ModelViewViewModelEditor/Private/Widgets/SMVVMViewBindingPanel.cpp
@@ -572,6 +572,11 @@ void SBindingsPanel::RegisterSettingsMenu()
 
 				if (TSharedPtr<SBindingsPanel> BindingsPanel = Context->BindingsPanel.Pin())
 				{
+					if (!BindingsPanel->BindingsList.IsValid())
+					{
+						return;
+					}
+
 					TSharedRef<SBindingsList> BindingsListRef = BindingsPanel->BindingsList.ToSharedRef();
 
 					InSection.AddMenuEntry(
