Adding a message to the ensure when function is a null ptr to help diagnose where it is coming from

#tests Debugged manually to test that branch of code
#rnx

diff --git a/Engine/Plugins/Runtime/ModelViewViewModel/Source/ModelViewViewModelBlueprint/Private/Bindings/MVVMConversionFunctionHelper.cpp b/Engine/Plugins/Runtime/ModelViewViewModel/Source/ModelViewViewModelBlueprint/Private/Bindings/MVVMConversionFunctionHelper.cpp
--- a/Engine/Plugins/Runtime/ModelViewViewModel/Source/ModelViewViewModelBlueprint/Private/Bindings/MVVMConversionFunctionHelper.cpp
+++ b/Engine/Plugins/Runtime/ModelViewViewModel/Source/ModelViewViewModelBlueprint/Private/Bindings/MVVMConversionFunctionHelper.cpp
@@ -597,7 +597,13 @@ namespace NamedNodes
 				}
 				else
 				{
-					ensure(false);
+					const TArrayView<const FMVVMBlueprintFieldPath> FieldPaths = PropertyPath.GetFieldPaths();
+					FName FunctionName;
+					if (FieldPaths.IsValidIndex(Index))
+					{
+						FunctionName = FieldPaths[Index].GetRawFieldName();
+					} 
+					ensureMsgf(false, TEXT("A function in path couldn't be resolved. Blueprint: %s Function: %s"), *Blueprint->GetName(), *FunctionName.ToString());
 					return MakeError();
 				}
 			}
