More modularity to Bindings list collapse/expand controls

#rnx

diff --git a/Engine/Plugins/Runtime/ModelViewViewModel/Source/ModelViewViewModelEditor/Private/Widgets/SMVVMViewBindingListView.cpp b/Engine/Plugins/Runtime/ModelViewViewModel/Source/ModelViewViewModelEditor/Private/Widgets/SMVVMViewBindingListView.cpp
--- a/Engine/Plugins/Runtime/ModelViewViewModel/Source/ModelViewViewModelEditor/Private/Widgets/SMVVMViewBindingListView.cpp
+++ b/Engine/Plugins/Runtime/ModelViewViewModel/Source/ModelViewViewModelEditor/Private/Widgets/SMVVMViewBindingListView.cpp
@@ -734,15 +734,40 @@ void SBindingsList::RequestNavigateToCondition(UMVVMBlueprintViewCondition* Cond
 	}
 }
 
-void SBindingsList::SetExpansion(bool InIsExpanded)
+void SBindingsList::SetRootGroupsExpansion(bool InIsExpanded)
 {
 	if (TreeView.IsValid())
 	{
-		// Set expansion only to the binding widgets
 		for (const TSharedPtr<FBindingEntry>& Entry : TreeView->GetRootItems())
 		{
 			TreeView->SetItemExpansion(Entry, InIsExpanded);
-		}	
+		}
+	}
+}
+
+void SBindingsList::SetBindingsExpansion(bool InIsExpanded)
+{
+	if (TreeView.IsValid())
+	{
+		TArray< TSharedPtr<FBindingEntry>> EntriesToVisit;
+		for (const TSharedPtr<FBindingEntry>& Entry : TreeView->GetRootItems())
+		{
+			// Ensure expanded binding is visible by expanding root group
+			if (InIsExpanded)
+			{
+				TreeView->SetItemExpansion(Entry, true);
+			}
+
+			EntriesToVisit.Append(Entry->GetAllChildren());
+		}
+
+		while (!EntriesToVisit.IsEmpty())
+		{
+			const TSharedPtr<FBindingEntry>& Entry = EntriesToVisit.Pop();
+			
+			TreeView->SetItemExpansion(Entry, InIsExpanded);
+			EntriesToVisit.Append(Entry->GetAllChildren());
+		}
 	}
 }
 
diff --git a/Engine/Plugins/Runtime/ModelViewViewModel/Source/ModelViewViewModelEditor/Private/Widgets/SMVVMViewBindingListView.h b/Engine/Plugins/Runtime/ModelViewViewModel/Source/ModelViewViewModelEditor/Private/Widgets/SMVVMViewBindingListView.h
--- a/Engine/Plugins/Runtime/ModelViewViewModel/Source/ModelViewViewModelEditor/Private/Widgets/SMVVMViewBindingListView.h
+++ b/Engine/Plugins/Runtime/ModelViewViewModel/Source/ModelViewViewModelEditor/Private/Widgets/SMVVMViewBindingListView.h
@@ -44,8 +44,9 @@ public:
 	void RequestNavigateToBinding(FGuid BindingId);
 	void RequestNavigateToEvent(UMVVMBlueprintViewEvent* Event);
 	void RequestNavigateToCondition(UMVVMBlueprintViewCondition* Condition);
-
-	void SetExpansion(bool InIsExpanded);
+	
+	void SetRootGroupsExpansion(bool InIsExpanded);
+	void SetBindingsExpansion(bool InIsExpanded);
 
 private:
 	TSharedRef<ITableRow> GenerateEntryRow(TSharedPtr<FBindingEntry> Entry, const TSharedRef<STableViewBase>& OwnerTable) const;
diff --git a/Engine/Plugins/Runtime/ModelViewViewModel/Source/ModelViewViewModelEditor/Private/Widgets/SMVVMViewBindingPanel.cpp b/Engine/Plugins/Runtime/ModelViewViewModel/Source/ModelViewViewModelEditor/Private/Widgets/SMVVMViewBindingPanel.cpp
--- a/Engine/Plugins/Runtime/ModelViewViewModel/Source/ModelViewViewModelEditor/Private/Widgets/SMVVMViewBindingPanel.cpp
+++ b/Engine/Plugins/Runtime/ModelViewViewModel/Source/ModelViewViewModelEditor/Private/Widgets/SMVVMViewBindingPanel.cpp
@@ -579,21 +579,39 @@ void SBindingsPanel::RegisterSettingsMenu()
 
 					TSharedRef<SBindingsList> BindingsListRef = BindingsPanel->BindingsList.ToSharedRef();
 
+					InSection.AddMenuEntry(
+						"ExpandAllCategories"
+						, LOCTEXT("ExpandAllCategories", "Expand All Categories")
+						, LOCTEXT("ExpandAllCategoriesTooltip", "Expands all root categories")
+						, FSlateIcon()
+						, FUIAction(FExecuteAction::CreateSP(BindingsListRef, &SBindingsList::SetRootGroupsExpansion, true))
+						, EUserInterfaceActionType::Button
+					);
+
+					InSection.AddMenuEntry(
+						"CollapseAllCategories"
+						, LOCTEXT("CollapseAllCategories", "Collapse All Categories")
+						, LOCTEXT("CollapseAllCategoriesTooltip", "Collapses all root categories")
+						, FSlateIcon()
+						, FUIAction(FExecuteAction::CreateSP(BindingsListRef, &SBindingsList::SetRootGroupsExpansion, false))
+						, EUserInterfaceActionType::Button
+					);
+
 					InSection.AddMenuEntry(
 						"ExpandAllBindings"
-						, LOCTEXT("ExpandAll", "Expand All")
-						, LOCTEXT("ExpandAllTooltip", "Expand All View Bindings")
+						, LOCTEXT("ExpandAllBindings", "Expand All Bindings")
+						, LOCTEXT("ExpandAllBindingsTooltip", "Expands all bindings in every category")
 						, FSlateIcon()
-						, FUIAction(FExecuteAction::CreateSP(BindingsListRef, &SBindingsList::SetExpansion, true))
+						, FUIAction(FExecuteAction::CreateSP(BindingsListRef, &SBindingsList::SetBindingsExpansion, true))
 						, EUserInterfaceActionType::Button
 					);
 
 					InSection.AddMenuEntry(
 						"CollapseAllBindings"
-						, LOCTEXT("CollapseAll", "Collapse All")
-						, LOCTEXT("CollapseAllTooltip", "Collapse All View Bindings")
+						, LOCTEXT("CollapseAllBindings", "Collapse All Bindings")
+						, LOCTEXT("CollapseAllBindingsTooltip", "Collapses all bindings in every category")
 						, FSlateIcon()
-						, FUIAction(FExecuteAction::CreateSP(BindingsListRef, &SBindingsList::SetExpansion, false))
+						, FUIAction(FExecuteAction::CreateSP(BindingsListRef, &SBindingsList::SetBindingsExpansion, false))
 						, EUserInterfaceActionType::Button
 					);
 				}
