Bulk Delete Bindings by supporting Multi Selection in Bindings List.

#rnx

diff --git a/Engine/Plugins/Runtime/ModelViewViewModel/Source/ModelViewViewModelEditor/Private/Widgets/SMVVMViewBindingListView.cpp b/Engine/Plugins/Runtime/ModelViewViewModel/Source/ModelViewViewModelEditor/Private/Widgets/SMVVMViewBindingListView.cpp
--- a/Engine/Plugins/Runtime/ModelViewViewModel/Source/ModelViewViewModelEditor/Private/Widgets/SMVVMViewBindingListView.cpp
+++ b/Engine/Plugins/Runtime/ModelViewViewModel/Source/ModelViewViewModelEditor/Private/Widgets/SMVVMViewBindingListView.cpp
@@ -200,7 +200,7 @@ void SBindingsList::Construct(const FArguments& InArgs, TSharedPtr<SBindingsPane
 	[
 		SAssignNew(TreeView, STreeView<TSharedPtr<FBindingEntry>>)
 		.TreeItemsSource(&FilteredRootGroups)
-		.SelectionMode(ESelectionMode::Single)
+		.SelectionMode(ESelectionMode::Multi)
 		.OnGenerateRow(this, &SBindingsList::GenerateEntryRow)
 		.OnGetChildren(this, &SBindingsList::GetChildrenOfEntry)
 		.OnContextMenuOpening(this, &SBindingsList::OnSourceConstructContextMenu)
diff --git a/Engine/Plugins/Runtime/ModelViewViewModel/Source/ModelViewViewModelEditor/Private/Widgets/SMVVMViewBindingPanel.cpp b/Engine/Plugins/Runtime/ModelViewViewModel/Source/ModelViewViewModelEditor/Private/Widgets/SMVVMViewBindingPanel.cpp
--- a/Engine/Plugins/Runtime/ModelViewViewModel/Source/ModelViewViewModelEditor/Private/Widgets/SMVVMViewBindingPanel.cpp
+++ b/Engine/Plugins/Runtime/ModelViewViewModel/Source/ModelViewViewModelEditor/Private/Widgets/SMVVMViewBindingPanel.cpp
@@ -212,13 +212,31 @@ void SBindingsPanel::HandleBlueprintViewChangedDelegate()
 void SBindingsPanel::OnBindingListSelectionChanged(TConstArrayView<FMVVMBlueprintViewBinding*> Selection)
 {
 	NotifyHook->Binding = FMVVMBlueprintViewBinding();
-	if (FMVVMBlueprintViewBinding* Binding = Selection.Num() > 0 ? Selection.Last() : nullptr)
+
+	// Currently only supports single object selection
+	if (FMVVMBlueprintViewBinding* Binding = Selection.Num() == 1 ? Selection.Last() : nullptr)
 	{
 		NotifyHook->Binding = *Binding;
 		TSharedRef<FStructOnScope> StructScope = MakeShared<FStructOnScope>(FMVVMBlueprintViewBinding::StaticStruct(), reinterpret_cast<uint8*>(Binding));
 		StructDetailsView->SetStructureData(StructScope);
 		DetailContainer->SetContent(StructDetailsView->GetWidget().ToSharedRef());
 	}
+	else if(Selection.Num() > 1)
+	{
+		DetailsView->SetObject(nullptr);
+		StructDetailsView->SetStructureData(nullptr);
+		DetailContainer->SetContent(
+			SNew(SHorizontalBox)
+			+ SHorizontalBox::Slot()
+				.VAlign(VAlign_Center)
+				.HAlign(HAlign_Center)
+				.AutoWidth()
+				[
+					SNew(STextBlock)
+						.Text(LOCTEXT("DetailsNotAvailable", "Details are not available when multiple bindings are selected.\nSelect only one binding to view details"))
+				]
+		);
+	}
 	else
 	{
 		// show empty details view
