Bulk Removal fix raw pointer comparisons following array shrink

#rnx

diff --git a/Engine/Plugins/Runtime/ModelViewViewModel/Source/ModelViewViewModelEditor/Private/Framework/MVVMRowHelper.cpp b/Engine/Plugins/Runtime/ModelViewViewModel/Source/ModelViewViewModelEditor/Private/Framework/MVVMRowHelper.cpp
--- a/Engine/Plugins/Runtime/ModelViewViewModel/Source/ModelViewViewModelEditor/Private/Framework/MVVMRowHelper.cpp
+++ b/Engine/Plugins/Runtime/ModelViewViewModel/Source/ModelViewViewModelEditor/Private/Framework/MVVMRowHelper.cpp
@@ -36,7 +36,7 @@ extern UNREALED_API UEditorEngine* GEditor;
 namespace UE::MVVM::BindingEntry
 {
 
-void FRowHelper::GatherAllChildBindings(UMVVMBlueprintView* BlueprintView, const TConstArrayView<TSharedPtr<FBindingEntry>> Entries, TArray<const FMVVMBlueprintViewBinding*>& OutBindings, TArray<UMVVMBlueprintViewEvent*>& OutEvents, TArray<UMVVMBlueprintViewCondition*>& OutConditions)
+void FRowHelper::GatherAllChildBindings(UMVVMBlueprintView* BlueprintView, const TConstArrayView<TSharedPtr<FBindingEntry>> Entries, TArray<FGuid>& OutBindingIds, TArray<UMVVMBlueprintViewEvent*>& OutEvents, TArray<UMVVMBlueprintViewCondition*>& OutConditions)
 {
 	for (const TSharedPtr<FBindingEntry>& Entry : Entries)
 	{
@@ -45,7 +45,7 @@ void FRowHelper::GatherAllChildBindings(UMVVMBlueprintView* BlueprintView, const
 			const FMVVMBlueprintViewBinding* Binding = Entry->GetBinding(BlueprintView);
 			if (Binding != nullptr)
 			{
-				OutBindings.AddUnique(Binding);
+				OutBindingIds.AddUnique(Binding->BindingId);
 			}
 		}
 
@@ -66,7 +66,7 @@ void FRowHelper::GatherAllChildBindings(UMVVMBlueprintView* BlueprintView, const
 				OutConditions.AddUnique(Condition);
 			}
 		}
-		GatherAllChildBindings(BlueprintView, Entry->GetAllChildren(), OutBindings, OutEvents, OutConditions);
+		GatherAllChildBindings(BlueprintView, Entry->GetAllChildren(), OutBindingIds, OutEvents, OutConditions);
 	}
 }
 
@@ -74,20 +74,23 @@ void FRowHelper::DeleteEntries(const UWidgetBlueprint* WidgetBlueprint, UMVVMBlu
 {
 	if (WidgetBlueprint && BlueprintView)
 	{
-		TArray<const FMVVMBlueprintViewBinding*> BindingsToRemove;
+		TArray<FGuid> BindingIdsToRemove;
 		TArray<UMVVMBlueprintViewEvent*> EventsToRemove;
 		TArray<UMVVMBlueprintViewCondition*> ConditionsToRemove;
-		GatherAllChildBindings(BlueprintView, Selection, BindingsToRemove, EventsToRemove, ConditionsToRemove);
+		GatherAllChildBindings(BlueprintView, Selection, BindingIdsToRemove, EventsToRemove, ConditionsToRemove);
 
-		if (BindingsToRemove.Num() == 0 && EventsToRemove.Num() == 0 && ConditionsToRemove.Num() == 0)
+		if (BindingIdsToRemove.Num() == 0 && EventsToRemove.Num() == 0 && ConditionsToRemove.Num() == 0)
 		{
 			return;
 		}
 
 		TArray<FText> BindingDisplayNames;
-		for (const FMVVMBlueprintViewBinding* Binding : BindingsToRemove)
+		for (const FGuid& BindingId : BindingIdsToRemove)
 		{
-			BindingDisplayNames.Add(FText::FromString(Binding->GetDisplayNameString(WidgetBlueprint)));
+			if (FMVVMBlueprintViewBinding* Binding = BlueprintView->GetBinding(BindingId))
+			{
+				BindingDisplayNames.Add(FText::FromString(Binding->GetDisplayNameString(WidgetBlueprint)));
+			}
 		}
 		for (const UMVVMBlueprintViewEvent* Event : EventsToRemove)
 		{
@@ -110,10 +113,12 @@ void FRowHelper::DeleteEntries(const UWidgetBlueprint* WidgetBlueprint, UMVVMBlu
 			FScopedTransaction Transaction(LOCTEXT("DeleteBindingsTransaction", "Delete Bindings"));
 			BlueprintView->Modify();
 
-			for (int32 BindingIndex = BindingsToRemove.Num() - 1; BindingIndex > -1; BindingIndex--)
+			for (const FGuid& BindingId : BindingIdsToRemove)
 			{
-				const FMVVMBlueprintViewBinding* Binding = BindingsToRemove[BindingIndex];
-				BlueprintView->RemoveBinding(Binding);
+				if (FMVVMBlueprintViewBinding* Binding = BlueprintView->GetBinding(BindingId))
+				{
+					BlueprintView->RemoveBinding(Binding);
+				}
 			}
 			for (UMVVMBlueprintViewEvent* Event : EventsToRemove)
 			{
diff --git a/Engine/Plugins/Runtime/ModelViewViewModel/Source/ModelViewViewModelEditor/Private/Framework/MVVMRowHelper.h b/Engine/Plugins/Runtime/ModelViewViewModel/Source/ModelViewViewModelEditor/Private/Framework/MVVMRowHelper.h
--- a/Engine/Plugins/Runtime/ModelViewViewModel/Source/ModelViewViewModelEditor/Private/Framework/MVVMRowHelper.h
+++ b/Engine/Plugins/Runtime/ModelViewViewModel/Source/ModelViewViewModelEditor/Private/Framework/MVVMRowHelper.h
@@ -22,7 +22,7 @@ namespace UE::MVVM::BindingEntry
 
 struct FRowHelper
 {
-	static void GatherAllChildBindings(UMVVMBlueprintView* BlueprintView, const TConstArrayView<TSharedPtr<FBindingEntry>> Entries, TArray<const FMVVMBlueprintViewBinding*>& OutBindings, TArray<UMVVMBlueprintViewEvent*>& OutEvents, TArray<UMVVMBlueprintViewCondition*>& OutConditions);
+	static void GatherAllChildBindings(UMVVMBlueprintView* BlueprintView, const TConstArrayView<TSharedPtr<FBindingEntry>> Entries, TArray<FGuid>& OutBindingIds, TArray<UMVVMBlueprintViewEvent*>& OutEvents, TArray<UMVVMBlueprintViewCondition*>& OutConditions);
 
 	static void DeleteEntries(const UWidgetBlueprint* WidgetBlueprint, UMVVMBlueprintView* BlueprintView, TArrayView<const TSharedPtr<FBindingEntry>> Selection);
 	static void ShowBlueprintGraph(FBlueprintEditor* Editor, UWidgetBlueprint* WidgetBlueprint, UMVVMBlueprintView* BlueprintView, TArrayView<const TSharedPtr<FBindingEntry>> Selection);
