UI Component: support binding to component properties similar to regular properties of a widget
- fix a bug in MVVM where an event binding to a property inside an object inside the root widget wouldn't compile.
- Update and deprecate the GetFields() function in FieldIterator

diff --git a/Engine/Plugins/Experimental/AnimNext/Source/AnimNextEditor/Private/Param/ObjectFunctionLocatorEditor.cpp b/Engine/Plugins/Experimental/AnimNext/Source/AnimNextEditor/Private/Param/ObjectFunctionLocatorEditor.cpp
--- a/Engine/Plugins/Experimental/AnimNext/Source/AnimNextEditor/Private/Param/ObjectFunctionLocatorEditor.cpp
+++ b/Engine/Plugins/Experimental/AnimNext/Source/AnimNextEditor/Private/Param/ObjectFunctionLocatorEditor.cpp
@@ -30,7 +30,7 @@ class SObjectFunctionLocatorEditor : public SCompoundWidget
 
 	class FFieldIterator : public UE::PropertyViewer::IFieldIterator
 	{
-		virtual TArray<FFieldVariant> GetFields(const UStruct* InStruct) const override
+		virtual TArray<FFieldVariant> GetFields(const UStruct* InStruct, const FName InFieldName, const UStruct* InContainerStruct) const override
 		{
 			TArray<FFieldVariant> Result;
 			for (TFieldIterator<UFunction> FunctionIt(InStruct, EFieldIteratorFlags::IncludeSuper); FunctionIt; ++FunctionIt)
diff --git a/Engine/Plugins/Experimental/AnimNext/Source/AnimNextEditor/Private/Param/ObjectPropertyLocatorEditor.cpp b/Engine/Plugins/Experimental/AnimNext/Source/AnimNextEditor/Private/Param/ObjectPropertyLocatorEditor.cpp
--- a/Engine/Plugins/Experimental/AnimNext/Source/AnimNextEditor/Private/Param/ObjectPropertyLocatorEditor.cpp
+++ b/Engine/Plugins/Experimental/AnimNext/Source/AnimNextEditor/Private/Param/ObjectPropertyLocatorEditor.cpp
@@ -1,4 +1,4 @@
-﻿// Copyright Epic Games, Inc. All Rights Reserved.
+// Copyright Epic Games, Inc. All Rights Reserved.
 
 #include "ObjectPropertyLocatorEditor.h"
 
@@ -30,7 +30,7 @@ class SObjectPropertyLocatorEditor : public SCompoundWidget
 
 	class FFieldIterator : public UE::PropertyViewer::IFieldIterator
 	{
-		virtual TArray<FFieldVariant> GetFields(const UStruct* InStruct) const override
+		virtual TArray<FFieldVariant> GetFields(const UStruct* InStruct, const FName InFieldName, const UStruct* InContainerStruct) const override
 		{
 			auto HasNestedObjectProperties = [](const UStruct* InPropertyStruct, auto& InHasNestedObjectProperties)
 			{
diff --git a/Engine/Plugins/Experimental/AnimNext/Source/AnimNextUncookedOnly/Private/Variables/UniversalObjectLocatorBindingType.cpp b/Engine/Plugins/Experimental/AnimNext/Source/AnimNextUncookedOnly/Private/Variables/UniversalObjectLocatorBindingType.cpp
--- a/Engine/Plugins/Experimental/AnimNext/Source/AnimNextUncookedOnly/Private/Variables/UniversalObjectLocatorBindingType.cpp
+++ b/Engine/Plugins/Experimental/AnimNext/Source/AnimNextUncookedOnly/Private/Variables/UniversalObjectLocatorBindingType.cpp
@@ -493,7 +493,7 @@ TSharedRef<SWidget> FUniversalObjectLocatorBindingType::CreateEditWidget(const T
 				: FilterType(InFilterType)
 			{}
 		
-			virtual TArray<FFieldVariant> GetFields(const UStruct* InStruct) const
+			virtual TArray<FFieldVariant> GetFields(const UStruct* InStruct, const FName InFieldName, const UStruct* InContainerStruct) const
 			{
 				auto PassesFilterChecks = [this](const FProperty* InProperty)
 				{
diff --git a/Engine/Plugins/Runtime/ModelViewViewModel/Source/ModelViewViewModelBlueprint/Private/Bindings/MVVMConversionFunctionHelper.cpp b/Engine/Plugins/Runtime/ModelViewViewModel/Source/ModelViewViewModelBlueprint/Private/Bindings/MVVMConversionFunctionHelper.cpp
--- a/Engine/Plugins/Runtime/ModelViewViewModel/Source/ModelViewViewModelBlueprint/Private/Bindings/MVVMConversionFunctionHelper.cpp
+++ b/Engine/Plugins/Runtime/ModelViewViewModel/Source/ModelViewViewModelBlueprint/Private/Bindings/MVVMConversionFunctionHelper.cpp
@@ -455,7 +455,14 @@ namespace NamedNodes
 					break;
 				}
 				case EMVVMBlueprintFieldPathSource::Widget:
-					RootProperty = Blueprint->SkeletonGeneratedClass->FindPropertyByName(PropertyPath.GetWidgetName());
+					if (PropertyPath.IsComponent())
+					{
+						bCreateSelfNodeForRootProperty = true;
+					}
+					else
+					{
+						RootProperty = Blueprint->SkeletonGeneratedClass->FindPropertyByName(PropertyPath.GetWidgetName());
+					}
 					break;
 				default:
 					check(false);
@@ -933,7 +940,10 @@ namespace NamedNodes
 			int32 ArgumentIndex = SetterNode->Pins.IndexOfByKey(WrapperSelfPin);
 			const float PosX = SetterNode->NodePosX;
 			const float PosY = SetterNode->NodePosY + ArgumentIndex * 100.0f;
-			TValueOrError<TArray<UEdGraphPin*>, void> PropertyPathPinsResult = Private::BuildPropertyPath(Blueprint, CreateGraphInternalResult.FunctionGraph, PropertyPath, CanCreateSetterGraphResult.LocalContainerPathIndex, FVector2f(PosX, PosY));
+
+			EMVVMBlueprintFieldPathSource PropertyPathSource = PropertyPath.GetSource(Blueprint);
+			int32 NumberOfFieldExcludingThePropertyPathSource = PropertyPath.IsComponent() || PropertyPathSource == EMVVMBlueprintFieldPathSource::SelfContext ? CanCreateSetterGraphResult.LocalContainerPathIndex + 1 : CanCreateSetterGraphResult.LocalContainerPathIndex;
+			TValueOrError<TArray<UEdGraphPin*>, void> PropertyPathPinsResult = Private::BuildPropertyPath(Blueprint, CreateGraphInternalResult.FunctionGraph, PropertyPath, NumberOfFieldExcludingThePropertyPathSource, FVector2f(PosX, PosY));
 			if (PropertyPathPinsResult.HasError())
 			{
 				Private::RemoveGraph(Blueprint, bAddToBlueprint, CreateGraphInternalResult);
@@ -1481,7 +1491,7 @@ void SetPropertyPathForPin(const UBlueprint* Blueprint, const FMVVMBlueprintProp
 					return;
 				}
 
-				const int32 NumberOfFields = PropertyPath.GetFieldPaths().Num();
+				const int32 NumberOfFields = PropertyPath.IsComponent() ? PropertyPath.GetCompleteFields(Blueprint).Num() : PropertyPath.GetFieldPaths().Num();
 				const float PosX = K2ConversionNode->NodePosX;
 				const float PosY = K2ConversionNode->NodePosY + ArgumentIndex * 100;
 				TValueOrError<TArray<UEdGraphPin*>, void> BuildPropertyPathResult = Private::BuildPropertyPath(Blueprint, FunctionGraph, PropertyPath, NumberOfFields, FVector2f(PosX, PosY));
@@ -1529,7 +1539,7 @@ void SetPropertyPathForPin(const UBlueprint* Blueprint, const FMVVMBlueprintProp
 			return;
 		}
 
-		const int32 NumberOfFields = PropertyPath.GetFieldPaths().Num();
+		const int32 NumberOfFields = PropertyPath.IsComponent() ? PropertyPath.GetCompleteFields(Blueprint).Num() : PropertyPath.GetFieldPaths().Num();
 		const float PosX = ConversionNode->NodePosX;
 		const float PosY = ConversionNode->NodePosY + ArgumentIndex * 100;
 		TValueOrError<TArray<UEdGraphPin*>, void> BuildPropertyPathResult = Private::BuildPropertyPath(Blueprint, FunctionGraph, PropertyPath, NumberOfFields, FVector2f(PosX, PosY));
diff --git a/Engine/Plugins/Runtime/ModelViewViewModel/Source/ModelViewViewModelBlueprint/Private/MVVMPropertyPath.cpp b/Engine/Plugins/Runtime/ModelViewViewModel/Source/ModelViewViewModelBlueprint/Private/MVVMPropertyPath.cpp
--- a/Engine/Plugins/Runtime/ModelViewViewModel/Source/ModelViewViewModelBlueprint/Private/MVVMPropertyPath.cpp
+++ b/Engine/Plugins/Runtime/ModelViewViewModel/Source/ModelViewViewModelBlueprint/Private/MVVMPropertyPath.cpp
@@ -4,6 +4,7 @@
 #include "BlueprintCompilationManager.h"
 #include "Blueprint/WidgetTree.h"
 #include "Engine/Blueprint.h"
+#include "Extensions/UIComponentContainer.h"
 #include "IPlatformFilePak.h"
 #include "StructUtils/UserDefinedStruct.h"
 #include "Kismet2/BlueprintEditorUtils.h"
@@ -12,6 +13,7 @@
 #include "MVVMBlueprintViewModelContext.h"
 #include "MVVMWidgetBlueprintExtension_View.h"
 #include "Types/MVVMFieldVariant.h"
+#include "UIComponentWidgetBlueprintExtension.h"
 #include "WidgetBlueprint.h"
 
 #include UE_INLINE_GENERATED_CPP_BY_NAME(MVVMPropertyPath)
@@ -254,19 +256,26 @@ TArray<FName> FMVVMBlueprintPropertyPath::GetFieldNames(const UClass* InSelfCont
 
 TArray<UE::MVVM::FMVVMConstFieldVariant> FMVVMBlueprintPropertyPath::GetFields(const UClass* InSelfContext) const
 {
-	TArray<UE::MVVM::FMVVMConstFieldVariant> Result;
-	Result.Reserve(Paths.Num());
-
-	for (const FMVVMBlueprintFieldPath& Path : Paths)
+	const UWidgetBlueprintGeneratedClass* WidgetBPG = Cast<UWidgetBlueprintGeneratedClass>(InSelfContext);
+	if (bIsComponent && WidgetBPG)
 	{
-		Result.Add(Path.GetField(InSelfContext));
+		return GetComponentFields(InSelfContext);
 	}
+	else
+	{
+		TArray<UE::MVVM::FMVVMConstFieldVariant> Result;
+		Result.Reserve(Paths.Num());
+		for (const FMVVMBlueprintFieldPath& Path : Paths)
+		{
+			Result.Add(Path.GetField(InSelfContext));
+		}
 
-	return Result;
+		return Result;
+	}
 }
 
 
-TArray<UE::MVVM::FMVVMConstFieldVariant> FMVVMBlueprintPropertyPath::GetCompleteFields(const UBlueprint* InSelfContext) const
+TArray<UE::MVVM::FMVVMConstFieldVariant> FMVVMBlueprintPropertyPath::GetCompleteFields(const UBlueprint* InSelfContext, bool bIsForDisplay) const
 {
 	TArray<UE::MVVM::FMVVMConstFieldVariant> Result;
 	Result.Reserve(Paths.Num() + 1);
@@ -295,9 +304,12 @@ TArray<UE::MVVM::FMVVMConstFieldVariant> FMVVMBlueprintPropertyPath::GetComplete
 	case EMVVMBlueprintFieldPathSource::SelfContext:
 		break;
 	case EMVVMBlueprintFieldPathSource::Widget:
-	{
-		FProperty* WidgetProperty = ContextClass->FindPropertyByName(GetWidgetName());
-		Result.Add(UE::MVVM::FMVVMConstFieldVariant(WidgetProperty));
+	{	
+		if (!bIsComponent || bIsForDisplay)
+		{
+			FProperty* WidgetProperty = ContextClass->FindPropertyByName(GetWidgetName());
+			Result.Add(UE::MVVM::FMVVMConstFieldVariant(WidgetProperty));
+		}
 		break;
 	}
 	default:
@@ -305,14 +317,112 @@ TArray<UE::MVVM::FMVVMConstFieldVariant> FMVVMBlueprintPropertyPath::GetComplete
 		break;
 	}
 
-	for (const FMVVMBlueprintFieldPath& Path : Paths)
+	if (bIsForDisplay)
+	{
+		for (const FMVVMBlueprintFieldPath& Path : Paths)
+		{
+			Result.Add(Path.GetField(ContextClass));
+		}
+	}
+	else
 	{
-		Result.Add(Path.GetField(ContextClass));
+		Result.Append(GetFields(ContextClass));
 	}
 
 	return Result;
 }
 
+TArray<UE::MVVM::FMVVMConstFieldVariant> FMVVMBlueprintPropertyPath::GetComponentFields(const UClass* InSelfContext) const
+{
+	TArray<UE::MVVM::FMVVMConstFieldVariant> Result;
+
+	if (bIsComponent)
+	{
+		if (const UWidgetBlueprint* WidgetBP = Cast<UWidgetBlueprint>(InSelfContext->ClassGeneratedBy))
+		{
+			const FObjectPropertyBase* WidgetProperty = nullptr;
+			if (GetSource(WidgetBP) == EMVVMBlueprintFieldPathSource::Widget)
+			{
+				if (UClass* GeneratedClass = WidgetBP->SkeletonGeneratedClass ? WidgetBP->SkeletonGeneratedClass : WidgetBP->GeneratedClass)
+				{
+					WidgetProperty = CastField<const FObjectPropertyBase>(GeneratedClass->FindPropertyByName(GetWidgetName()));
+				}
+			}
+
+			const UWidgetBlueprint* CurrentWidgetBP = WidgetBP;
+			const UWidgetBlueprint* LastWidgetBP = nullptr;
+			if (WidgetProperty && WidgetProperty->PropertyClass->IsChildOf(UUserWidget::StaticClass()))
+			{
+				if (UWidgetBlueprintGeneratedClass* UserWidget = Cast<UWidgetBlueprintGeneratedClass>(WidgetProperty->PropertyClass))
+				{
+					LastWidgetBP = CurrentWidgetBP;
+					CurrentWidgetBP = Cast<UWidgetBlueprint>(UserWidget->ClassGeneratedBy);
+				}
+			}
+
+			const FObjectPropertyBase* OwningWidget = WidgetProperty;
+
+			for (const FMVVMBlueprintFieldPath& Path : Paths)
+			{
+				if (ensure(CurrentWidgetBP))
+				{
+					if (const FProperty* PropertyPtr = Path.GetField(InSelfContext).GetProperty())
+					{
+						if (const FObjectPropertyBase* ObjectProperty = CastField<const FObjectPropertyBase>(PropertyPtr))
+						{
+							if (ObjectProperty->PropertyClass->IsChildOf(UUserWidget::StaticClass()))
+							{
+								if (UWidgetBlueprintGeneratedClass* UserWidget = Cast<UWidgetBlueprintGeneratedClass>(ObjectProperty->PropertyClass))
+								{
+									LastWidgetBP = CurrentWidgetBP;
+									CurrentWidgetBP = Cast<UWidgetBlueprint>(UserWidget->ClassGeneratedBy);
+								}
+							}
+							if (ObjectProperty->PropertyClass->IsChildOf(UWidget::StaticClass()))
+							{
+								OwningWidget = ObjectProperty;
+							}
+							else if (ObjectProperty->PropertyClass->IsChildOf(UUIComponent::StaticClass()))
+							{
+								if (OwningWidget)
+								{
+									const UWidgetBlueprint* BlueprintOwningExtension = CurrentWidgetBP;
+									if (LastWidgetBP && OwningWidget->PropertyClass->IsChildOf(UUserWidget::StaticClass()) && OwningWidget->PropertyClass->ClassGeneratedBy == CurrentWidgetBP)
+									{
+										BlueprintOwningExtension = LastWidgetBP;
+									}
+									if (BlueprintOwningExtension)
+									{
+										if (UUIComponentWidgetBlueprintExtension* ComponentWidgetBlueprintExtension = UWidgetBlueprintExtension::GetExtension<UUIComponentWidgetBlueprintExtension>(BlueprintOwningExtension))
+										{
+											if (UWidget* Widget = BlueprintOwningExtension->WidgetTree->FindWidget(OwningWidget->GetFName()))
+											{
+												TArray<UUIComponent*> ComponentsList = ComponentWidgetBlueprintExtension->GetComponentsFor(Widget);
+												for (UUIComponent* Component : ComponentsList)
+												{
+													FName ComponentPropertyName = UUIComponentContainer::GetPropertyNameForComponent(Component, Widget->GetFName());
+
+													if (Result.Num() > 0 && ObjectProperty->PropertyClass == Component->GetClass() && ObjectProperty->GetFName() == ComponentPropertyName)
+													{
+														Result.Pop();
+														continue;
+													}
+												}
+											}
+										}
+									}
+								}
+							}
+						}
+					}
+				}
+				Result.Add(Path.GetField(InSelfContext));
+			}
+		}
+	}
+
+	return Result;
+}
 
 FString FMVVMBlueprintPropertyPath::GetPropertyPath(const UClass* InSelfContext) const
 {
diff --git a/Engine/Plugins/Runtime/ModelViewViewModel/Source/ModelViewViewModelBlueprint/Private/MVVMViewBlueprintCompiler.cpp b/Engine/Plugins/Runtime/ModelViewViewModel/Source/ModelViewViewModelBlueprint/Private/MVVMViewBlueprintCompiler.cpp
--- a/Engine/Plugins/Runtime/ModelViewViewModel/Source/ModelViewViewModelBlueprint/Private/MVVMViewBlueprintCompiler.cpp
+++ b/Engine/Plugins/Runtime/ModelViewViewModel/Source/ModelViewViewModelBlueprint/Private/MVVMViewBlueprintCompiler.cpp
@@ -1010,7 +1010,7 @@ void FMVVMViewBlueprintCompiler::CreateRequiredProperties(const FWidgetBlueprint
 		if (FieldPathSource == EMVVMBlueprintFieldPathSource::SelfContext)
 		{
 			const TArrayView<const FMVVMBlueprintFieldPath> FieldPaths = PropertyPath.GetFieldPaths();
-			if (FieldPaths.Num() >= 2)
+			if (FieldPaths.Num() >= 2 && !PropertyPath.IsComponent())
 			{
 				if (FieldPaths[0].GetBindingKind() == EBindingKind::Property && FieldPaths[0].IsFieldSelfContext())
 				{
@@ -4263,8 +4263,11 @@ TValueOrError<FMVVMViewBlueprintCompiler::FGetFieldsResult, FText> FMVVMViewBlue
 		check(WidgetNameToWidgetPointerMap.Contains(DestinationWidgetName));
 		checkf(!DestinationWidgetName.IsNone(), TEXT("The destination should have been checked and set bAreSourceContextsValid."));
 
-		const FCompilerUserWidgetProperty* FoundUserWidgetProperty = NeededUserWidgetProperties.FindByPredicate([DestinationWidgetName](const FCompilerUserWidgetProperty& Other) { return Other.Name == DestinationWidgetName; });
-		check(FoundUserWidgetProperty);
+		if (!PropertyPath.IsComponent())
+		{
+			const FCompilerUserWidgetProperty* FoundUserWidgetProperty = NeededUserWidgetProperties.FindByPredicate([DestinationWidgetName](const FCompilerUserWidgetProperty& Other) { return Other.Name == DestinationWidgetName; });
+			check(FoundUserWidgetProperty);
+		}
 
 		TValueOrError<TSharedPtr<FCompilerBindingSource>, FText> FindSourceResult = FindSource(DestinationWidgetName, FCompilerBindingSource::EType::Widget);
 		if (FindSourceResult.HasError())
@@ -4272,7 +4275,7 @@ TValueOrError<FMVVMViewBlueprintCompiler::FGetFieldsResult, FText> FMVVMViewBlue
 			return MakeError(FindSourceResult.StealError());
 		}
 
-		Result.GeneratedFields = AppendBaseField(Class, DestinationWidgetName, PropertyPathFields);
+		Result.GeneratedFields = PropertyPath.IsComponent() ? AppendBaseField(Class, FName(), PropertyPathFields) : AppendBaseField(Class, DestinationWidgetName, PropertyPathFields);
 		Result.OptionalSource = FindSourceResult.StealValue();
 		break;
 	}
diff --git a/Engine/Plugins/Runtime/ModelViewViewModel/Source/ModelViewViewModelBlueprint/Public/MVVMPropertyPath.h b/Engine/Plugins/Runtime/ModelViewViewModel/Source/ModelViewViewModelBlueprint/Public/MVVMPropertyPath.h
--- a/Engine/Plugins/Runtime/ModelViewViewModel/Source/ModelViewViewModelBlueprint/Public/MVVMPropertyPath.h
+++ b/Engine/Plugins/Runtime/ModelViewViewModel/Source/ModelViewViewModelBlueprint/Public/MVVMPropertyPath.h
@@ -125,6 +125,10 @@ private:
 	UPROPERTY(VisibleAnywhere, Category = "MVVM")
 	EMVVMBlueprintFieldPathSource Source = EMVVMBlueprintFieldPathSource::None;
 
+	/** True if this property path is for a UI component. This affects how some fields in the path are filtered out. */
+	UPROPERTY(VisibleAnywhere, Category = "MVVM")
+	bool bIsComponent = false;
+
 #if WITH_EDITORONLY_DATA
 	// Use the Paths. BindingReference and BindingKind are deprecated.
 	UPROPERTY()
@@ -163,7 +167,7 @@ public:
 	 * Get the binding fields. Resolves reference deprecation / redirectors.
 	 * Returns Viewmodel.Field.SubProperty.SubProperty from ViewModel.Field.SubProperty.SubProperty
 	 */
-	MODELVIEWVIEWMODELBLUEPRINT_API TArray<UE::MVVM::FMVVMConstFieldVariant> GetCompleteFields(const UBlueprint* SelfContext) const;
+	MODELVIEWVIEWMODELBLUEPRINT_API TArray<UE::MVVM::FMVVMConstFieldVariant> GetCompleteFields(const UBlueprint* SelfContext, bool bIsForDisplay = false) const;
 
 	/**
 	 * Get the full path without the first property name.
@@ -256,6 +260,16 @@ public:
 		Source = EMVVMBlueprintFieldPathSource::SelfContext;
 	}
 
+	void SetIsComponent(bool InIsComponent)
+	{
+		bIsComponent = InIsComponent;
+	}
+
+	bool IsComponent() const
+	{
+		return bIsComponent;
+	}
+
 	/**
 	 * Handle when a field that is part of the path is renamed
 	 * Returns true if the renamed field was part of the path
@@ -280,7 +294,8 @@ public:
 		return WidgetName == Other.WidgetName
 			&& ContextId == Other.ContextId
 			&& Source == Other.Source
-			&& Paths == Other.Paths;
+			&& Paths == Other.Paths
+			&& bIsComponent == Other.bIsComponent;
 	}
 
 	bool operator!=(const FMVVMBlueprintPropertyPath& Other) const
@@ -306,6 +321,7 @@ public:
 
 private:
 	MODELVIEWVIEWMODELBLUEPRINT_API void DeprecationUpdateSource(const UBlueprint* InContext);
+	TArray<UE::MVVM::FMVVMConstFieldVariant> GetComponentFields(const UClass* InSelfContext) const;
 };
 
 template<>
diff --git a/Engine/Plugins/Runtime/ModelViewViewModel/Source/ModelViewViewModelEditor/Private/Types/MVVMBindingSource.cpp b/Engine/Plugins/Runtime/ModelViewViewModel/Source/ModelViewViewModelEditor/Private/Types/MVVMBindingSource.cpp
--- a/Engine/Plugins/Runtime/ModelViewViewModel/Source/ModelViewViewModelEditor/Private/Types/MVVMBindingSource.cpp
+++ b/Engine/Plugins/Runtime/ModelViewViewModel/Source/ModelViewViewModelEditor/Private/Types/MVVMBindingSource.cpp
@@ -3,6 +3,7 @@
 #include "Types/MVVMBindingSource.h"
 #include "Blueprint/WidgetTree.h"
 #include "Editor.h"
+#include "Extensions/UIComponent.h"
 #include "MVVMBlueprintView.h"
 #include "MVVMBlueprintViewModelContext.h"
 #include "MVVMPropertyPath.h"
@@ -115,6 +116,20 @@ FBindingSource FBindingSource::CreateForWidget(const UWidgetBlueprint* WidgetBlu
 	return Source;
 }
 
+FBindingSource FBindingSource::CreateForWidgetComponent(const UWidgetBlueprint* WidgetBlueprint, FName WidgetName, UUIComponent* Component)
+{
+	FBindingSource Source;
+
+	Source.Source = EMVVMBlueprintFieldPathSource::Widget;
+	Source.WidgetName = WidgetName;
+	if (UWidget* Widget = WidgetBlueprint->WidgetTree->FindWidget(Source.WidgetName))
+	{
+		Source.DisplayName = Widget->GetLabelText();
+		Source.Class = Component->GetClass();
+	}
+	return Source;
+}
+
 FBindingSource FBindingSource::CreateForViewModel(const UWidgetBlueprint* WidgetBlueprint, FGuid ViewModelId)
 {
 	FBindingSource Source;
diff --git a/Engine/Plugins/Runtime/ModelViewViewModel/Source/ModelViewViewModelEditor/Private/Types/MVVMBindingSource.h b/Engine/Plugins/Runtime/ModelViewViewModel/Source/ModelViewViewModelEditor/Private/Types/MVVMBindingSource.h
--- a/Engine/Plugins/Runtime/ModelViewViewModel/Source/ModelViewViewModelEditor/Private/Types/MVVMBindingSource.h
+++ b/Engine/Plugins/Runtime/ModelViewViewModel/Source/ModelViewViewModelEditor/Private/Types/MVVMBindingSource.h
@@ -9,6 +9,7 @@
 struct FMVVMBindingName;
 struct FMVVMBlueprintPropertyPath;
 struct FMVVMBlueprintViewModelContext;
+class UUIComponent;
 class UWidget;
 class UWidgetBlueprint;
 
@@ -83,6 +84,7 @@ public:
 	static FBindingSource CreateForBlueprint(const UWidgetBlueprint* WidgetBlueprint);
 	static FBindingSource CreateForWidget(const UWidgetBlueprint* WidgetBlueprint, const UWidget* Widget);
 	static FBindingSource CreateForWidget(const UWidgetBlueprint* WidgetBlueprint, FName WidgetName);
+	static FBindingSource CreateForWidgetComponent(const UWidgetBlueprint* WidgetBlueprint, FName WidgetName, UUIComponent* Component);
 	static FBindingSource CreateForViewModel(const UWidgetBlueprint* WidgetBlueprint, FGuid ViewModelId);
 	static FBindingSource CreateForViewModel(const UWidgetBlueprint* WidgetBlueprint, FName ViewModelName);
 	static FBindingSource CreateForViewModel(const UWidgetBlueprint* WidgetBlueprint, const FMVVMBlueprintViewModelContext& ViewModelContext);
diff --git a/Engine/Plugins/Runtime/ModelViewViewModel/Source/ModelViewViewModelEditor/Private/Widgets/SMVVMPropertyPath.cpp b/Engine/Plugins/Runtime/ModelViewViewModel/Source/ModelViewViewModelEditor/Private/Widgets/SMVVMPropertyPath.cpp
--- a/Engine/Plugins/Runtime/ModelViewViewModel/Source/ModelViewViewModelEditor/Private/Widgets/SMVVMPropertyPath.cpp
+++ b/Engine/Plugins/Runtime/ModelViewViewModel/Source/ModelViewViewModelEditor/Private/Widgets/SMVVMPropertyPath.cpp
@@ -49,7 +49,7 @@ void SPropertyPath::SetPropertyPath(const FMVVMBlueprintPropertyPath& InProperty
 
 	if (Source != EMVVMBlueprintFieldPathSource::None)
 	{
-		TArray<UE::MVVM::FMVVMConstFieldVariant> AllFields = InPropertyPath.GetCompleteFields(WidgetBlueprintPtr);
+		TArray<UE::MVVM::FMVVMConstFieldVariant> AllFields = InPropertyPath.GetCompleteFields(WidgetBlueprintPtr, true);
 
 		// Find the NotifyFieldId to highlight it. Error are valid.
 		TOptional<FInt32Range> HighlightRange;
diff --git a/Engine/Plugins/Runtime/ModelViewViewModel/Source/ModelViewViewModelEditor/Private/Widgets/SMVVMViewModelBindingListWidget.cpp b/Engine/Plugins/Runtime/ModelViewViewModel/Source/ModelViewViewModelEditor/Private/Widgets/SMVVMViewModelBindingListWidget.cpp
--- a/Engine/Plugins/Runtime/ModelViewViewModel/Source/ModelViewViewModelEditor/Private/Widgets/SMVVMViewModelBindingListWidget.cpp
+++ b/Engine/Plugins/Runtime/ModelViewViewModel/Source/ModelViewViewModelEditor/Private/Widgets/SMVVMViewModelBindingListWidget.cpp
@@ -4,8 +4,10 @@
 
 #include "Bindings/MVVMBindingHelper.h"
 #include "Bindings/MVVMFieldPathHelper.h"
+#include "Blueprint/WidgetBlueprintGeneratedClass.h"
 #include "Blueprint/WidgetTree.h"
 #include "BlueprintEditorSettings.h"
+#include "Extensions/UIComponentContainer.h"
 #include "GraphEditorSettings.h"
 #include "MVVMBlueprintViewModelContext.h"
 #include "MVVMDeveloperProjectSettings.h"
@@ -15,6 +17,7 @@
 #include "Styling/MVVMEditorStyle.h"
 #include "Types/MVVMAvailableBinding.h"
 #include "Types/MVVMBindingSource.h"
+#include "UIComponentWidgetBlueprintExtension.h"
 #include "WidgetBlueprint.h"
 #include "Widgets/Images/SImage.h"
 #include "Widgets/Layout/SBox.h"
@@ -28,10 +31,10 @@ namespace UE::MVVM
 
 namespace Private
 {
-	TOptional<FFieldVariant> PassFilter(const UBlueprint* Blueprint, const FMVVMAvailableBinding& Binding, const UStruct* Struct, const FMVVMFieldVariant& FieldVariant, EFieldVisibility FieldVisibilityFlags, const FProperty* AssignableTo, bool bDoObjectProperty);
-	TOptional<FFieldVariant> PassFilter(const UBlueprint* Blueprint, const FMVVMAvailableBinding& Binding, const UStruct* Struct, EFieldVisibility FieldVisibilityFlags, const FProperty* AssignableTo, bool bDoObjectProperty);
+	TOptional<FFieldVariant> PassFilter(const UBlueprint* Blueprint, const FMVVMAvailableBinding& Binding, const UStruct* Struct, const FMVVMFieldVariant& FieldVariant, EFieldVisibility FieldVisibilityFlags, const FProperty* AssignableTo, bool bDoObjectProperty, bool bAllowComponentProperty = false);
+	TOptional<FFieldVariant> PassFilter(const UBlueprint* Blueprint, const FMVVMAvailableBinding& Binding, const UStruct* Struct, EFieldVisibility FieldVisibilityFlags, const FProperty* AssignableTo, bool bDoObjectProperty, bool bAllowComponentProperty = false);
 
-	TOptional<FFieldVariant> PassFilter(const UBlueprint* Blueprint, const FMVVMAvailableBinding& Binding, const UStruct* Struct, EFieldVisibility FieldVisibilityFlags, const FProperty* AssignableTo, bool bDoObjectProperty)
+	TOptional<FFieldVariant> PassFilter(const UBlueprint* Blueprint, const FMVVMAvailableBinding& Binding, const UStruct* Struct, EFieldVisibility FieldVisibilityFlags, const FProperty* AssignableTo, bool bDoObjectProperty, bool bAllowComponentProperty)
 	{
 		if (EnumHasAllFlags(FieldVisibilityFlags, EFieldVisibility::Readable) && !Binding.IsReadable())
 		{
@@ -44,10 +47,10 @@ namespace Private
 		}
 
 		FMVVMFieldVariant FieldVariant = BindingHelper::FindFieldByName(Struct, Binding.GetBindingName());
-		return PassFilter(Blueprint, Binding, Struct, FieldVariant, FieldVisibilityFlags, AssignableTo, bDoObjectProperty);
+		return PassFilter(Blueprint, Binding, Struct, FieldVariant, FieldVisibilityFlags, AssignableTo, bDoObjectProperty, bAllowComponentProperty);
 	}
 
-	TOptional<FFieldVariant> PassFilter(const UBlueprint* Blueprint, const FMVVMAvailableBinding& Binding, const UStruct* Struct, const FMVVMFieldVariant& FieldVariant, EFieldVisibility FieldVisibilityFlags, const FProperty* AssignableTo, bool bDoObjectProperty)
+	TOptional<FFieldVariant> PassFilter(const UBlueprint* Blueprint, const FMVVMAvailableBinding& Binding, const UStruct* Struct, const FMVVMFieldVariant& FieldVariant, EFieldVisibility FieldVisibilityFlags, const FProperty* AssignableTo, bool bDoObjectProperty, bool bAllowComponentProperty)
 	{
 		static FName NAME_Hidden("Hidden");
 		if (ensure(!FieldVariant.IsEmpty()))
@@ -147,6 +150,13 @@ namespace Private
 					}
 				}
 
+				// Hide the component property on the userwidget. It should only be visible from the owning widget.
+				// The metadata helps us determine whether this is a component property generated by the compiler or not. We want to filter only those ones when bAllowComponentProperty is false.
+				if (Property->HasMetaData(UUIComponentWidgetBlueprintExtension::MD_ComponentVariable) && !bAllowComponentProperty)
+				{
+					return TOptional<FFieldVariant>();
+				}
+
 				if (!GetDefault<UMVVMDeveloperProjectSettings>()->IsPropertyAllowed(Blueprint, Struct, Property))
 				{
 					return TOptional<FFieldVariant>();
@@ -168,11 +178,11 @@ FFieldIterator_Bindable::FFieldIterator_Bindable(const UWidgetBlueprint* InWidge
 {
 }
 
-TArray<FFieldVariant> FFieldIterator_Bindable::GetFields(const UStruct* Struct) const
+TArray<FFieldVariant> FFieldIterator_Bindable::GetFields(const UStruct* Struct, const FName FieldName, const UStruct* Container) const
 {
 	TArray<FFieldVariant> Result;
 
-	auto AddResult = [this, &Result, Struct](const TArray<FMVVMAvailableBinding>& AvailableBindingsList)
+	auto AddResult = [this, &Result](const TArray<FMVVMAvailableBinding>& AvailableBindingsList, const UStruct* Struct, const bool bAllowComponentProperty)
 	{
 		Result.Reserve(AvailableBindingsList.Num());
 
@@ -191,11 +201,11 @@ TArray<FFieldVariant> FFieldIterator_Bindable::GetFields(const UStruct* Struct)
 			TOptional<FFieldVariant> PassResult;
 			if (FilterFlags == EFilterFlag::All)
 			{
-				PassResult = Private::PassFilter(WidgetBlueprint.Get(), Value, Struct, FieldVisibilityFlags, AssignableTo, false);
+				PassResult = Private::PassFilter(WidgetBlueprint.Get(), Value, Struct, FieldVisibilityFlags, AssignableTo, false, bAllowComponentProperty);
 			}
 			else
 			{
-				PassResult = Private::PassFilter(WidgetBlueprint.Get(), Value, Struct, EFieldVisibility::None, nullptr, false);
+				PassResult = Private::PassFilter(WidgetBlueprint.Get(), Value, Struct, EFieldVisibility::None, nullptr, false, bAllowComponentProperty);
 			}
 			if (PassResult.IsSet())
 			{
@@ -210,12 +220,12 @@ TArray<FFieldVariant> FFieldIterator_Bindable::GetFields(const UStruct* Struct)
 		TSubclassOf<UObject> AccessorClass = WidgetBlueprintPtr ? WidgetBlueprintPtr->GeneratedClass : nullptr;
 
 		TArray<FMVVMAvailableBinding> Bindings = bIsBindingToEvent ? UMVVMSubsystem::GetAvailableBindingsForEvent(const_cast<UClass*>(Class), AccessorClass) : UMVVMSubsystem::GetAvailableBindings(const_cast<UClass*>(Class), AccessorClass);
-		AddResult(Bindings);
+		AddResult(Bindings, Struct, false);
 	}
 	else if (const UScriptStruct* ScriptStruct = Cast<const UScriptStruct>(Struct))
 	{
 		TArray<FMVVMAvailableBinding> Bindings = UMVVMSubsystem::GetAvailableBindingsForStruct(ScriptStruct);
-		AddResult(Bindings);
+		AddResult(Bindings, Struct, false);
 	}
 
 	Result.Sort([](const FFieldVariant& A, const FFieldVariant& B)
@@ -249,6 +259,46 @@ TArray<FFieldVariant> FFieldIterator_Bindable::GetFields(const UStruct* Struct)
 			return A.GetFName().LexicalLess(B.GetFName());
 		});
 
+	// Add any component properties at the end.
+	if (const UWidgetBlueprintGeneratedClass* WidgetBPG = Cast<UWidgetBlueprintGeneratedClass>(Container))
+	{
+		if (const UWidgetBlueprint* ContainerWidgetBlueprintPtr = Cast<UWidgetBlueprint>(WidgetBPG->ClassGeneratedBy))
+		{
+			const UWidgetBlueprint* WidgetBlueprintPtr = WidgetBlueprint.Get();
+
+			TSubclassOf<UObject> AccessorClass = WidgetBlueprintPtr ? WidgetBlueprintPtr->GeneratedClass : nullptr;
+
+			if (const UClass* Class = Cast<const UClass>(Struct))
+			{
+				if (UWidget* Widget = ContainerWidgetBlueprintPtr->WidgetTree->FindWidget(FieldName))
+				{
+					if (UUIComponentWidgetBlueprintExtension* ComponentWidgetBlueprintExtension = UWidgetBlueprintExtension::GetExtension<UUIComponentWidgetBlueprintExtension>(ContainerWidgetBlueprintPtr))
+					{
+						TArray<FMVVMAvailableBinding> ComponentBindings = bIsBindingToEvent ? UMVVMSubsystem::GetAvailableBindingsForEvent(Cast<const UClass>(Container), AccessorClass) : UMVVMSubsystem::GetAvailableBindings(Cast<const UClass>(Container), AccessorClass);
+
+						TArray<UUIComponent*> ComponentsList = ComponentWidgetBlueprintExtension->GetComponentsFor(Widget);
+						for (UUIComponent* Component : ComponentsList)
+						{
+							const FMVVMAvailableBinding* FoundBinding = ComponentBindings.FindByPredicate(
+								[Component, Widget](const FMVVMAvailableBinding& Binding)
+								{
+									FName ComponentPropertyName = UUIComponentContainer::GetPropertyNameForComponent(Component, Widget->GetFName());
+									return Binding.GetBindingName().ToString() == ComponentPropertyName;
+								});
+							if (FoundBinding)
+							{
+								TArray <FMVVMAvailableBinding> Bindings;
+								Bindings.Reserve(1);
+								Bindings.Add(*FoundBinding);
+								AddResult(Bindings, Container, true);
+							}
+						}
+					}
+				}
+			}
+		}
+	}
+
 	return Result;
 }
 
@@ -479,6 +529,17 @@ void SSourceBindingList::AddSources(TArrayView<const FBindingSource> InSources)
 					if (Widget)
 					{
 						Handle = PropertyViewer->AddInstance(Widget);
+					
+						// Add a new source for each component on this widget
+						if (UUIComponentWidgetBlueprintExtension* ComponentWidgetBlueprintExtension = UWidgetBlueprintExtension::GetExtension<UUIComponentWidgetBlueprintExtension>(WidgetBlueprintPtr))
+						{
+							TArray<UUIComponent*> ComponentsList = ComponentWidgetBlueprintExtension->GetComponentsFor(Widget);
+							for (UUIComponent* Component : ComponentsList)
+							{
+								SPropertyViewer::FHandle ComponentHandle = PropertyViewer->AddInstance(Component);
+								Sources.Emplace(FBindingSource::CreateForWidgetComponent(WidgetBlueprint.Get(), Widget->GetFName(), Component), ComponentHandle);
+							}
+						}
 					}
 					else
 					{
@@ -491,6 +552,97 @@ void SSourceBindingList::AddSources(TArrayView<const FBindingSource> InSources)
 	}
 }
 
+bool IsPropertyPathComponent(TArrayView<const FFieldVariant> FieldPath, FMVVMBlueprintPropertyPath& PropertyPath, const UWidgetBlueprint* WidgetBlueprintPtr)
+{
+	if (ensure(WidgetBlueprintPtr))
+	{
+		const FObjectPropertyBase* WidgetProperty = nullptr;
+		if (PropertyPath.GetSource(WidgetBlueprintPtr) == EMVVMBlueprintFieldPathSource::Widget)
+		{
+			if (UClass* GeneratedClass = WidgetBlueprintPtr->SkeletonGeneratedClass ? WidgetBlueprintPtr->SkeletonGeneratedClass : WidgetBlueprintPtr->GeneratedClass)
+			{
+				WidgetProperty = CastField<const FObjectPropertyBase>(GeneratedClass->FindPropertyByName(PropertyPath.GetWidgetName()));
+			}
+		}
+
+		const UWidgetBlueprint* CurrentWidgetBP = WidgetBlueprintPtr;
+		const UWidgetBlueprint* LastWidgetBP = nullptr;
+		if (WidgetProperty && WidgetProperty->PropertyClass->IsChildOf(UUserWidget::StaticClass()))
+		{
+			if (UWidgetBlueprintGeneratedClass* UserWidget = Cast<UWidgetBlueprintGeneratedClass>(WidgetProperty->PropertyClass))
+			{
+				LastWidgetBP = CurrentWidgetBP;
+				CurrentWidgetBP = Cast<UWidgetBlueprint>(UserWidget->ClassGeneratedBy);
+			}
+		}
+
+		const FObjectPropertyBase* OwningWidget = WidgetProperty;
+
+		/* We need 3 pieces of information to decide if this property path includes a valid component that is registered with a UIComponentWidgetBlueprintExtension:
+			1. CurrentWidgetBP: The widget blueprint that holds a component extension. 
+			There might be multiple of these in the path, we only care about the last one we encounter before finding the component itself.
+			We need to use LastWidgetBP when the component object is on a userwidget.  
+			2. OwningWidget: The widget that owns this component. We care about the name of this widget so we can verify it with the widget blueprint extension.
+			3. The component itself. 
+		*/
+		for (const FFieldVariant& FieldVariant : FieldPath)
+		{
+			if (ensure(CurrentWidgetBP))
+			{
+				if (const FProperty* PropertyPtr = FieldVariant.Get<FProperty>())
+				{
+					if (const FObjectPropertyBase* ObjectProperty = CastField<const FObjectPropertyBase>(PropertyPtr))
+					{
+						if (ObjectProperty->PropertyClass->IsChildOf(UUserWidget::StaticClass()))
+						{
+							if (UWidgetBlueprintGeneratedClass* UserWidget = Cast<UWidgetBlueprintGeneratedClass>(ObjectProperty->PropertyClass))
+							{
+								LastWidgetBP = CurrentWidgetBP;
+								CurrentWidgetBP = Cast<UWidgetBlueprint>(UserWidget->ClassGeneratedBy);
+							}
+						}
+
+						if (ObjectProperty->PropertyClass->IsChildOf(UWidget::StaticClass()))
+						{
+							OwningWidget = ObjectProperty;
+						}
+						else if (ObjectProperty->PropertyClass->IsChildOf(UUIComponent::StaticClass()))
+						{
+							if (OwningWidget)
+							{
+								const UWidgetBlueprint* BlueprintOwningExtension = CurrentWidgetBP;
+								if (LastWidgetBP && OwningWidget->PropertyClass->IsChildOf(UUserWidget::StaticClass()) && OwningWidget->PropertyClass->ClassGeneratedBy == CurrentWidgetBP)
+								{
+									BlueprintOwningExtension = LastWidgetBP;
+								}
+
+								if (UUIComponentWidgetBlueprintExtension* ComponentWidgetBlueprintExtension = UWidgetBlueprintExtension::GetExtension<UUIComponentWidgetBlueprintExtension>(BlueprintOwningExtension))
+								{
+									if (UWidget* Widget = BlueprintOwningExtension->WidgetTree->FindWidget(OwningWidget->GetFName()))
+									{
+										TArray<UUIComponent*> ComponentsList = ComponentWidgetBlueprintExtension->GetComponentsFor(Widget);
+										for (UUIComponent* Component : ComponentsList)
+										{
+											check(Component);
+											FName ComponentPropertyName = UUIComponentContainer::GetPropertyNameForComponent(Component, Widget->GetFName());
+											if (ObjectProperty->PropertyClass == Component->GetClass() && ObjectProperty->GetFName() == ComponentPropertyName)
+											{
+												return true;
+											}
+										}
+									}
+								}
+							}
+						}
+					}
+				}
+			}
+		}
+	}
+
+	return false;
+}
+
 FMVVMBlueprintPropertyPath SSourceBindingList::CreateBlueprintPropertyPath(SPropertyViewer::FHandle Handle, TArrayView<const FFieldVariant> FieldPath) const
 {
 	if (FieldPath.Num() < 0 || !Handle.IsValid())
@@ -516,6 +668,17 @@ FMVVMBlueprintPropertyPath SSourceBindingList::CreateBlueprintPropertyPath(SProp
 		return FMVVMBlueprintPropertyPath();
 	}
 
+	if (FieldPath.Num() > 0)
+	{
+		if (const FProperty* Property = FieldPath.Last().Get<FProperty>())
+		{
+			if (Property->HasMetaData(UUIComponentWidgetBlueprintExtension::MD_ComponentVariable))
+			{
+				return FMVVMBlueprintPropertyPath();
+			}
+		}
+	}
+
 	const UClass* AccessorClass = WidgetBlueprintPtr->SkeletonGeneratedClass ? WidgetBlueprintPtr->SkeletonGeneratedClass : WidgetBlueprintPtr->GeneratedClass;
 	FMVVMBlueprintPropertyPath PropertyPath;
 	if (FieldPath.Num() > 0)
@@ -574,10 +737,42 @@ FMVVMBlueprintPropertyPath SSourceBindingList::CreateBlueprintPropertyPath(SProp
 		{
 			Source->Key.SetSourceTo(PropertyPath);
 			PropertyPath.ResetPropertyPath();
+
+			// When the component is directly added to the property viewer, we will need to add it to the path manually since the given path doesn't include the container.
+			FFieldVariant ComponentField;
+			TArray<const FFieldVariant> CompleteFieldPath;
+			if (PropertyPath.GetSource(WidgetBlueprintPtr) == EMVVMBlueprintFieldPathSource::Widget)
+			{
+				if (UUIComponentWidgetBlueprintExtension* ComponentWidgetBlueprintExtension = UWidgetBlueprintExtension::GetExtension<UUIComponentWidgetBlueprintExtension>(WidgetBlueprintPtr))
+				{
+					if (UWidget* Widget = WidgetBlueprint->WidgetTree->FindWidget(Source->Key.GetWidgetName()))
+					{
+						TArray<UUIComponent*> ComponentsList = ComponentWidgetBlueprintExtension->GetComponentsFor(Widget);
+						for (UUIComponent* Component : ComponentsList)
+						{
+							if (Source->Key.GetClass() == Component->GetClass())
+							{
+								FName ComponentPropertyName = UUIComponentContainer::GetPropertyNameForComponent(Component, Widget->GetFName());
+								FProperty* Property = FindFProperty<FObjectPropertyBase>(Cast<UWidgetBlueprintGeneratedClass>(WidgetBlueprint->GeneratedClass), ComponentPropertyName);
+								ensure(Property);
+								FMVVMConstFieldVariant Field = FMVVMConstFieldVariant(Property);
+								PropertyPath.AppendPropertyPath(WidgetBlueprintPtr, Field);
+								CompleteFieldPath.Add(FFieldVariant(Property));
+								break;
+							}
+						}
+					}
+				}
+			}
+
 			for (const FFieldVariant& Field : FieldPath)
 			{
 				PropertyPath.AppendPropertyPath(WidgetBlueprintPtr, FMVVMConstFieldVariant(Field));
+				CompleteFieldPath.Add(Field);
 			}
+
+			// Determine if this property path holds a valid registered UI component.
+			PropertyPath.SetIsComponent(IsPropertyPathComponent(CompleteFieldPath, PropertyPath, WidgetBlueprintPtr));
 		}
 	}
 	else if(AccessorClass)
diff --git a/Engine/Plugins/Runtime/ModelViewViewModel/Source/ModelViewViewModelEditor/Private/Widgets/SMVVMViewModelBindingListWidget.h b/Engine/Plugins/Runtime/ModelViewViewModel/Source/ModelViewViewModelEditor/Private/Widgets/SMVVMViewModelBindingListWidget.h
--- a/Engine/Plugins/Runtime/ModelViewViewModel/Source/ModelViewViewModelEditor/Private/Widgets/SMVVMViewModelBindingListWidget.h
+++ b/Engine/Plugins/Runtime/ModelViewViewModel/Source/ModelViewViewModelEditor/Private/Widgets/SMVVMViewModelBindingListWidget.h
@@ -34,7 +34,7 @@ namespace UE::MVVM
 	{
 	public:
 		FFieldIterator_Bindable(const UWidgetBlueprint* WidgetBlueprint, EFieldVisibility InVisibilityFlags, const FProperty* InAssignableTo = nullptr, const bool InIsBindingToEvent = false);
-		virtual TArray<FFieldVariant> GetFields(const UStruct*) const override;
+		virtual TArray<FFieldVariant> GetFields(const UStruct* Struct, const FName FieldName, const UStruct* Container) const override;
 		EFieldVisibility GetFieldVisibilityFlags() const
 		{
 			return FieldVisibilityFlags;
diff --git a/Engine/Source/Editor/UMGEditor/Private/UIComponentWidgetBlueprintExtension.cpp b/Engine/Source/Editor/UMGEditor/Private/UIComponentWidgetBlueprintExtension.cpp
--- a/Engine/Source/Editor/UMGEditor/Private/UIComponentWidgetBlueprintExtension.cpp
+++ b/Engine/Source/Editor/UMGEditor/Private/UIComponentWidgetBlueprintExtension.cpp
@@ -7,8 +7,13 @@
 #include "Extensions/UIComponentWidgetBlueprintGeneratedClassExtension.h"
 #include "Extensions/UIComponentUserWidgetExtension.h"
 
-#include "WidgetBlueprintEditor.h"
 #include "Blueprint/WidgetTree.h"
+#include "Kismet2/BlueprintEditorUtils.h"
+#include "WidgetBlueprintEditor.h"
+
+#define LOCTEXT_NAMESPACE "UIComponentWidgetBlueprintExtension"
+
+const FName UUIComponentWidgetBlueprintExtension::MD_ComponentVariable = "GeneratedForComponent";
 
 UUIComponentWidgetBlueprintExtension::UUIComponentWidgetBlueprintExtension(const FObjectInitializer& ObjectInitializer)
 	: Super(ObjectInitializer)
@@ -35,7 +40,7 @@ void UUIComponentWidgetBlueprintExtension::RemoveComponent(const UClass* Compone
 	}
 }
 
-TArray<UUIComponent*> UUIComponentWidgetBlueprintExtension::GetComponentsFor(UWidget* Target) const
+TArray<UUIComponent*> UUIComponentWidgetBlueprintExtension::GetComponentsFor(const UWidget* Target) const
 {
 	TArray<UUIComponent*> Components;
 	if(ensure(Target))
@@ -149,6 +154,25 @@ void UUIComponentWidgetBlueprintExtension::HandleCleanAndSanitizeClass(UWidgetBl
 	}
 }
 
+void UUIComponentWidgetBlueprintExtension::HandlePopulateGeneratedVariables(const FWidgetBlueprintCompilerContext::FPopulateGeneratedVariablesContext& Context)
+{
+	ComponentContainer->ForEachComponentTarget([this, &Context](const FUIComponentTarget& ComponentTarget)
+		{
+			if (ensure(ComponentTarget.GetComponent()))
+			{
+				FBPVariableDescription ComponentVariableDesc;
+				ComponentVariableDesc.VarName = UUIComponentContainer::GetPropertyNameForComponent(ComponentTarget.GetComponent(), ComponentTarget.GetTargetName());
+				ComponentVariableDesc.VarGuid = FGuid::NewGuid();
+				ComponentVariableDesc.VarType = FEdGraphPinType(UEdGraphSchema_K2::PC_Object, NAME_None, ComponentTarget.GetComponent()->GetClass(), EPinContainerType::None, true, FEdGraphTerminalType());
+				ComponentVariableDesc.FriendlyName = ComponentTarget.GetComponent()->GetName();
+				ComponentVariableDesc.PropertyFlags = (CPF_BlueprintVisible | CPF_BlueprintReadOnly | CPF_Transient | CPF_RepSkip);
+				ComponentVariableDesc.SetMetaData(UUIComponentWidgetBlueprintExtension::MD_ComponentVariable, TEXT("true"));
+				ComponentVariableDesc.Category = FText::FromString(TEXT("Component"));
+				Context.AddGeneratedVariable(MoveTemp(ComponentVariableDesc));
+			}
+		});
+}
+
 void UUIComponentWidgetBlueprintExtension::HandleFinishCompilingClass(UWidgetBlueprintGeneratedClass* Class)
 {
 	Super::HandleFinishCompilingClass(Class);
@@ -193,4 +217,6 @@ bool UUIComponentWidgetBlueprintExtension::HandleValidateGeneratedClass(UWidgetB
 void UUIComponentWidgetBlueprintExtension::HandleEndCompilation()
 {
 	CompilerContext = nullptr;
-}
\ No newline at end of file
+}
+
+#undef LOCTEXT_NAMESPACE
\ No newline at end of file
diff --git a/Engine/Source/Editor/UMGEditor/Private/UIComponentWidgetBlueprintExtension.h b/Engine/Source/Editor/UMGEditor/Private/UIComponentWidgetBlueprintExtension.h
deleted file mode 100644
diff --git a/Engine/Source/Editor/UMGEditor/Public/UIComponentWidgetBlueprintExtension.h b/Engine/Source/Editor/UMGEditor/Public/UIComponentWidgetBlueprintExtension.h
new file mode 100644
--- /dev/null
+++ b/Engine/Source/Editor/UMGEditor/Public/UIComponentWidgetBlueprintExtension.h
@@ -0,0 +1,58 @@
+﻿// Copyright Epic Games, Inc. All Rights Reserved.
+
+#pragma once
+
+#include "CoreMinimal.h"
+#include "UObject/ObjectMacros.h"
+#include "UObject/Object.h"
+#include "WidgetBlueprintExtension.h"
+
+#include "Extensions/UIComponent.h"
+
+#include "UIComponentWidgetBlueprintExtension.generated.h"
+
+class UUIComponentUserWidgetExtension;
+class UUIComponentContainer;
+
+UCLASS()
+/***
+ * Extension to the Widget Blueprint that will save the Components and will be used in the editor for compilation.
+ */
+class UMGEDITOR_API UUIComponentWidgetBlueprintExtension : public UWidgetBlueprintExtension
+{
+	GENERATED_BODY()
+	
+	explicit UUIComponentWidgetBlueprintExtension(const FObjectInitializer& ObjectInitializer = FObjectInitializer());
+
+public:
+	static const FName MD_ComponentVariable;
+
+	UUIComponent* AddComponent(const UClass* ComponentClass, FName OwnerName);
+	void RemoveComponent(const UClass* ComponentClass, FName OwnerName);
+	TArray<UUIComponent*> GetComponentsFor(const UWidget* Target) const;
+	UUIComponent* GetComponent(const UClass* ComponentClass, FName OwnerName) const;
+
+	UUIComponentUserWidgetExtension* GetOrCreateExtension(UUserWidget* PreviewWidget);	
+	void RenameWidget(const FName& OldVarName, const FName& NewVarName);
+
+	bool VerifyContainer(UUserWidget* UserWidget) const;
+
+protected:
+	virtual void HandleBeginCompilation(FWidgetBlueprintCompilerContext& InCreationContext) override;	
+	virtual void HandleCleanAndSanitizeClass(UWidgetBlueprintGeneratedClass* ClassToClean, UObject* OldCDO) override;
+	virtual void HandlePopulateGeneratedVariables(const FWidgetBlueprintCompilerContext::FPopulateGeneratedVariablesContext& Context) override;
+	virtual void HandleFinishCompilingClass(UWidgetBlueprintGeneratedClass* Class) override;
+	virtual bool HandleValidateGeneratedClass(UWidgetBlueprintGeneratedClass* Class) override; 
+
+	virtual void HandleEndCompilation() override;
+
+private:	
+	UUIComponentContainer* DuplicateContainer(UObject* Outer) const;
+
+	FWidgetBlueprintCompilerContext* CompilerContext = nullptr;
+
+#if WITH_EDITORONLY_DATA	
+	UPROPERTY()
+	TObjectPtr<UUIComponentContainer> ComponentContainer;
+#endif //WITH_EDITORONLY_DATA
+};
\ No newline at end of file
diff --git a/Engine/Source/Runtime/AdvancedWidgets/Private/Framework/PropertyViewer/IFieldIterator.cpp b/Engine/Source/Runtime/AdvancedWidgets/Private/Framework/PropertyViewer/IFieldIterator.cpp
--- a/Engine/Source/Runtime/AdvancedWidgets/Private/Framework/PropertyViewer/IFieldIterator.cpp
+++ b/Engine/Source/Runtime/AdvancedWidgets/Private/Framework/PropertyViewer/IFieldIterator.cpp
@@ -11,7 +11,7 @@ namespace UE::PropertyViewer
 /**
  * FFieldIterator_BlueprintVisible
  */
-TArray<FFieldVariant> FFieldIterator_BlueprintVisible::GetFields(const UStruct* Struct) const
+TArray<FFieldVariant> FFieldIterator_BlueprintVisible::GetFields(const UStruct* Struct, const FName FieldName, const UStruct* ContainerStruct)  const
 {
 	TArray<FFieldVariant> Result;
 	for (TFieldIterator<FProperty> PropertyIt(Struct, EFieldIteratorFlags::IncludeSuper); PropertyIt; ++PropertyIt)
diff --git a/Engine/Source/Runtime/AdvancedWidgets/Private/Widgets/PropertyViewer/SPropertyViewerImpl.cpp b/Engine/Source/Runtime/AdvancedWidgets/Private/Widgets/PropertyViewer/SPropertyViewerImpl.cpp
--- a/Engine/Source/Runtime/AdvancedWidgets/Private/Widgets/PropertyViewer/SPropertyViewerImpl.cpp
+++ b/Engine/Source/Runtime/AdvancedWidgets/Private/Widgets/PropertyViewer/SPropertyViewerImpl.cpp
@@ -322,10 +322,13 @@ void FTreeNode::BuildChildNodesRecursive(IFieldIterator& FieldIterator, IFieldEx
 	ChildNodes.Reset();
 
 	const UStruct* ChildStructType = nullptr;
+	const UStruct* CurrentFieldOwner = nullptr;
+	FName CurrentFieldName = NAME_None;
 	if (bIsStructProperty || bIsObjectProperty)
 	{
 		check(!PropertyName.IsNone());
-		const UStruct* CurrentFieldOwner = FieldOwner.Get();
+		CurrentFieldName = PropertyName;
+		CurrentFieldOwner = FieldOwner.Get();
 		FProperty* PropertyPtr = CurrentFieldOwner ? CurrentFieldOwner->FindPropertyByName(PropertyName) : nullptr;
 		if (const FStructProperty* StructProperty = CastField<const FStructProperty>(PropertyPtr))
 		{
@@ -362,7 +365,8 @@ void FTreeNode::BuildChildNodesRecursive(IFieldIterator& FieldIterator, IFieldEx
 	}
 	else if (!FunctionName.IsNone())
 	{
-		const UStruct* CurrentFieldOwner = FieldOwner.Get();
+		CurrentFieldOwner = FieldOwner.Get();
+		CurrentFieldName = FunctionName;
 		UFunction* FunctionPtr = CurrentFieldOwner ? CastChecked<const UClass>(CurrentFieldOwner)->FindFunctionByName(FunctionName) : nullptr;
 		if (FunctionPtr)
 		{
@@ -376,11 +380,12 @@ void FTreeNode::BuildChildNodesRecursive(IFieldIterator& FieldIterator, IFieldEx
 	else if (const TSharedPtr<FContainer> ContainerPin = Container.Pin())
 	{
 		ChildStructType = ContainerPin->GetStruct();
+		CurrentFieldName = FName(ContainerPin->GetDisplayName().Get(FText::GetEmpty()).ToString());
 	}
 
 	if (ChildStructType)
 	{
-		for (const FFieldVariant& FieldIt : FieldIterator.GetFields(ChildStructType))
+		for (const FFieldVariant& FieldIt : FieldIterator.GetFields(ChildStructType, CurrentFieldName, CurrentFieldOwner))
 		{
 			if (const FProperty* PropertyIt = FieldIt.Get<FProperty>())
 			{
diff --git a/Engine/Source/Runtime/AdvancedWidgets/Public/Framework/PropertyViewer/IFieldIterator.h b/Engine/Source/Runtime/AdvancedWidgets/Public/Framework/PropertyViewer/IFieldIterator.h
--- a/Engine/Source/Runtime/AdvancedWidgets/Public/Framework/PropertyViewer/IFieldIterator.h
+++ b/Engine/Source/Runtime/AdvancedWidgets/Public/Framework/PropertyViewer/IFieldIterator.h
@@ -18,13 +18,25 @@ namespace UE::PropertyViewer
 class IFieldIterator
 {
 public:
-	virtual TArray<FFieldVariant> GetFields(const UStruct*) const = 0;
+	UE_DEPRECATED(5.6, "GetFields(const UStruct*) is deprecated. Please use the overload with three arguments: GetFields(const UStruct*, const FName, const UStruct*).")
+	virtual TArray<FFieldVariant> GetFields(const UStruct*) const 
+	{
+		return TArray<FFieldVariant>();
+	};
+
+PRAGMA_DISABLE_DEPRECATION_WARNINGS
+	virtual TArray<FFieldVariant> GetFields(const UStruct* Struct, const FName FieldName, const UStruct* ContainerStruct) const
+	{
+		return GetFields(Struct);
+	};
+PRAGMA_ENABLE_DEPRECATION_WARNINGS
+
 	virtual ~IFieldIterator() = default;
 };
 
 class FFieldIterator_BlueprintVisible : public IFieldIterator
 {
-	ADVANCEDWIDGETS_API virtual TArray<FFieldVariant> GetFields(const UStruct*) const override;
+	ADVANCEDWIDGETS_API virtual TArray<FFieldVariant> GetFields(const UStruct* Struct, const FName FieldName, const UStruct* ContainerStruct) const override;
 };
 
 } //namespace
diff --git a/Engine/Source/Runtime/UMG/Private/Extensions/UIComponentContainer.cpp b/Engine/Source/Runtime/UMG/Private/Extensions/UIComponentContainer.cpp
--- a/Engine/Source/Runtime/UMG/Private/Extensions/UIComponentContainer.cpp
+++ b/Engine/Source/Runtime/UMG/Private/Extensions/UIComponentContainer.cpp
@@ -172,6 +172,11 @@ bool UUIComponentContainer::IsEmpty() const
 	return Components.IsEmpty();
 }
 
+FName UUIComponentContainer::GetPropertyNameForComponent(const UUIComponent* Component, const FName& TargetName)
+{
+	return FName(Component->GetName() + "_" + TargetName.ToString());
+}
+
 #ifdef WITH_EDITOR
 void UUIComponentContainer::RenameWidget(const FName& OldName, const FName& NewName)
 {
diff --git a/Engine/Source/Runtime/UMG/Private/Extensions/UIComponentWidgetBlueprintGeneratedClassExtension.cpp b/Engine/Source/Runtime/UMG/Private/Extensions/UIComponentWidgetBlueprintGeneratedClassExtension.cpp
--- a/Engine/Source/Runtime/UMG/Private/Extensions/UIComponentWidgetBlueprintGeneratedClassExtension.cpp
+++ b/Engine/Source/Runtime/UMG/Private/Extensions/UIComponentWidgetBlueprintGeneratedClassExtension.cpp
@@ -3,11 +3,13 @@
 #include "Extensions/UIComponentWidgetBlueprintGeneratedClassExtension.h"
 
 #include "Blueprint/UserWidget.h"
+#include "Blueprint/WidgetBlueprintGeneratedClass.h"
 #include "CoreMinimal.h"
 #include "Blueprint/WidgetTree.h"
 #include "Extensions/UIComponent.h"
 #include "Extensions/UIComponentContainer.h"
 #include "Extensions/UIComponentUserWidgetExtension.h"
+#include "UObject/UnrealType.h"
 
 UUIComponentContainer* UUIComponentWidgetBlueprintGeneratedClassExtension::DuplicateContainer(UUserWidget* UserWidget) const
 {
@@ -98,4 +100,23 @@ void UUIComponentWidgetBlueprintGeneratedClassExtension::PreConstruct(UUserWidge
 			UserWidgetExtension->Rename(nullptr, UserWidget);
 		}
 	}
+
+	if (!IsDesignTime)
+	{
+		// Set the component object in the corresponding property on the userwidget.
+		ComponentContainer->ForEachComponentTarget([UserWidget](const FUIComponentTarget& ComponentTarget) {
+			FName ComponentPropertyName = UUIComponentContainer::GetPropertyNameForComponent(ComponentTarget.GetComponent(), ComponentTarget.GetTargetName());
+
+			FObjectPropertyBase* FoundPanelObjectProperty = FindFProperty<FObjectPropertyBase>(UserWidget->GetClass(), ComponentPropertyName);
+			bool bValidPanelObject = FoundPanelObjectProperty && FoundPanelObjectProperty->PropertyClass->IsChildOf(UUIComponent::StaticClass());
+			if (ensureAlwaysMsgf(bValidPanelObject, TEXT("The compiler should have added the property")))
+			{
+				if (ensure(FoundPanelObjectProperty->PropertyClass == ComponentTarget.GetComponent()->GetClass()))
+				{
+					FoundPanelObjectProperty->SetObjectPropertyValue_InContainer(UserWidget, ComponentTarget.GetComponent());
+				}
+			}
+			});
+
+	}
 }
diff --git a/Engine/Source/Runtime/UMG/Public/Extensions/UIComponentContainer.h b/Engine/Source/Runtime/UMG/Public/Extensions/UIComponentContainer.h
--- a/Engine/Source/Runtime/UMG/Public/Extensions/UIComponentContainer.h
+++ b/Engine/Source/Runtime/UMG/Public/Extensions/UIComponentContainer.h
@@ -63,6 +63,8 @@ public:
 	void InitializeComponents(const UUserWidget* UserWidget);
 	
 	bool IsEmpty() const;
+
+	static FName GetPropertyNameForComponent(const UUIComponent* Component, const FName& TargetName);
 #ifdef WITH_EDITOR
 	void RenameWidget(const FName& OldName, const FName& NewName);
 	void CleanupUIComponents(const UWidgetTree* WidgetTree);	
