Context Menu Action to Duplicate Binding Entries

[RN] minor creator/ui

diff --git a/Engine/Plugins/Runtime/ModelViewViewModel/Source/ModelViewViewModelBlueprint/Private/MVVMBlueprintView.cpp b/Engine/Plugins/Runtime/ModelViewViewModel/Source/ModelViewViewModelBlueprint/Private/MVVMBlueprintView.cpp
--- a/Engine/Plugins/Runtime/ModelViewViewModel/Source/ModelViewViewModelBlueprint/Private/MVVMBlueprintView.cpp
+++ b/Engine/Plugins/Runtime/ModelViewViewModel/Source/ModelViewViewModelBlueprint/Private/MVVMBlueprintView.cpp
@@ -209,6 +209,22 @@ void UMVVMBlueprintView::RemoveBinding(const FMVVMBlueprintViewBinding* Binding)
 	RemoveBindingAt(Index);
 }
 
+const FMVVMBlueprintViewBinding* UMVVMBlueprintView::DuplicateBinding(const FMVVMBlueprintViewBinding* Binding)
+{
+	FMVVMBlueprintViewBinding& NewBinding = Bindings.AddDefaulted_GetRef();
+
+	// Copy all existing binding properties except the BindingId
+	NewBinding = *Binding;
+	NewBinding.BindingId = FGuid::NewGuid();
+
+	OnBindingsAdded.Broadcast();
+	OnBindingsUpdated.Broadcast();
+
+	FBlueprintEditorUtils::MarkBlueprintAsModified(GetOuterUMVVMWidgetBlueprintExtension_View()->GetWidgetBlueprint());
+
+	return &NewBinding;
+}
+
 FMVVMBlueprintViewBinding& UMVVMBlueprintView::AddDefaultBinding()
 {
 	FMVVMBlueprintViewBinding& NewBinding = Bindings.AddDefaulted_GetRef();
@@ -257,6 +273,9 @@ UMVVMBlueprintViewEvent* UMVVMBlueprintView::AddDefaultEvent()
 
 	OnBindingsAdded.Broadcast();
 	OnBindingsUpdated.Broadcast();
+
+	FBlueprintEditorUtils::MarkBlueprintAsModified(GetOuterUMVVMWidgetBlueprintExtension_View()->GetWidgetBlueprint());
+
 	return Event;
 }
 
@@ -266,9 +285,23 @@ void UMVVMBlueprintView::RemoveEvent(UMVVMBlueprintViewEvent* Event)
 	{
 		Event->RemoveWrapperGraph();
 		OnBindingsUpdated.Broadcast();
+
+		FBlueprintEditorUtils::MarkBlueprintAsModified(GetOuterUMVVMWidgetBlueprintExtension_View()->GetWidgetBlueprint());
 	}
 }
 
+UMVVMBlueprintViewEvent* UMVVMBlueprintView::DuplicateEvent(UMVVMBlueprintViewEvent* Event)
+{
+	UMVVMBlueprintViewEvent* NewEvent = DuplicateObject<UMVVMBlueprintViewEvent>(Event, this);
+	Events.Add(NewEvent);
+
+	OnBindingsAdded.Broadcast();
+
+	FBlueprintEditorUtils::MarkBlueprintAsModified(GetOuterUMVVMWidgetBlueprintExtension_View()->GetWidgetBlueprint());
+
+	return NewEvent;
+}
+
 UMVVMBlueprintViewCondition* UMVVMBlueprintView::AddDefaultCondition()
 {
 	UMVVMBlueprintViewCondition* Condition = NewObject<UMVVMBlueprintViewCondition>(this);
@@ -276,6 +309,9 @@ UMVVMBlueprintViewCondition* UMVVMBlueprintView::AddDefaultCondition()
 
 	OnBindingsAdded.Broadcast();
 	OnBindingsUpdated.Broadcast();
+
+	FBlueprintEditorUtils::MarkBlueprintAsModified(GetOuterUMVVMWidgetBlueprintExtension_View()->GetWidgetBlueprint());
+
 	return Condition;
 }
 
@@ -285,9 +321,24 @@ void UMVVMBlueprintView::RemoveCondition(UMVVMBlueprintViewCondition* Condition)
 	{
 		Condition->RemoveWrapperGraph();
 		OnBindingsUpdated.Broadcast();
+
+		FBlueprintEditorUtils::MarkBlueprintAsModified(GetOuterUMVVMWidgetBlueprintExtension_View()->GetWidgetBlueprint());
 	}
 }
 
+UMVVMBlueprintViewCondition* UMVVMBlueprintView::DuplicateCondition(UMVVMBlueprintViewCondition* Condition)
+{
+	UMVVMBlueprintViewCondition* NewCondition = DuplicateObject<UMVVMBlueprintViewCondition>(Condition, this);
+	Conditions.Add(NewCondition);
+
+	OnBindingsAdded.Broadcast();
+	OnBindingsUpdated.Broadcast();
+
+	FBlueprintEditorUtils::MarkBlueprintAsModified(GetOuterUMVVMWidgetBlueprintExtension_View()->GetWidgetBlueprint());
+
+	return NewCondition;
+}
+
 TArray<FText> UMVVMBlueprintView::GetBindingMessages(FGuid Id, UE::MVVM::EBindingMessageType InMessageType) const
 {
 	TArray<FText> Results;
diff --git a/Engine/Plugins/Runtime/ModelViewViewModel/Source/ModelViewViewModelBlueprint/Public/MVVMBlueprintView.h b/Engine/Plugins/Runtime/ModelViewViewModel/Source/ModelViewViewModelBlueprint/Public/MVVMBlueprintView.h
--- a/Engine/Plugins/Runtime/ModelViewViewModel/Source/ModelViewViewModelBlueprint/Public/MVVMBlueprintView.h
+++ b/Engine/Plugins/Runtime/ModelViewViewModel/Source/ModelViewViewModelBlueprint/Public/MVVMBlueprintView.h
@@ -109,6 +109,8 @@ public:
 	FMVVMBlueprintViewBinding* FindBinding(const UWidget* Widget, const FProperty* Property);
 
 	void RemoveBinding(const FMVVMBlueprintViewBinding* Binding);
+	const FMVVMBlueprintViewBinding* DuplicateBinding(const FMVVMBlueprintViewBinding* Binding);
+
 	void RemoveBindingAt(int32 Index);
 
 	FMVVMBlueprintViewBinding& AddDefaultBinding();
@@ -135,6 +137,7 @@ public:
 
 	UMVVMBlueprintViewEvent* AddDefaultEvent();
 	void RemoveEvent(UMVVMBlueprintViewEvent* Event);
+	UMVVMBlueprintViewEvent* DuplicateEvent(UMVVMBlueprintViewEvent* Event);
 
 	TArrayView<TObjectPtr<UMVVMBlueprintViewEvent>> GetEvents()
 	{
@@ -148,6 +151,7 @@ public:
 
 	UMVVMBlueprintViewCondition* AddDefaultCondition();
 	void RemoveCondition(UMVVMBlueprintViewCondition* Condition);
+	UMVVMBlueprintViewCondition* DuplicateCondition(UMVVMBlueprintViewCondition* Condition);
 
 	TArrayView<TObjectPtr<UMVVMBlueprintViewCondition>> GetConditions()
 	{
diff --git a/Engine/Plugins/Runtime/ModelViewViewModel/Source/ModelViewViewModelEditor/Private/Framework/MVVMRowHelper.cpp b/Engine/Plugins/Runtime/ModelViewViewModel/Source/ModelViewViewModelEditor/Private/Framework/MVVMRowHelper.cpp
--- a/Engine/Plugins/Runtime/ModelViewViewModel/Source/ModelViewViewModelEditor/Private/Framework/MVVMRowHelper.cpp
+++ b/Engine/Plugins/Runtime/ModelViewViewModel/Source/ModelViewViewModelEditor/Private/Framework/MVVMRowHelper.cpp
@@ -132,6 +132,60 @@ void FRowHelper::DeleteEntries(const UWidgetBlueprint* WidgetBlueprint, UMVVMBlu
 	}
 }
 
+void FRowHelper::DuplicateEntries(const UWidgetBlueprint* WidgetBlueprint, UMVVMBlueprintView* BlueprintView, TArrayView<const TSharedPtr<FBindingEntry>> Selection, TArray<const TSharedPtr<FBindingEntry>>& OutSelection)
+{
+	if (WidgetBlueprint && BlueprintView)
+	{
+		TArray<FGuid> BindingIdsToDuplicate;
+		TArray<UMVVMBlueprintViewEvent*> EventsToDuplicate;
+		TArray<UMVVMBlueprintViewCondition*> ConditionsToDuplicate;
+		GatherAllChildBindings(BlueprintView, Selection, BindingIdsToDuplicate, EventsToDuplicate, ConditionsToDuplicate);
+
+		if (BindingIdsToDuplicate.Num() == 0 && EventsToDuplicate.Num() == 0 && ConditionsToDuplicate.Num() == 0)
+		{
+			return;
+		}
+ 
+		OutSelection.Reserve(Selection.Num());
+
+		FScopedTransaction Transaction(LOCTEXT("DuplicateBindingsTransaction", "Duplicate Bindings"));
+		BlueprintView->Modify();
+
+		for (const FGuid& BindingId : BindingIdsToDuplicate)
+		{
+			if (FMVVMBlueprintViewBinding* Binding = BlueprintView->GetBinding(BindingId))
+			{
+				const FMVVMBlueprintViewBinding* NewBinding = BlueprintView->DuplicateBinding(Binding);
+
+				TSharedPtr<FBindingEntry> NewEntry = MakeShared<FBindingEntry>();
+				NewEntry->SetBindingId(NewBinding->BindingId);
+
+				OutSelection.Add(NewEntry);
+			}
+		}
+
+		for (UMVVMBlueprintViewEvent* Event : EventsToDuplicate)
+		{
+			UMVVMBlueprintViewEvent* NewEvent = BlueprintView->DuplicateEvent(Event);
+
+			TSharedPtr<FBindingEntry> NewEntry = MakeShared<FBindingEntry>();
+			NewEntry->SetEvent(NewEvent);
+
+			OutSelection.Add(NewEntry);
+		}
+		
+		for (UMVVMBlueprintViewCondition* Condition : ConditionsToDuplicate)
+		{
+			UMVVMBlueprintViewCondition* NewCondition = BlueprintView->DuplicateCondition(Condition);
+
+			TSharedPtr<FBindingEntry> NewEntry = MakeShared<FBindingEntry>();
+			NewEntry->SetCondition(NewCondition);
+
+			OutSelection.Add(NewEntry);
+		}
+	}
+}
+
 void FRowHelper::ShowBlueprintGraph(FBlueprintEditor* BlueprintEditor, UWidgetBlueprint* WidgetBlueprint, UMVVMBlueprintView* BlueprintView, TArrayView<const TSharedPtr<FBindingEntry>> Entries)
 {
 	auto ShowGraph = [WidgetBlueprint, BlueprintEditor](UEdGraph* Graph)
@@ -306,6 +360,14 @@ void HandleDeleteEntries(UWidgetBlueprint* WidgetBlueprint, UMVVMBlueprintView*
 	FRowHelper::DeleteEntries(WidgetBlueprint, View, WeakToSharedPtr(Entries));
 }
 
+void HandleDuplicateEntries(UWidgetBlueprint* WidgetBlueprint, UMVVMBlueprintView* View, TSharedPtr<TArray<TWeakPtr<FBindingEntry>>> Entries, FOnContextMenuEntryCallback OnSelectionChanged)
+{
+	TArray<const TSharedPtr<FBindingEntry>> NewSelection;
+	FRowHelper::DuplicateEntries(WidgetBlueprint, View, WeakToSharedPtr(Entries), NewSelection);
+
+	OnSelectionChanged.ExecuteIfBound(NewSelection);
+}
+
 void HandleResetSelectedPin(UWidgetBlueprint* WidgetBlueprint, UMVVMBlueprintView* View, TSharedPtr<TArray<TWeakPtr<FBindingEntry>>> Entries)
 {
 	const UMVVMEditorSubsystem* EditorSubsystem = GEditor->GetEditorSubsystem<UMVVMEditorSubsystem>();
@@ -388,7 +450,7 @@ void HandleResetOrphanedSelectedPin(UWidgetBlueprint* WidgetBlueprint, UMVVMBlue
 
 } //namespace
 
-FMenuBuilder FRowHelper::CreateContextMenu(UWidgetBlueprint* WidgetBlueprint, UMVVMBlueprintView* View, TConstArrayView<TSharedPtr<FBindingEntry>> Entries)
+FMenuBuilder FRowHelper::CreateContextMenu(UWidgetBlueprint* WidgetBlueprint, UMVVMBlueprintView* View, TConstArrayView<TSharedPtr<FBindingEntry>> Entries, FOnContextMenuEntryCallback OnSelectionChanged)
 {
 	const bool bShouldCloseWindowAfterMenuSelection = true;
 	FMenuBuilder MenuBuilder(bShouldCloseWindowAfterMenuSelection, nullptr);
@@ -406,17 +468,21 @@ FMenuBuilder FRowHelper::CreateContextMenu(UWidgetBlueprint* WidgetBlueprint, UM
 		const UMVVMEditorSubsystem* EditorSubsystem = GEditor->GetEditorSubsystem<UMVVMEditorSubsystem>();
 		{
 			bool bCanRemoveEntry = true;
+			bool bCanDuplicateEntry = true;
 			for (const TSharedPtr<FBindingEntry>& Entry : Entries)
 			{
 				switch (Entry->GetRowType())
 				{
 				case FBindingEntry::ERowType::Group:
+					bCanDuplicateEntry = false;
+					break; 
 				case FBindingEntry::ERowType::Binding:
 				case FBindingEntry::ERowType::Event:
 				case FBindingEntry::ERowType::Condition:
 					break;
 				default:
 					bCanRemoveEntry = false;
+					bCanDuplicateEntry = false;
 					break;
 				}
 			}
@@ -428,6 +494,14 @@ FMenuBuilder FRowHelper::CreateContextMenu(UWidgetBlueprint* WidgetBlueprint, UM
 				LOCTEXT("RemoveBindingTooltip", "Remove bindings or events."),
 				FSlateIcon(FAppStyle::GetAppStyleSetName(), "Icons.Delete"),
 				RemoveAction);
+
+			FUIAction DuplicateAction;
+			DuplicateAction.ExecuteAction = FExecuteAction::CreateStatic(&Private::HandleDuplicateEntries, WidgetBlueprint, View, CopiedEntries, OnSelectionChanged);
+			DuplicateAction.CanExecuteAction = FCanExecuteAction::CreateLambda([bCanDuplicateEntry]() { return bCanDuplicateEntry; });
+			MenuBuilder.AddMenuEntry(LOCTEXT("DuplicateBinding", "Duplicate"),
+				LOCTEXT("DuplicateBindingTooltip", "Duplicate bindings or events."),
+				FSlateIcon(FAppStyle::GetAppStyleSetName(), "Icons.Duplicate"),
+				DuplicateAction);
 		}
 
 		{
diff --git a/Engine/Plugins/Runtime/ModelViewViewModel/Source/ModelViewViewModelEditor/Private/Framework/MVVMRowHelper.h b/Engine/Plugins/Runtime/ModelViewViewModel/Source/ModelViewViewModelEditor/Private/Framework/MVVMRowHelper.h
--- a/Engine/Plugins/Runtime/ModelViewViewModel/Source/ModelViewViewModelEditor/Private/Framework/MVVMRowHelper.h
+++ b/Engine/Plugins/Runtime/ModelViewViewModel/Source/ModelViewViewModelEditor/Private/Framework/MVVMRowHelper.h
@@ -17,20 +17,25 @@ class UMVVMBlueprintViewCondition;
 
 namespace UE::MVVM { struct FBindingEntry; }
 
+
 namespace UE::MVVM::BindingEntry
 {
 
+DECLARE_DELEGATE_OneParam(FOnContextMenuEntryCallback, TConstArrayView<const TSharedPtr<FBindingEntry>> /*Entries*/);
+
 struct FRowHelper
 {
 	static void GatherAllChildBindings(UMVVMBlueprintView* BlueprintView, const TConstArrayView<TSharedPtr<FBindingEntry>> Entries, TArray<FGuid>& OutBindingIds, TArray<UMVVMBlueprintViewEvent*>& OutEvents, TArray<UMVVMBlueprintViewCondition*>& OutConditions);
 
 	static void DeleteEntries(const UWidgetBlueprint* WidgetBlueprint, UMVVMBlueprintView* BlueprintView, TArrayView<const TSharedPtr<FBindingEntry>> Selection);
+	static void DuplicateEntries(const UWidgetBlueprint* WidgetBlueprint, UMVVMBlueprintView* BlueprintView, TArrayView<const TSharedPtr<FBindingEntry>> Selection, TArray<const TSharedPtr<FBindingEntry>>& OutSelection);
+
 	static void ShowBlueprintGraph(FBlueprintEditor* Editor, UWidgetBlueprint* WidgetBlueprint, UMVVMBlueprintView* BlueprintView, TArrayView<const TSharedPtr<FBindingEntry>> Selection);
 
 	static TOptional<FMVVMBlueprintPropertyPath> DropFieldSelector(const UWidgetBlueprint* WidgetBlueprint, const FDragDropEvent& DragDropEvent);
 	static void DragEnterFieldSelector(const UWidgetBlueprint* WidgetBlueprint, const FDragDropEvent& DragDropEvent);
 
-	static FMenuBuilder CreateContextMenu(UWidgetBlueprint* WidgetBlueprint, UMVVMBlueprintView* View, TConstArrayView<TSharedPtr<FBindingEntry>> Entries);
+	static FMenuBuilder CreateContextMenu(UWidgetBlueprint* WidgetBlueprint, UMVVMBlueprintView* View, TConstArrayView<TSharedPtr<FBindingEntry>> Entries, FOnContextMenuEntryCallback OnSelectionChanged = FOnContextMenuEntryCallback());
 };
 
 } // namespace UE::MVVM
diff --git a/Engine/Plugins/Runtime/ModelViewViewModel/Source/ModelViewViewModelEditor/Private/Widgets/SMVVMViewBindingListView.cpp b/Engine/Plugins/Runtime/ModelViewViewModel/Source/ModelViewViewModelEditor/Private/Widgets/SMVVMViewBindingListView.cpp
--- a/Engine/Plugins/Runtime/ModelViewViewModel/Source/ModelViewViewModelEditor/Private/Widgets/SMVVMViewBindingListView.cpp
+++ b/Engine/Plugins/Runtime/ModelViewViewModel/Source/ModelViewViewModelEditor/Private/Widgets/SMVVMViewBindingListView.cpp
@@ -176,6 +176,27 @@ namespace Private
 			FilteredRootGroups = RootGroups;
 		}
 	}
+
+	void SetBindingEntrySelection(const TSharedPtr<STreeView<TSharedPtr<FBindingEntry>>>& TreeView, TConstArrayView<TSharedPtr<FBindingEntry>> Entries, TConstArrayView<const TSharedPtr<FBindingEntry>> EntriesToSelect)
+	{
+		if (!TreeView || Entries.IsEmpty() || EntriesToSelect.IsEmpty())
+		{
+			return;
+		}
+
+		for (const TSharedPtr<FBindingEntry>& Entry : Entries)
+		{
+			TFunction<bool(TSharedPtr<FBindingEntry>)> EquivalenceFunc = [Entry](TSharedPtr<FBindingEntry> Other) { return *Entry == *Other; };;
+			if (EntriesToSelect.ContainsByPredicate(EquivalenceFunc))
+			{
+				TreeView->SetItemSelection(Entry, true);
+				TreeView->RequestScrollIntoView(Entry);
+			}
+
+			SetBindingEntrySelection(TreeView, Entry->GetAllChildren(), EntriesToSelect);
+		}
+	}
+
 } // namespace
 
 
@@ -643,7 +664,7 @@ void SBindingsList::HandleRefreshChildren(FObjectKey ObjectHolder)
 	}
 }
 
-TSharedRef<ITableRow> SBindingsList::GenerateEntryRow(TSharedPtr<FBindingEntry> Entry, const TSharedRef<STableViewBase>& OwnerTable) const
+TSharedRef<ITableRow> SBindingsList::GenerateEntryRow(TSharedPtr<FBindingEntry> Entry, const TSharedRef<STableViewBase>& OwnerTable)
 {
 	TSharedPtr<ITableRow> Row;
 
@@ -703,8 +724,9 @@ TSharedPtr<SWidget> SBindingsList::OnSourceConstructContextMenu()
 	UMVVMWidgetBlueprintExtension_View* MVVMExtensionPtr = MVVMExtension.Get();
 	UWidgetBlueprint* WidgetBlueprint = MVVMExtensionPtr ? MVVMExtensionPtr->GetWidgetBlueprint() : nullptr;
 	UMVVMBlueprintView* View = MVVMExtensionPtr ? MVVMExtensionPtr->GetBlueprintView() : nullptr;
+	BindingEntry::FOnContextMenuEntryCallback SelectionCallback = BindingEntry::FOnContextMenuEntryCallback::CreateSP(this, &SBindingsList::SetSelection);
 
-	return BindingEntry::FRowHelper::CreateContextMenu(WidgetBlueprint, View, TreeView->GetSelectedItems()).MakeWidget();
+	return BindingEntry::FRowHelper::CreateContextMenu(WidgetBlueprint, View, TreeView->GetSelectedItems(), SelectionCallback).MakeWidget();
 }
 
 void SBindingsList::RequestNavigateToBinding(FGuid BindingId)
@@ -771,6 +793,15 @@ void SBindingsList::SetBindingsExpansion(bool InIsExpanded)
 	}
 }
 
+void SBindingsList::SetSelection(TConstArrayView<const TSharedPtr<FBindingEntry>> InEntries)
+{
+	if (TreeView.IsValid())
+	{
+		TreeView->ClearSelection();
+		Private::SetBindingEntrySelection(TreeView, TreeView->GetRootItems(), InEntries);
+	}
+}
+
 FReply SBindingsList::OnKeyDown(const FGeometry& MyGeometry, const FKeyEvent& InKeyEvent)
 {
 	if (InKeyEvent.GetKey() == EKeys::Delete)
diff --git a/Engine/Plugins/Runtime/ModelViewViewModel/Source/ModelViewViewModelEditor/Private/Widgets/SMVVMViewBindingListView.h b/Engine/Plugins/Runtime/ModelViewViewModel/Source/ModelViewViewModelEditor/Private/Widgets/SMVVMViewBindingListView.h
--- a/Engine/Plugins/Runtime/ModelViewViewModel/Source/ModelViewViewModelEditor/Private/Widgets/SMVVMViewBindingListView.h
+++ b/Engine/Plugins/Runtime/ModelViewViewModel/Source/ModelViewViewModelEditor/Private/Widgets/SMVVMViewBindingListView.h
@@ -48,8 +48,10 @@ public:
 	void SetRootGroupsExpansion(bool InIsExpanded);
 	void SetBindingsExpansion(bool InIsExpanded);
 
+	void SetSelection(TConstArrayView<const TSharedPtr<FBindingEntry>> InEntries);
+
 private:
-	TSharedRef<ITableRow> GenerateEntryRow(TSharedPtr<FBindingEntry> Entry, const TSharedRef<STableViewBase>& OwnerTable) const;
+	TSharedRef<ITableRow> GenerateEntryRow(TSharedPtr<FBindingEntry> Entry, const TSharedRef<STableViewBase>& OwnerTable);
 	void OnSourceListSelectionChanged(TSharedPtr<FBindingEntry> Entry, ESelectInfo::Type SelectionType) const;
 	void GetChildrenOfEntry(TSharedPtr<FBindingEntry> Entry, TArray<TSharedPtr<FBindingEntry>>& OutChildren) const;
 
