UI Component: implement field notification
- fix MVVM bug with binding to a field notify property of another property on the root widget

diff --git a/Engine/Plugins/Runtime/ModelViewViewModel/Source/ModelViewViewModel/Private/Bindings/MVVMFieldPathHelper.cpp b/Engine/Plugins/Runtime/ModelViewViewModel/Source/ModelViewViewModel/Private/Bindings/MVVMFieldPathHelper.cpp
--- a/Engine/Plugins/Runtime/ModelViewViewModel/Source/ModelViewViewModel/Private/Bindings/MVVMFieldPathHelper.cpp
+++ b/Engine/Plugins/Runtime/ModelViewViewModel/Source/ModelViewViewModel/Private/Bindings/MVVMFieldPathHelper.cpp
@@ -3,6 +3,9 @@
 #include "Bindings/MVVMFieldPathHelper.h"
 
 #include "Bindings/MVVMBindingHelper.h"
+#include "Blueprint/UserWidget.h"
+#include "Blueprint/WidgetBlueprintGeneratedClass.h"
+#include "Extensions/UIComponent.h"
 #include "Misc/MemStack.h"
 #include "MVVMSubsystem.h"
 #include "Templates/ValueOrError.h"
@@ -517,6 +520,31 @@ TValueOrError<const UStruct*, void> GetFieldAsContainer(const UE::MVVM::FMVVMCon
 	return MakeError();
 }
 
+const UWidgetBlueprintGeneratedClass* GetComponentPropertyPathSource(TArray<UE::MVVM::FMVVMConstFieldVariant> ComponentPaths, const UWidgetBlueprintGeneratedClass* Accessor)
+{
+	const UWidgetBlueprintGeneratedClass* ResultWidgetBlueprintGeneratedClass = Accessor;
+
+	for (int32 Index = 0; Index < ComponentPaths.Num(); ++Index)
+	{
+		const UE::MVVM::FMVVMConstFieldVariant& Field = ComponentPaths[Index];
+		if (Field.IsProperty())
+		{
+			if (const FObjectPropertyBase* ObjectProperty = CastField<const FObjectPropertyBase>(Field.GetProperty()))
+			{
+				if (ObjectProperty->PropertyClass->IsChildOf(UUIComponent::StaticClass()))
+				{
+					if (const UWidgetBlueprintGeneratedClass* OwnerWidgetBlueprintGeneratedClass = Cast<UWidgetBlueprintGeneratedClass>(ObjectProperty->GetOwnerClass()))
+					{
+						ResultWidgetBlueprintGeneratedClass = OwnerWidgetBlueprintGeneratedClass;
+						break;
+					}
+				}
+			}
+		}
+	}
+	return ResultWidgetBlueprintGeneratedClass;
+}
+
 } // namespace
 
 #undef LOCTEXT_NAMESPACE
diff --git a/Engine/Plugins/Runtime/ModelViewViewModel/Source/ModelViewViewModel/Public/Bindings/MVVMFieldPathHelper.h b/Engine/Plugins/Runtime/ModelViewViewModel/Source/ModelViewViewModel/Public/Bindings/MVVMFieldPathHelper.h
--- a/Engine/Plugins/Runtime/ModelViewViewModel/Source/ModelViewViewModel/Public/Bindings/MVVMFieldPathHelper.h
+++ b/Engine/Plugins/Runtime/ModelViewViewModel/Source/ModelViewViewModel/Public/Bindings/MVVMFieldPathHelper.h
@@ -9,7 +9,7 @@ namespace UE::MVVM { struct FFieldContext; }
 namespace UE::MVVM { struct FMVVMConstFieldVariant; }
 namespace UE::MVVM { struct FMVVMFieldVariant; }
 template <typename ValueType, typename ErrorType> class TValueOrError;
-
+class UWidgetBlueprintGeneratedClass;
 
 namespace UE::MVVM::FieldPathHelper
 {
@@ -69,4 +69,11 @@ namespace UE::MVVM::FieldPathHelper
 	 */
 	[[nodiscard]] MODELVIEWVIEWMODEL_API TValueOrError<const UStruct*, void> GetFieldAsContainer(const UE::MVVM::FMVVMConstFieldVariant Field);
 
+	/**
+	 * Returns the generated widget blueprint class of the widget that owns this component.
+	 * Accessor: the widget blueprint class that owns the binding.
+	 * ComponentPaths: The path to a field on a component.
+	 */
+	[[nodiscard]] MODELVIEWVIEWMODEL_API const UWidgetBlueprintGeneratedClass* GetComponentPropertyPathSource(TArray<UE::MVVM::FMVVMConstFieldVariant> ComponentPaths, const UWidgetBlueprintGeneratedClass* Accessor);
+
 } // namespace
diff --git a/Engine/Plugins/Runtime/ModelViewViewModel/Source/ModelViewViewModelBlueprint/Private/Bindings/MVVMConversionFunctionHelper.cpp b/Engine/Plugins/Runtime/ModelViewViewModel/Source/ModelViewViewModelBlueprint/Private/Bindings/MVVMConversionFunctionHelper.cpp
--- a/Engine/Plugins/Runtime/ModelViewViewModel/Source/ModelViewViewModelBlueprint/Private/Bindings/MVVMConversionFunctionHelper.cpp
+++ b/Engine/Plugins/Runtime/ModelViewViewModel/Source/ModelViewViewModelBlueprint/Private/Bindings/MVVMConversionFunctionHelper.cpp
@@ -4,6 +4,7 @@
 
 #include "Bindings/MVVMBindingHelper.h"
 #include "Bindings/MVVMFieldPathHelper.h"
+#include "Blueprint/WidgetBlueprintGeneratedClass.h"
 #include "Components/Widget.h"
 #include "Containers/Deque.h"
 #include "MVVMBlueprintView.h"
@@ -457,12 +458,15 @@ namespace NamedNodes
 				case EMVVMBlueprintFieldPathSource::Widget:
 					if (PropertyPath.IsComponent())
 					{
-						bCreateSelfNodeForRootProperty = true;
-					}
-					else
-					{
-						RootProperty = Blueprint->SkeletonGeneratedClass->FindPropertyByName(PropertyPath.GetWidgetName());
+						const UWidgetBlueprintGeneratedClass* ComponentSource = FieldPathHelper::GetComponentPropertyPathSource(PropertyPath.GetFields(BlueprintClass), Cast<UWidgetBlueprintGeneratedClass>(BlueprintClass));
+						if (ComponentSource == BlueprintClass)
+						{
+							bCreateSelfNodeForRootProperty = true;
+							break;
+						}
 					}
+
+					RootProperty = Blueprint->SkeletonGeneratedClass->FindPropertyByName(PropertyPath.GetWidgetName());
 					break;
 				default:
 					check(false);
diff --git a/Engine/Plugins/Runtime/ModelViewViewModel/Source/ModelViewViewModelBlueprint/Private/MVVMPropertyPath.cpp b/Engine/Plugins/Runtime/ModelViewViewModel/Source/ModelViewViewModelBlueprint/Private/MVVMPropertyPath.cpp
--- a/Engine/Plugins/Runtime/ModelViewViewModel/Source/ModelViewViewModelBlueprint/Private/MVVMPropertyPath.cpp
+++ b/Engine/Plugins/Runtime/ModelViewViewModel/Source/ModelViewViewModelBlueprint/Private/MVVMPropertyPath.cpp
@@ -336,89 +336,74 @@ TArray<UE::MVVM::FMVVMConstFieldVariant> FMVVMBlueprintPropertyPath::GetComponen
 {
 	TArray<UE::MVVM::FMVVMConstFieldVariant> Result;
 
-	if (bIsComponent)
+	if (!bIsComponent)
 	{
-		if (const UWidgetBlueprint* WidgetBP = Cast<UWidgetBlueprint>(InSelfContext->ClassGeneratedBy))
+		return Result;
+	}
+
+	const UWidgetBlueprint* WidgetBP = Cast<UWidgetBlueprint>(InSelfContext->ClassGeneratedBy);
+	if (!WidgetBP)
+	{
+		return Result;
+	}
+
+	const FObjectPropertyBase* WidgetProperty = nullptr;
+	if (GetSource(WidgetBP) == EMVVMBlueprintFieldPathSource::Widget)
+	{
+		if (UClass* GeneratedClass = WidgetBP->SkeletonGeneratedClass ? WidgetBP->SkeletonGeneratedClass : WidgetBP->GeneratedClass)
 		{
-			const FObjectPropertyBase* WidgetProperty = nullptr;
-			if (GetSource(WidgetBP) == EMVVMBlueprintFieldPathSource::Widget)
+			WidgetProperty = CastField<const FObjectPropertyBase>(GeneratedClass->FindPropertyByName(GetWidgetName()));
+		}
+	}
+
+	const FObjectPropertyBase* OwningWidget = WidgetProperty;
+
+	for (const FMVVMBlueprintFieldPath& Path : Paths)
+	{
+		if (const FProperty* PropertyPtr = Path.GetField(InSelfContext).GetProperty())
+		{
+			const FObjectPropertyBase* ObjectProperty = CastField<const FObjectPropertyBase>(PropertyPtr);
+			if (!ObjectProperty)
 			{
-				if (UClass* GeneratedClass = WidgetBP->SkeletonGeneratedClass ? WidgetBP->SkeletonGeneratedClass : WidgetBP->GeneratedClass)
-				{
-					WidgetProperty = CastField<const FObjectPropertyBase>(GeneratedClass->FindPropertyByName(GetWidgetName()));
-				}
+				Result.Add(Path.GetField(InSelfContext));
+				continue;
 			}
 
-			const UWidgetBlueprint* CurrentWidgetBP = WidgetBP;
-			const UWidgetBlueprint* LastWidgetBP = nullptr;
-			if (WidgetProperty && WidgetProperty->PropertyClass->IsChildOf(UUserWidget::StaticClass()))
+			if (ObjectProperty->PropertyClass->IsChildOf(UWidget::StaticClass()))
+			{
+				OwningWidget = ObjectProperty;
+			}
+			else if (ObjectProperty->PropertyClass->IsChildOf(UUIComponent::StaticClass()))
 			{
-				if (UWidgetBlueprintGeneratedClass* UserWidget = Cast<UWidgetBlueprintGeneratedClass>(WidgetProperty->PropertyClass))
+				if (!OwningWidget)
 				{
-					LastWidgetBP = CurrentWidgetBP;
-					CurrentWidgetBP = Cast<UWidgetBlueprint>(UserWidget->ClassGeneratedBy);
+					Result.Add(Path.GetField(InSelfContext));
+					continue;
 				}
-			}
-
-			const FObjectPropertyBase* OwningWidget = WidgetProperty;
 
-			for (const FMVVMBlueprintFieldPath& Path : Paths)
-			{
-				if (ensure(CurrentWidgetBP))
+				if (UWidgetBlueprintGeneratedClass* OwnerWidgetBlueprintGeneratedClass = Cast<UWidgetBlueprintGeneratedClass>(ObjectProperty->GetOwnerClass()))
 				{
-					if (const FProperty* PropertyPtr = Path.GetField(InSelfContext).GetProperty())
+					UWidgetBlueprint* OwnerWidgetBlueprint = Cast<UWidgetBlueprint>(OwnerWidgetBlueprintGeneratedClass->ClassGeneratedBy);
+					if (UUIComponentWidgetBlueprintExtension* ComponentWidgetBlueprintExtension = UWidgetBlueprintExtension::GetExtension<UUIComponentWidgetBlueprintExtension>(OwnerWidgetBlueprint))
 					{
-						if (const FObjectPropertyBase* ObjectProperty = CastField<const FObjectPropertyBase>(PropertyPtr))
+						if (UWidget* Widget = OwnerWidgetBlueprint->WidgetTree->FindWidget(OwningWidget->GetFName()))
 						{
-							if (ObjectProperty->PropertyClass->IsChildOf(UUserWidget::StaticClass()))
-							{
-								if (UWidgetBlueprintGeneratedClass* UserWidget = Cast<UWidgetBlueprintGeneratedClass>(ObjectProperty->PropertyClass))
-								{
-									LastWidgetBP = CurrentWidgetBP;
-									CurrentWidgetBP = Cast<UWidgetBlueprint>(UserWidget->ClassGeneratedBy);
-								}
-							}
-							if (ObjectProperty->PropertyClass->IsChildOf(UWidget::StaticClass()))
-							{
-								OwningWidget = ObjectProperty;
-							}
-							else if (ObjectProperty->PropertyClass->IsChildOf(UUIComponent::StaticClass()))
+							TArray<UUIComponent*> ComponentsList = ComponentWidgetBlueprintExtension->GetComponentsFor(Widget);
+							for (UUIComponent* Component : ComponentsList)
 							{
-								if (OwningWidget)
+								const FName ComponentPropertyName = UUIComponentContainer::GetPropertyNameForComponent(Component, Widget->GetFName());
+								if (Result.Num() > 0 && ObjectProperty->PropertyClass == Component->GetClass() && ObjectProperty->GetFName() == ComponentPropertyName)
 								{
-									const UWidgetBlueprint* BlueprintOwningExtension = CurrentWidgetBP;
-									if (LastWidgetBP && OwningWidget->PropertyClass->IsChildOf(UUserWidget::StaticClass()) && OwningWidget->PropertyClass->ClassGeneratedBy == CurrentWidgetBP)
-									{
-										BlueprintOwningExtension = LastWidgetBP;
-									}
-									if (BlueprintOwningExtension)
-									{
-										if (UUIComponentWidgetBlueprintExtension* ComponentWidgetBlueprintExtension = UWidgetBlueprintExtension::GetExtension<UUIComponentWidgetBlueprintExtension>(BlueprintOwningExtension))
-										{
-											if (UWidget* Widget = BlueprintOwningExtension->WidgetTree->FindWidget(OwningWidget->GetFName()))
-											{
-												TArray<UUIComponent*> ComponentsList = ComponentWidgetBlueprintExtension->GetComponentsFor(Widget);
-												for (UUIComponent* Component : ComponentsList)
-												{
-													FName ComponentPropertyName = UUIComponentContainer::GetPropertyNameForComponent(Component, Widget->GetFName());
-
-													if (Result.Num() > 0 && ObjectProperty->PropertyClass == Component->GetClass() && ObjectProperty->GetFName() == ComponentPropertyName)
-													{
-														Result.Pop();
-														continue;
-													}
-												}
-											}
-										}
-									}
+									Result.Pop();
+									continue;
 								}
 							}
 						}
-					}
+					}	
 				}
-				Result.Add(Path.GetField(InSelfContext));
 			}
-		}
+		}					
+		Result.Add(Path.GetField(InSelfContext));		
 	}
 
 	return Result;
diff --git a/Engine/Plugins/Runtime/ModelViewViewModel/Source/ModelViewViewModelBlueprint/Private/MVVMViewBlueprintCompiler.cpp b/Engine/Plugins/Runtime/ModelViewViewModel/Source/ModelViewViewModelBlueprint/Private/MVVMViewBlueprintCompiler.cpp
--- a/Engine/Plugins/Runtime/ModelViewViewModel/Source/ModelViewViewModelBlueprint/Private/MVVMViewBlueprintCompiler.cpp
+++ b/Engine/Plugins/Runtime/ModelViewViewModel/Source/ModelViewViewModelBlueprint/Private/MVVMViewBlueprintCompiler.cpp
@@ -9,6 +9,7 @@
 #include "Blueprint/WidgetTree.h"
 #include "Components/Widget.h"
 #include "Extensions/MVVMBlueprintViewExtension.h"
+#include "Extensions/UIComponent.h"
 #include "HAL/IConsoleManager.h"
 #include "Kismet2/BlueprintEditorUtils.h"
 #include "K2Node_ComponentBoundEvent.h"
@@ -1002,15 +1003,38 @@ void FMVVMViewBlueprintCompiler::CreateRequiredProperties(const FWidgetBlueprint
 
 	auto GenerateCompilerContext = [Self = this, WidgetBP = Context.GetWidgetBlueprint(), &ViewModelGuids, &WidgetSourcesCreated, &WidgetUserPropertyCreated, &bSelfBindingSourceCreated](bool bInCreateSource, const FMVVMBlueprintPropertyPath& PropertyPath) -> TValueOrError<void, FText>
 	{
+		auto CreateSourceForSelf = [Self, &bSelfBindingSourceCreated]() -> TSharedRef<FCompilerBindingSource>
+		{
+			TSharedRef<FCompilerBindingSource> SourceContext = MakeShared<FCompilerBindingSource>();
+			SourceContext->AuthoritativeClass = nullptr;
+			SourceContext->Name = Self->WidgetBlueprintCompilerContext.WidgetBlueprint()->GetFName();
+			SourceContext->Type = FCompilerBindingSource::EType::Self;
+			SourceContext->bIsOptional = false;
+
+			if (!bSelfBindingSourceCreated)
+			{
+				Self->NeededBindingSources.Add(SourceContext);
+				bSelfBindingSourceCreated = true;
+			}
+			return SourceContext;
+		};
 		FName PropertyPathWidgetName = PropertyPath.GetWidgetName();
 		FGuid PropertyPathViewModelId = PropertyPath.GetViewModelId();
 
+		bool bIsComponentOnSelf = false;
+		if (PropertyPath.IsComponent())
+		{
+			const UWidgetBlueprintGeneratedClass* WidgetBPGeneratedClass = Cast<UWidgetBlueprintGeneratedClass>(WidgetBP->GeneratedClass);
+			const UWidgetBlueprintGeneratedClass* ComponentOwningWidgetBlueprint = UE::MVVM::FieldPathHelper::GetComponentPropertyPathSource(PropertyPath.GetFields(WidgetBPGeneratedClass), WidgetBPGeneratedClass);
+			bIsComponentOnSelf = ComponentOwningWidgetBlueprint == WidgetBPGeneratedClass;
+		}
+
 		// If the path is "self.widget.property", "self.viewmodel.property", then remove the "self".
 		EMVVMBlueprintFieldPathSource FieldPathSource = PropertyPath.GetSource(Self->WidgetBlueprintCompilerContext.WidgetBlueprint());
 		if (FieldPathSource == EMVVMBlueprintFieldPathSource::SelfContext)
 		{
 			const TArrayView<const FMVVMBlueprintFieldPath> FieldPaths = PropertyPath.GetFieldPaths();
-			if (FieldPaths.Num() >= 2 && !PropertyPath.IsComponent())
+			if (FieldPaths.Num() >= 2 && !bIsComponentOnSelf)
 			{
 				if (FieldPaths[0].GetBindingKind() == EBindingKind::Property && FieldPaths[0].IsFieldSelfContext())
 				{
@@ -1033,17 +1057,7 @@ void FMVVMViewBlueprintCompiler::CreateRequiredProperties(const FWidgetBlueprint
 		{
 		case EMVVMBlueprintFieldPathSource::SelfContext:
 		{
-			if (!bSelfBindingSourceCreated)
-			{
-				bSelfBindingSourceCreated = true;
-
-				TSharedRef<FCompilerBindingSource> SourceContext = MakeShared<FCompilerBindingSource>();
-				SourceContext->AuthoritativeClass = nullptr;
-				SourceContext->Name = Self->WidgetBlueprintCompilerContext.WidgetBlueprint()->GetFName();
-				SourceContext->Type = FCompilerBindingSource::EType::Self;
-				SourceContext->bIsOptional = false;
-				Self->NeededBindingSources.Add(SourceContext);
-			}
+			CreateSourceForSelf();
 			return MakeValue();
 		}
 		case EMVVMBlueprintFieldPathSource::Widget:
@@ -1102,6 +1116,11 @@ void FMVVMViewBlueprintCompiler::CreateRequiredProperties(const FWidgetBlueprint
 
 					WidgetUserPropertyCreated.Add(Widget->GetFName());
 				}
+
+				if (bIsComponentOnSelf)
+				{
+					CreateSourceForSelf();
+				}
 			}
 			break;
 		}
@@ -2176,7 +2195,14 @@ void FMVVMViewBlueprintCompiler::CreateReadFieldContexts(UWidgetBlueprintGenerat
 							bool bValidWidget = CreateFieldResult.GetValue()->Source->Type == FCompilerBindingSource::EType::Widget && PathSource == EMVVMBlueprintFieldPathSource::Widget;
 							bool bDynamic = CreateFieldResult.GetValue()->Source->Type == FCompilerBindingSource::EType::DynamicViewmodel;
 							bool bSelf = CreateFieldResult.GetValue()->Source->Type == FCompilerBindingSource::EType::Self && PathSource == EMVVMBlueprintFieldPathSource::SelfContext;
-							if (!(bValidViewModel || bValidWidget || bDynamic || bSelf))
+							bool bIsComponentOnSelf = false;
+							if (ReadFieldContext->bIsComponent)
+							{
+								const UWidgetBlueprintGeneratedClass* ComponentOwningWidgetBlueprint = UE::MVVM::FieldPathHelper::GetComponentPropertyPathSource(ReadFieldContext->SkeletalGeneratedFields, Class);
+								bIsComponentOnSelf = ComponentOwningWidgetBlueprint == Class;
+							}
+							bIsComponentOnSelf = bIsComponentOnSelf && CreateFieldResult.GetValue()->Source->Type == FCompilerBindingSource::EType::Self && PathSource == EMVVMBlueprintFieldPathSource::Widget;
+							if (!(bValidViewModel || bValidWidget || bDynamic || bSelf || bIsComponentOnSelf))
 							{
 								Self->AddMessageForBinding(Binding
 									, FText::Format(LOCTEXT("CreateNotifyFieldIdFailedInvalidInvalidContext", "Internal error. The property path '{0}' is invalid. The context is invalid."), PropertyPathToText(Class, Self->BlueprintView.Get(), PropertyPath))
@@ -2212,6 +2238,7 @@ void FMVVMViewBlueprintCompiler::CreateReadFieldContexts(UWidgetBlueprintGenerat
 				Found->GeneratedFields = MoveTemp(CreateSourceResult.GetValue().GeneratedFields);
 				Found->SkeletalGeneratedFields = MoveTemp(CreateSourceResult.GetValue().SkeletalGeneratedFields);
 				Found->GeneratedFrom = MoveTemp(CreateSourceResult.GetValue().GeneratedFrom);
+				Found->bIsComponent = MoveTemp(CreateSourceResult.GetValue().bIsComponent);
 
 				if (CreateFieldId(BindingSourcePath, Found, FMVVMBlueprintPinId()).HasError())
 				{
@@ -2249,6 +2276,7 @@ void FMVVMViewBlueprintCompiler::CreateReadFieldContexts(UWidgetBlueprintGenerat
 							Found->GeneratedFields = MoveTemp(CreateSourceResult.GetValue().GeneratedFields);
 							Found->SkeletalGeneratedFields = MoveTemp(CreateSourceResult.GetValue().SkeletalGeneratedFields);
 							Found->GeneratedFrom = MoveTemp(CreateSourceResult.GetValue().GeneratedFrom);
+							Found->bIsComponent = MoveTemp(CreateSourceResult.GetValue().bIsComponent);
 
 							if (CreateFieldId(Pin.GetPath(), Found, Pin.GetId()).HasError())
 							{
@@ -2295,6 +2323,7 @@ void FMVVMViewBlueprintCompiler::CreateReadFieldContexts(UWidgetBlueprintGenerat
 					Found->GeneratedFields = MoveTemp(CreateSourceResult.GetValue().GeneratedFields);
 					Found->GeneratedFrom = MoveTemp(CreateSourceResult.GetValue().GeneratedFrom);
 					Found->SkeletalGeneratedFields = MoveTemp(CreateSourceResult.GetValue().SkeletalGeneratedFields);
+					Found->bIsComponent = MoveTemp(CreateSourceResult.GetValue().bIsComponent);
 					GeneratedReadFieldPaths.Add(Found.ToSharedRef());
 				}
 
@@ -2330,6 +2359,8 @@ void FMVVMViewBlueprintCompiler::CreateReadFieldContexts(UWidgetBlueprintGenerat
 				FoundReadFieldPath->GeneratedFrom = MoveTemp(CreateConditionSourceResult.GetValue().GeneratedFrom);
 				FoundReadFieldPath->GeneratedFields = MoveTemp(CreateConditionSourceResult.GetValue().GeneratedFields);
 				FoundReadFieldPath->SkeletalGeneratedFields = MoveTemp(CreateConditionSourceResult.GetValue().SkeletalGeneratedFields);
+				FoundReadFieldPath->bIsComponent = MoveTemp(CreateConditionSourceResult.GetValue().bIsComponent);
+
 				GeneratedReadFieldPaths.Add(FoundReadFieldPath.ToSharedRef());
 			}
 
@@ -4194,10 +4225,16 @@ TValueOrError<FMVVMViewBlueprintCompiler::FGetFieldsResult, FText> FMVVMViewBlue
 	TArray<UE::MVVM::FMVVMConstFieldVariant> PropertyPathFields = PropertyPath.GetFields(Class);
 	FGuid ViewModelId = PropertyPath.GetViewModelId();
 	FName DestinationWidgetName = PropertyPath.GetWidgetName();
+	bool bIsComponentOnSelf = false;
+	if (PropertyPath.IsComponent())
+	{
+		const UWidgetBlueprintGeneratedClass* ComponentOwningWidgetBlueprint = UE::MVVM::FieldPathHelper::GetComponentPropertyPathSource(PropertyPath.GetFields(Class), Class);
+		bIsComponentOnSelf = ComponentOwningWidgetBlueprint == Class;
+	}
 
 	if (FieldPathSource == EMVVMBlueprintFieldPathSource::SelfContext)
 	{
-		if (PropertyPathFields.Num() >= 2 && PropertyPathFields[0].IsProperty())
+		if (PropertyPathFields.Num() >= 2 && PropertyPathFields[0].IsProperty() && !bIsComponentOnSelf)
 		{
 			FName FieldName = PropertyPathFields[0].GetName();
 			if (WidgetNameToWidgetPointerMap.Contains(FieldName))
@@ -4263,19 +4300,16 @@ TValueOrError<FMVVMViewBlueprintCompiler::FGetFieldsResult, FText> FMVVMViewBlue
 		check(WidgetNameToWidgetPointerMap.Contains(DestinationWidgetName));
 		checkf(!DestinationWidgetName.IsNone(), TEXT("The destination should have been checked and set bAreSourceContextsValid."));
 
-		if (!PropertyPath.IsComponent())
-		{
-			const FCompilerUserWidgetProperty* FoundUserWidgetProperty = NeededUserWidgetProperties.FindByPredicate([DestinationWidgetName](const FCompilerUserWidgetProperty& Other) { return Other.Name == DestinationWidgetName; });
-			check(FoundUserWidgetProperty);
-		}
+		const FCompilerUserWidgetProperty* FoundUserWidgetProperty = NeededUserWidgetProperties.FindByPredicate([DestinationWidgetName](const FCompilerUserWidgetProperty& Other) { return Other.Name == DestinationWidgetName; });
+		check(FoundUserWidgetProperty);
 
-		TValueOrError<TSharedPtr<FCompilerBindingSource>, FText> FindSourceResult = FindSource(DestinationWidgetName, FCompilerBindingSource::EType::Widget);
+		TValueOrError<TSharedPtr<FCompilerBindingSource>, FText> FindSourceResult = bIsComponentOnSelf ? FindSource(WidgetBlueprintCompilerContext.WidgetBlueprint()->GetFName(), FCompilerBindingSource::EType::Self) : FindSource(DestinationWidgetName, FCompilerBindingSource::EType::Widget);
 		if (FindSourceResult.HasError())
 		{
 			return MakeError(FindSourceResult.StealError());
 		}
 
-		Result.GeneratedFields = PropertyPath.IsComponent() ? AppendBaseField(Class, FName(), PropertyPathFields) : AppendBaseField(Class, DestinationWidgetName, PropertyPathFields);
+		Result.GeneratedFields = bIsComponentOnSelf ? AppendBaseField(Class, FName(), PropertyPathFields) : AppendBaseField(Class, DestinationWidgetName, PropertyPathFields);
 		Result.OptionalSource = FindSourceResult.StealValue();
 		break;
 	}
@@ -4319,6 +4353,7 @@ TValueOrError<FMVVMViewBlueprintCompiler::FCreateFieldsResult, FText> FMVVMViewB
 		Result.GeneratedFrom = MoveTemp(GetFieldResult.GetValue().GeneratedFrom);
 		Result.OptionalSource = MoveTemp(GetFieldResult.GetValue().OptionalSource);
 		Result.GeneratedFields = MoveTemp(GetFieldResult.GetValue().GeneratedFields);
+		Result.bIsComponent = PropertyPath.IsComponent();
 	}
 
 	if (!IsPropertyPathValid(WidgetBlueprintCompilerContext.WidgetBlueprint(), Result.GeneratedFields))
@@ -4386,14 +4421,25 @@ TValueOrError<TSharedPtr<FMVVMViewBlueprintCompiler::FCompilerNotifyFieldId>, FT
 		return nullptr;
 	};
 
+	// If the source type is Self, we don't include "self" as a field in SkeletalGeneratedFields, so the ViewModelIndex will not consider that. Meanwhile, we do include the widget if source type is Widget.
+	// So a ViewModelIndex of 0 with source type Self is equivalent to a ViewModelIndex of 1 with source type widget.
+ 	bool bIsSelf = ReadFieldContext->Source->Type == FCompilerBindingSource::EType::Self && ReadFieldContext->SkeletalGeneratedFields.Num() > 1;
+	bool bIsComponentOnSelf = false;
+	if (ReadFieldContext->bIsComponent)
+	{
+		const UWidgetBlueprintGeneratedClass* ComponentOwningWidgetBlueprint = UE::MVVM::FieldPathHelper::GetComponentPropertyPathSource(ReadFieldContext->SkeletalGeneratedFields, Class);
+		bIsComponentOnSelf = ComponentOwningWidgetBlueprint == Class;
+	}
+
+	int32 ViewModelIndexIncludingSource = (bIsSelf || bIsComponentOnSelf) ? BindingInfo.ViewModelIndex + 1 : BindingInfo.ViewModelIndex;
 	// Sanity check
 	{
 		const UClass* ExpectedClass = nullptr;
-		if (BindingInfo.ViewModelIndex < 1 && BindingInfo.NotifyFieldClass)
+		if (ViewModelIndexIncludingSource < 1 && BindingInfo.NotifyFieldClass)
 		{
 			ExpectedClass = ReadFieldContext->Source->AuthoritativeClass;
 		}
-		else if (BindingInfo.ViewModelIndex >= 0)
+		else if (ViewModelIndexIncludingSource >= 0 && BindingInfo.ViewModelIndex != INDEX_NONE)
 		{
 			ExpectedClass = GetClassFromField(ReadFieldContext->SkeletalGeneratedFields[BindingInfo.ViewModelIndex]);
 		}
@@ -4407,14 +4453,14 @@ TValueOrError<TSharedPtr<FMVVMViewBlueprintCompiler::FCompilerNotifyFieldId>, FT
 	// The INotifyFieldValueChanged/viewmodel is not the first and only INotifyFieldValueChanged/viewmodel property path.
 	//Create a new source in PropertyPath creator mode. Create a special binding to update the viewmodel when it changes.
 	//This binding (calling this function) will use the new source.
-	if (BindingInfo.ViewModelIndex >= 1)
+	if (ViewModelIndexIncludingSource >= 1)
 	{
 		if (!GetDefault<UMVVMDeveloperProjectSettings>()->bAllowLongSourcePath)
 		{
 			return MakeError(LOCTEXT("DynamicSourceEntryNotSupport", "Long source entry is not supported. Add the viewmodel manually."));
 		}
 
-		for (int32 DynamicIndex = 1; DynamicIndex <= BindingInfo.ViewModelIndex; ++DynamicIndex)
+		for (int32 DynamicIndex = (bIsSelf || bIsComponentOnSelf) && BindingInfo.ViewModelIndex == 0 ? 0 : 1 ; DynamicIndex <= BindingInfo.ViewModelIndex; ++DynamicIndex)
 		{
 			if (!ReadFieldContext->SkeletalGeneratedFields.IsValidIndex(DynamicIndex))
 			{
@@ -4459,6 +4505,11 @@ TValueOrError<TSharedPtr<FMVVMViewBlueprintCompiler::FCompilerNotifyFieldId>, FT
 				}
 			}
 
+			if (bIsComponentOnSelf || bIsSelf)
+			{
+				NewParentSourceName = NewParentSourceName.IsNone() ? Class->ClassGeneratedBy->GetFName() : NewParentSourceName;
+			}
+
 			// Does the parent exist
 			TSharedPtr<FCompilerBindingSource> ParentBindingSource;
 			{
diff --git a/Engine/Plugins/Runtime/ModelViewViewModel/Source/ModelViewViewModelBlueprint/Private/MVVMViewBlueprintCompiler.h b/Engine/Plugins/Runtime/ModelViewViewModel/Source/ModelViewViewModelBlueprint/Private/MVVMViewBlueprintCompiler.h
--- a/Engine/Plugins/Runtime/ModelViewViewModel/Source/ModelViewViewModelBlueprint/Private/MVVMViewBlueprintCompiler.h
+++ b/Engine/Plugins/Runtime/ModelViewViewModel/Source/ModelViewViewModelBlueprint/Private/MVVMViewBlueprintCompiler.h
@@ -254,6 +254,7 @@ private:
 		TArray<UE::MVVM::FMVVMConstFieldVariant> GeneratedFields; // the string path converted into field
 		TArray<UE::MVVM::FMVVMConstFieldVariant> SkeletalGeneratedFields; // the field path converted with getter and setter
 		EMVVMBlueprintFieldPathSource GeneratedFrom = EMVVMBlueprintFieldPathSource::None;
+		bool bIsComponent = false;
 
 		TSharedPtr<FCompilerNotifyFieldId> NotificationField;
 		FCompiledBindingLibraryCompiler::FFieldPathHandle LibraryCompilerHandle;
@@ -417,6 +418,7 @@ private:
 		EMVVMBlueprintFieldPathSource GeneratedFrom = EMVVMBlueprintFieldPathSource::None;
 		TArray<UE::MVVM::FMVVMConstFieldVariant> GeneratedFields;
 		TArray<UE::MVVM::FMVVMConstFieldVariant> SkeletalGeneratedFields;
+		bool bIsComponent;
 	};
 	TValueOrError<FGetFieldsResult, FText> GetFields(const UWidgetBlueprintGeneratedClass* Class, const FMVVMBlueprintPropertyPath& PropertyPath) const;
 	TValueOrError<FCreateFieldsResult, FText> CreateFieldContext(const UWidgetBlueprintGeneratedClass* Class, const FMVVMBlueprintPropertyPath& PropertyPath, bool bForSourceReading) const;
diff --git a/Engine/Plugins/Runtime/ModelViewViewModel/Source/ModelViewViewModelEditor/Private/Widgets/SMVVMViewModelBindingListWidget.cpp b/Engine/Plugins/Runtime/ModelViewViewModel/Source/ModelViewViewModelEditor/Private/Widgets/SMVVMViewModelBindingListWidget.cpp
--- a/Engine/Plugins/Runtime/ModelViewViewModel/Source/ModelViewViewModelEditor/Private/Widgets/SMVVMViewModelBindingListWidget.cpp
+++ b/Engine/Plugins/Runtime/ModelViewViewModel/Source/ModelViewViewModelEditor/Private/Widgets/SMVVMViewModelBindingListWidget.cpp
@@ -554,87 +554,63 @@ void SSourceBindingList::AddSources(TArrayView<const FBindingSource> InSources)
 
 bool IsPropertyPathComponent(TArrayView<const FFieldVariant> FieldPath, FMVVMBlueprintPropertyPath& PropertyPath, const UWidgetBlueprint* WidgetBlueprintPtr)
 {
-	if (ensure(WidgetBlueprintPtr))
+	if (!ensure(WidgetBlueprintPtr))
 	{
-		const FObjectPropertyBase* WidgetProperty = nullptr;
-		if (PropertyPath.GetSource(WidgetBlueprintPtr) == EMVVMBlueprintFieldPathSource::Widget)
+		return false;
+	}
+
+	const FObjectPropertyBase* WidgetProperty = nullptr;
+	if (PropertyPath.GetSource(WidgetBlueprintPtr) == EMVVMBlueprintFieldPathSource::Widget)
+	{
+		if (UClass* GeneratedClass = WidgetBlueprintPtr->SkeletonGeneratedClass ? WidgetBlueprintPtr->SkeletonGeneratedClass : WidgetBlueprintPtr->GeneratedClass)
 		{
-			if (UClass* GeneratedClass = WidgetBlueprintPtr->SkeletonGeneratedClass ? WidgetBlueprintPtr->SkeletonGeneratedClass : WidgetBlueprintPtr->GeneratedClass)
-			{
-				WidgetProperty = CastField<const FObjectPropertyBase>(GeneratedClass->FindPropertyByName(PropertyPath.GetWidgetName()));
-			}
+			WidgetProperty = CastField<const FObjectPropertyBase>(GeneratedClass->FindPropertyByName(PropertyPath.GetWidgetName()));
 		}
-
-		const UWidgetBlueprint* CurrentWidgetBP = WidgetBlueprintPtr;
-		const UWidgetBlueprint* LastWidgetBP = nullptr;
-		if (WidgetProperty && WidgetProperty->PropertyClass->IsChildOf(UUserWidget::StaticClass()))
+	}
+	const FObjectPropertyBase* OwningWidget = WidgetProperty;
+
+	/* We need 3 pieces of information to decide if this property path includes a valid component that is registered with a UIComponentWidgetBlueprintExtension:
+		1. The widget blueprint that holds a component extension. This is the parent class of the component property.
+		We need to use LastWidgetBP when the component object is on a userwidget.  
+		2. OwningWidget: The widget that owns this component. We care about the name of this widget so we can verify it with the widget blueprint extension.
+		3. The component itself. 
+	*/
+	for (const FFieldVariant& FieldVariant : FieldPath)
+	{
+		const FObjectPropertyBase* ObjectProperty = FieldVariant.Get<FObjectPropertyBase>();
+			
+		if (!ObjectProperty)
 		{
-			if (UWidgetBlueprintGeneratedClass* UserWidget = Cast<UWidgetBlueprintGeneratedClass>(WidgetProperty->PropertyClass))
-			{
-				LastWidgetBP = CurrentWidgetBP;
-				CurrentWidgetBP = Cast<UWidgetBlueprint>(UserWidget->ClassGeneratedBy);
-			}
+			continue;
 		}
-
-		const FObjectPropertyBase* OwningWidget = WidgetProperty;
-
-		/* We need 3 pieces of information to decide if this property path includes a valid component that is registered with a UIComponentWidgetBlueprintExtension:
-			1. CurrentWidgetBP: The widget blueprint that holds a component extension. 
-			There might be multiple of these in the path, we only care about the last one we encounter before finding the component itself.
-			We need to use LastWidgetBP when the component object is on a userwidget.  
-			2. OwningWidget: The widget that owns this component. We care about the name of this widget so we can verify it with the widget blueprint extension.
-			3. The component itself. 
-		*/
-		for (const FFieldVariant& FieldVariant : FieldPath)
+		if (ObjectProperty->PropertyClass->IsChildOf(UWidget::StaticClass()))
+		{
+			OwningWidget = ObjectProperty;
+		}
+		else if (ObjectProperty->PropertyClass->IsChildOf(UUIComponent::StaticClass()))
 		{
-			if (ensure(CurrentWidgetBP))
+			if (!OwningWidget)
 			{
-				if (const FProperty* PropertyPtr = FieldVariant.Get<FProperty>())
+				continue;
+			}
+			if (UWidgetBlueprintGeneratedClass* OwnerWidgetBlueprintGeneratedClass = Cast<UWidgetBlueprintGeneratedClass>(ObjectProperty->GetOwnerClass()))
+			{
+				UWidgetBlueprint* OwnerWidgetBlueprint = Cast<UWidgetBlueprint>(OwnerWidgetBlueprintGeneratedClass->ClassGeneratedBy);
+				if (UUIComponentWidgetBlueprintExtension* ComponentWidgetBlueprintExtension = UWidgetBlueprintExtension::GetExtension<UUIComponentWidgetBlueprintExtension>(OwnerWidgetBlueprint))
 				{
-					if (const FObjectPropertyBase* ObjectProperty = CastField<const FObjectPropertyBase>(PropertyPtr))
+					if (UWidget* Widget = OwnerWidgetBlueprint->WidgetTree->FindWidget(OwningWidget->GetFName()))
 					{
-						if (ObjectProperty->PropertyClass->IsChildOf(UUserWidget::StaticClass()))
-						{
-							if (UWidgetBlueprintGeneratedClass* UserWidget = Cast<UWidgetBlueprintGeneratedClass>(ObjectProperty->PropertyClass))
-							{
-								LastWidgetBP = CurrentWidgetBP;
-								CurrentWidgetBP = Cast<UWidgetBlueprint>(UserWidget->ClassGeneratedBy);
-							}
-						}
-
-						if (ObjectProperty->PropertyClass->IsChildOf(UWidget::StaticClass()))
-						{
-							OwningWidget = ObjectProperty;
-						}
-						else if (ObjectProperty->PropertyClass->IsChildOf(UUIComponent::StaticClass()))
+						TArray<UUIComponent*> ComponentsList = ComponentWidgetBlueprintExtension->GetComponentsFor(Widget);
+						for (UUIComponent* Component : ComponentsList)
 						{
-							if (OwningWidget)
+							check(Component);
+							FName ComponentPropertyName = UUIComponentContainer::GetPropertyNameForComponent(Component, Widget->GetFName());
+							if (ObjectProperty->PropertyClass == Component->GetClass() && ObjectProperty->GetFName() == ComponentPropertyName)
 							{
-								const UWidgetBlueprint* BlueprintOwningExtension = CurrentWidgetBP;
-								if (LastWidgetBP && OwningWidget->PropertyClass->IsChildOf(UUserWidget::StaticClass()) && OwningWidget->PropertyClass->ClassGeneratedBy == CurrentWidgetBP)
-								{
-									BlueprintOwningExtension = LastWidgetBP;
-								}
-
-								if (UUIComponentWidgetBlueprintExtension* ComponentWidgetBlueprintExtension = UWidgetBlueprintExtension::GetExtension<UUIComponentWidgetBlueprintExtension>(BlueprintOwningExtension))
-								{
-									if (UWidget* Widget = BlueprintOwningExtension->WidgetTree->FindWidget(OwningWidget->GetFName()))
-									{
-										TArray<UUIComponent*> ComponentsList = ComponentWidgetBlueprintExtension->GetComponentsFor(Widget);
-										for (UUIComponent* Component : ComponentsList)
-										{
-											check(Component);
-											FName ComponentPropertyName = UUIComponentContainer::GetPropertyNameForComponent(Component, Widget->GetFName());
-											if (ObjectProperty->PropertyClass == Component->GetClass() && ObjectProperty->GetFName() == ComponentPropertyName)
-											{
-												return true;
-											}
-										}
-									}
-								}
+									return true;
 							}
 						}
-					}
+					}			
 				}
 			}
 		}
diff --git a/Engine/Source/Editor/UMGEditor/Private/UIComponentWidgetBlueprintExtension.cpp b/Engine/Source/Editor/UMGEditor/Private/UIComponentWidgetBlueprintExtension.cpp
--- a/Engine/Source/Editor/UMGEditor/Private/UIComponentWidgetBlueprintExtension.cpp
+++ b/Engine/Source/Editor/UMGEditor/Private/UIComponentWidgetBlueprintExtension.cpp
@@ -165,8 +165,9 @@ void UUIComponentWidgetBlueprintExtension::HandlePopulateGeneratedVariables(cons
 				ComponentVariableDesc.VarGuid = FGuid::NewGuid();
 				ComponentVariableDesc.VarType = FEdGraphPinType(UEdGraphSchema_K2::PC_Object, NAME_None, ComponentTarget.GetComponent()->GetClass(), EPinContainerType::None, true, FEdGraphTerminalType());
 				ComponentVariableDesc.FriendlyName = ComponentTarget.GetComponent()->GetName();
-				ComponentVariableDesc.PropertyFlags = (CPF_BlueprintVisible | CPF_BlueprintReadOnly | CPF_Transient | CPF_RepSkip);
+				ComponentVariableDesc.PropertyFlags = (CPF_PersistentInstance | CPF_InstancedReference | CPF_BlueprintVisible | CPF_BlueprintReadOnly | CPF_Transient | CPF_RepSkip);
 				ComponentVariableDesc.SetMetaData(UUIComponentWidgetBlueprintExtension::MD_ComponentVariable, TEXT("true"));
+				ComponentVariableDesc.SetMetaData(FBlueprintMetadata::MD_FieldNotify, TEXT("true"));
 				ComponentVariableDesc.Category = FText::FromString(TEXT("Component"));
 				Context.AddGeneratedVariable(MoveTemp(ComponentVariableDesc));
 			}
diff --git a/Engine/Source/Runtime/UMG/Private/Extensions/UIComponent.cpp b/Engine/Source/Runtime/UMG/Private/Extensions/UIComponent.cpp
--- a/Engine/Source/Runtime/UMG/Private/Extensions/UIComponent.cpp
+++ b/Engine/Source/Runtime/UMG/Private/Extensions/UIComponent.cpp
@@ -5,6 +5,7 @@
 #include "CoreGlobals.h"
 #include "Components/Widget.h"
 #include "Blueprint/UserWidget.h"
+#include "Engine/BlueprintGeneratedClass.h"
 
 void UUIComponent::Initialize(UWidget* Target)
 {
@@ -53,4 +54,58 @@ void UUIComponent::OnConstruct()
 void UUIComponent::OnDestruct()
 {
 
+}
+
+void UUIComponent::FFieldNotificationClassDescriptor::ForEachField(const UClass* Class, TFunctionRef<bool(::UE::FieldNotification::FFieldId FieldId)> Callback) const
+{
+	if (const UBlueprintGeneratedClass* BPClass = Cast<const UBlueprintGeneratedClass>(Class))
+	{
+		BPClass->ForEachFieldNotify(Callback, true);
+	}
+}
+
+FDelegateHandle UUIComponent::AddFieldValueChangedDelegate(UE::FieldNotification::FFieldId InFieldId, FFieldValueChangedDelegate InNewDelegate)
+{
+	return (InFieldId.IsValid()) ? Delegates.Add(this, InFieldId, MoveTemp(InNewDelegate)) : FDelegateHandle();
+}
+
+bool UUIComponent::RemoveFieldValueChangedDelegate(UE::FieldNotification::FFieldId InFieldId, FDelegateHandle InHandle)
+{
+	if (InFieldId.IsValid() && InHandle.IsValid())
+	{
+		return Delegates.RemoveFrom(this, InFieldId, InHandle).bRemoved;
+	}
+	return false;
+}
+
+int32 UUIComponent::RemoveAllFieldValueChangedDelegates(const void* InUserObject)
+{
+	if (InUserObject)
+	{
+		return Delegates.RemoveAll(this, InUserObject).RemoveCount;
+	}
+	return false;
+}
+
+int32 UUIComponent::RemoveAllFieldValueChangedDelegates(UE::FieldNotification::FFieldId InFieldId, const void* InUserObject)
+{
+	if (InFieldId.IsValid() && InUserObject)
+	{
+		return Delegates.RemoveAll(this, InFieldId, InUserObject).RemoveCount;
+	}
+	return false;
+}
+
+void UUIComponent::BroadcastFieldValueChanged(UE::FieldNotification::FFieldId InFieldId)
+{
+	if (InFieldId.IsValid())
+	{
+		Delegates.Broadcast(this, InFieldId);
+	}
+}
+
+const UE::FieldNotification::IClassDescriptor& UUIComponent::GetFieldNotificationDescriptor() const
+{
+	static FFieldNotificationClassDescriptor Local;
+	return Local;
 }
\ No newline at end of file
diff --git a/Engine/Source/Runtime/UMG/Public/Extensions/UIComponent.h b/Engine/Source/Runtime/UMG/Public/Extensions/UIComponent.h
--- a/Engine/Source/Runtime/UMG/Public/Extensions/UIComponent.h
+++ b/Engine/Source/Runtime/UMG/Public/Extensions/UIComponent.h
@@ -3,7 +3,10 @@
 #pragma once
 
 #include "CoreMinimal.h"
+#include "FieldNotificationDelegate.h"
 #include "UObject/ObjectMacros.h"
+#include "IFieldNotificationClassDescriptor.h"
+#include "INotifyFieldValueChanged.h"
 
 #include "UIComponent.generated.h"
 
@@ -13,12 +16,16 @@ class UWidget;
  * This is the base class to for UI Components that can be added to any UMG Widgets
  * in UMG Designer.When initialized, it will pass the widget it's attached to.
  */
-UCLASS(Abstract, MinimalAPI, Experimental)
-class UUIComponent : public UObject
+UCLASS(Abstract, MinimalAPI, CustomFieldNotify, Experimental)
+class UUIComponent : public UObject, public INotifyFieldValueChanged
 {
 	GENERATED_BODY()
 
 public:
+	struct FFieldNotificationClassDescriptor : public ::UE::FieldNotification::IClassDescriptor
+	{
+		UMG_API virtual void ForEachField(const UClass* Class, TFunctionRef<bool(UE::FieldNotification::FFieldId FielId)> Callback) const override;
+	};
 	/**
 	 * Called when the owner widget is initialized.
 	 */
@@ -46,6 +53,15 @@ public:
 	*/
 	UMG_API TWeakObjectPtr<UWidget> GetOwner() const;
 
+	//~ Begin INotifyFieldValueChanged Interface
+	UMG_API virtual FDelegateHandle AddFieldValueChangedDelegate(UE::FieldNotification::FFieldId InFieldId, FFieldValueChangedDelegate InNewDelegate) override final;
+	UMG_API virtual bool RemoveFieldValueChangedDelegate(UE::FieldNotification::FFieldId InFieldId, FDelegateHandle InHandle) override final;
+	UMG_API virtual int32 RemoveAllFieldValueChangedDelegates(FDelegateUserObjectConst InUserObject) override final;
+	UMG_API virtual int32 RemoveAllFieldValueChangedDelegates(UE::FieldNotification::FFieldId InFieldId, FDelegateUserObjectConst InUserObject) override final;
+	UMG_API virtual void BroadcastFieldValueChanged(UE::FieldNotification::FFieldId InFieldId) override final;
+	UMG_API virtual const UE::FieldNotification::IClassDescriptor& GetFieldNotificationDescriptor() const override;
+	//~ End INotifyFieldValueChanged Interface
+
 protected:
 	UMG_API virtual void OnInitialize();
 	UMG_API virtual void OnPreConstruct(bool bIsDesignTime);
@@ -55,4 +71,5 @@ protected:
 private:
 	UPROPERTY(transient, DuplicateTransient)
 	TWeakObjectPtr<UWidget> Owner;
+	UE::FieldNotification::FFieldMulticastDelegate Delegates;
 };
