Support Event and Condition in Details Panel

#rnx

diff --git a/Engine/Plugins/Runtime/ModelViewViewModel/Source/ModelViewViewModelEditor/Private/Widgets/SMVVMViewBindingListView.cpp b/Engine/Plugins/Runtime/ModelViewViewModel/Source/ModelViewViewModelEditor/Private/Widgets/SMVVMViewBindingListView.cpp
--- a/Engine/Plugins/Runtime/ModelViewViewModel/Source/ModelViewViewModelEditor/Private/Widgets/SMVVMViewBindingListView.cpp
+++ b/Engine/Plugins/Runtime/ModelViewViewModel/Source/ModelViewViewModelEditor/Private/Widgets/SMVVMViewBindingListView.cpp
@@ -835,17 +835,50 @@ void SBindingsList::OnSourceListSelectionChanged(TSharedPtr<FBindingEntry> Entry
 			if (UMVVMBlueprintView* View = MVVMExtensionPtr->GetBlueprintView())
 			{
 				TArray<TSharedPtr<FBindingEntry>> SelectedEntries = TreeView->GetSelectedItems();
-				TArray<FMVVMBlueprintViewBinding*> SelectedBindings;
+				TArray<FBindingsSelectionVariantType> SelectionVariants;
 
 				for (const TSharedPtr<FBindingEntry>& SelectedEntry : SelectedEntries)
 				{
-					if (FMVVMBlueprintViewBinding* SelectedBinding = Entry->GetBinding(View))
+					switch (Entry->GetRowType())
 					{
-						SelectedBindings.Add(SelectedBinding);
+					case FBindingEntry::ERowType::Binding:
+					{
+						if (FMVVMBlueprintViewBinding* SelectedBinding = Entry->GetBinding(View))
+						{
+							FBindingsSelectionVariantType SelectionVariant;
+							SelectionVariant.Set<FMVVMBlueprintViewBinding*>(SelectedBinding);
+							SelectionVariants.Add(SelectionVariant);
+						}
+						break;
+					}
+					case FBindingEntry::ERowType::Condition:
+					{
+						if (UMVVMBlueprintViewCondition* SelectedCondition = Entry->GetCondition())
+						{
+							FBindingsSelectionVariantType SelectionVariant;
+							SelectionVariant.Set<UMVVMBlueprintViewCondition*>(SelectedCondition);
+							SelectionVariants.Add(SelectionVariant);
+						}
+						break;
+					}
+					case FBindingEntry::ERowType::Event:
+					{
+						if (UMVVMBlueprintViewEvent* SelectedEvent = Entry->GetEvent())
+						{
+							FBindingsSelectionVariantType SelectionVariant;
+							SelectionVariant.Set<UMVVMBlueprintViewEvent*>(SelectedEvent);
+							SelectionVariants.Add(SelectionVariant);
+						}
+						break;
+					}
+					default:
+					{
+						break;
+					}
 					}
 				}
 
-				BindingPanelPtr->OnBindingListSelectionChanged(SelectedBindings);
+				BindingPanelPtr->OnBindingListSelectionChanged(SelectionVariants);
 			}
 		}
 	}
diff --git a/Engine/Plugins/Runtime/ModelViewViewModel/Source/ModelViewViewModelEditor/Private/Widgets/SMVVMViewBindingPanel.cpp b/Engine/Plugins/Runtime/ModelViewViewModel/Source/ModelViewViewModelEditor/Private/Widgets/SMVVMViewBindingPanel.cpp
--- a/Engine/Plugins/Runtime/ModelViewViewModel/Source/ModelViewViewModelEditor/Private/Widgets/SMVVMViewBindingPanel.cpp
+++ b/Engine/Plugins/Runtime/ModelViewViewModel/Source/ModelViewViewModelEditor/Private/Widgets/SMVVMViewBindingPanel.cpp
@@ -12,6 +12,7 @@
 #include "MVVMWidgetBlueprintExtension_View.h"
 #include "WidgetBlueprintEditor.h"
 #include "MVVMBlueprintViewCondition.h"
+#include "MVVMBlueprintViewEvent.h"
 #include "WidgetBlueprintToolMenuContext.h"
 
 #include "Customizations/MVVMConversionPathCustomization.h"
@@ -46,6 +47,7 @@ namespace UE::MVVM
 
 namespace Private
 {
+
 struct FStructDetailNotifyHook : FNotifyHook
 {
 	virtual ~FStructDetailNotifyHook() = default;
@@ -209,41 +211,95 @@ void SBindingsPanel::HandleBlueprintViewChangedDelegate()
 	];
 }
 
-void SBindingsPanel::OnBindingListSelectionChanged(TConstArrayView<FMVVMBlueprintViewBinding*> Selection)
+void SBindingsPanel::OnBindingListSelectionChanged(TConstArrayView<FBindingsSelectionVariantType> Selection)
 {
 	NotifyHook->Binding = FMVVMBlueprintViewBinding();
 
-	// Currently only supports single object selection
-	if (FMVVMBlueprintViewBinding* Binding = Selection.Num() == 1 ? Selection.Last() : nullptr)
+	const int SelectionCount = Selection.Num();
+
+	if (Selection.Num() == 1)
 	{
-		NotifyHook->Binding = *Binding;
-		TSharedRef<FStructOnScope> StructScope = MakeShared<FStructOnScope>(FMVVMBlueprintViewBinding::StaticStruct(), reinterpret_cast<uint8*>(Binding));
-		StructDetailsView->SetStructureData(StructScope);
-		DetailContainer->SetContent(StructDetailsView->GetWidget().ToSharedRef());
+		TUnion<UObject*, TSharedPtr<FStructOnScope>> ObjectOrStruct;
+		FBindingsSelectionVariantType SelectionVariant = Selection.Last();
+
+		if (FMVVMBlueprintViewBinding* Binding = SelectionVariant.Get<FMVVMBlueprintViewBinding*>(nullptr))
+		{
+			NotifyHook->Binding = *Binding;
+
+			TSharedRef<FStructOnScope> StructScope = MakeShared<FStructOnScope>(FMVVMBlueprintViewBinding::StaticStruct(), reinterpret_cast<uint8*>(Binding));
+			ObjectOrStruct.SetSubtype<TSharedPtr<FStructOnScope>>(StructScope);
+		}
+
+		if (UMVVMBlueprintViewCondition* Condition = SelectionVariant.Get<UMVVMBlueprintViewCondition*>(nullptr))
+		{
+			ObjectOrStruct.SetSubtype<UObject*>(Condition);
+		}
+
+		if (UMVVMBlueprintViewEvent* Event = SelectionVariant.Get<UMVVMBlueprintViewEvent*>(nullptr))
+		{
+			ObjectOrStruct.SetSubtype<UObject*>(Event);
+		}
+
+		SetDetailsView(ObjectOrStruct);
+
 	}
 	else if(Selection.Num() > 1)
+	{
+		TUnion<UObject*, TSharedPtr<FStructOnScope>> EmptyObjectOrStruct;
+		SetDetailsView(EmptyObjectOrStruct, LOCTEXT("DetailsNotAvailable", "Details are not available when multiple bindings are selected.\nSelect only one binding to view details"));
+	}
+	else
+	{
+		TUnion<UObject*, TSharedPtr<FStructOnScope>> EmptyObjectOrStruct;
+		SetDetailsView(EmptyObjectOrStruct);
+	}
+}
+
+void SBindingsPanel::SetDetailsView(TUnion<UObject*, TSharedPtr<FStructOnScope>> ObjectOrStruct, const FText& ErrorMessage)
+{
+	if (!ErrorMessage.IsEmpty())
 	{
 		DetailsView->SetObject(nullptr);
 		StructDetailsView->SetStructureData(nullptr);
 		DetailContainer->SetContent(
 			SNew(SHorizontalBox)
 			+ SHorizontalBox::Slot()
-				.VAlign(VAlign_Center)
-				.HAlign(HAlign_Center)
-				.AutoWidth()
-				[
-					SNew(STextBlock)
-						.Text(LOCTEXT("DetailsNotAvailable", "Details are not available when multiple bindings are selected.\nSelect only one binding to view details"))
-				]
+			.VAlign(VAlign_Center)
+			.HAlign(HAlign_Center)
+			.AutoWidth()
+			[
+				SNew(STextBlock)
+					.Text(ErrorMessage)
+			]
 		);
+
+		return;
 	}
-	else
+
+	if (ObjectOrStruct.HasSubtype<UObject*>())
 	{
-		// show empty details view
-		DetailsView->SetObject(nullptr);
+		UObject* Object = ObjectOrStruct.GetSubtype<UObject*>();
+		DetailsView->SetObject(Object);
 		StructDetailsView->SetStructureData(nullptr);
 		DetailContainer->SetContent(DetailsView.ToSharedRef());
+
+		return;
+	}
+
+	if (ObjectOrStruct.HasSubtype<TSharedPtr<FStructOnScope>>())
+	{
+		TSharedPtr<FStructOnScope> SctructScope = ObjectOrStruct.GetSubtype<TSharedPtr<FStructOnScope>>();
+		DetailsView->SetObject(nullptr);
+		StructDetailsView->SetStructureData(SctructScope);
+		DetailContainer->SetContent(StructDetailsView->GetWidget().ToSharedRef());
+
+		return;
 	}
+
+	// Show empty details
+	DetailsView->SetObject(nullptr);
+	StructDetailsView->SetStructureData(nullptr);
+	DetailContainer->SetContent(DetailsView.ToSharedRef());
 }
 
 void SBindingsPanel::Tick(const FGeometry& AllottedGeometry, const double InCurrentTime, const float InDeltaTime)
diff --git a/Engine/Plugins/Runtime/ModelViewViewModel/Source/ModelViewViewModelEditor/Private/Widgets/SMVVMViewBindingPanel.h b/Engine/Plugins/Runtime/ModelViewViewModel/Source/ModelViewViewModelEditor/Private/Widgets/SMVVMViewBindingPanel.h
--- a/Engine/Plugins/Runtime/ModelViewViewModel/Source/ModelViewViewModelEditor/Private/Widgets/SMVVMViewBindingPanel.h
+++ b/Engine/Plugins/Runtime/ModelViewViewModel/Source/ModelViewViewModelEditor/Private/Widgets/SMVVMViewBindingPanel.h
@@ -2,16 +2,18 @@
 
 #pragma once
 
+#include "Containers/Union.h"
 #include "WidgetReference.h"
 #include "Widgets/SCompoundWidget.h"
 
-
 struct FMVVMBlueprintViewBinding;
 class FWidgetBlueprintEditor;
 class IDetailsView;
 class IStructureDetailsView;
 class SBorder;
 class UBlueprintExtension;
+class UMVVMBlueprintViewCondition;
+class UMVVMBlueprintViewEvent;
 class UMVVMWidgetBlueprintExtension_View;
 
 enum class ECheckBoxState : uint8;
@@ -21,6 +23,8 @@ namespace UE::MVVM
 class SBindingsList;
 namespace Private { struct FStructDetailNotifyHook; }
 
+using FBindingsSelectionVariantType = TVariant<FMVVMBlueprintViewBinding*, UMVVMBlueprintViewCondition*, UMVVMBlueprintViewEvent*>;
+
 class SBindingsPanel : public SCompoundWidget
 {
 private:
@@ -42,7 +46,7 @@ public:
 	virtual FReply OnDrop(const FGeometry& MyGeometry, const FDragDropEvent& DragDropEvent) override;
 	//~ End SWidget Interface
 
-	void OnBindingListSelectionChanged(TConstArrayView<FMVVMBlueprintViewBinding*> Selection);
+	void OnBindingListSelectionChanged(TConstArrayView<FBindingsSelectionVariantType> Selection);
 
 	static void RegisterSettingsMenu();
 
@@ -84,6 +88,7 @@ private:
 	bool IsDetailsViewEditingEnabled() const;
 	void RefreshDetailsView();
 	void RefreshNotifyHookBinding();
+	void SetDetailsView(TUnion<UObject*, TSharedPtr<FStructOnScope>> ObjectOrStruct, const FText& ErrorMessage = FText::GetEmpty());
 
 	EVisibility GetViewModelMessageVisibility() const;
 	EVisibility GetBindingMessageVisibility() const;
