Itemization - Support for input actions
#rnx
#okforversepublic

#changelist validated
#virtualized

diff --git a/Engine/Plugins/Runtime/ModelViewViewModel/ModelViewViewModel.uplugin b/Engine/Plugins/Runtime/ModelViewViewModel/ModelViewViewModel.uplugin
--- a/Engine/Plugins/Runtime/ModelViewViewModel/ModelViewViewModel.uplugin
+++ b/Engine/Plugins/Runtime/ModelViewViewModel/ModelViewViewModel.uplugin
@@ -44,6 +44,10 @@
 	],
 
 	"Plugins": [
+		{
+			"Name": "EnhancedInput",
+			"Enabled": true
+		},
 	]
 }
 
diff --git a/Engine/Plugins/Runtime/ModelViewViewModel/Source/ModelViewViewModelEditor/ModelViewViewModelEditor.Build.cs b/Engine/Plugins/Runtime/ModelViewViewModel/Source/ModelViewViewModelEditor/ModelViewViewModelEditor.Build.cs
--- a/Engine/Plugins/Runtime/ModelViewViewModel/Source/ModelViewViewModelEditor/ModelViewViewModelEditor.Build.cs
+++ b/Engine/Plugins/Runtime/ModelViewViewModel/Source/ModelViewViewModelEditor/ModelViewViewModelEditor.Build.cs
@@ -12,6 +12,7 @@ public class ModelViewViewModelEditor : ModuleRules
 			new string[]
 			{
 				"Engine",
+				"EnhancedInput",
 				"BlueprintGraph",
 				"Core",
 				"CoreUObject",
diff --git a/Engine/Plugins/Runtime/ModelViewViewModel/Source/ModelViewViewModelEditor/Private/Node/MVVMK2Node_LoadSoftResource.cpp b/Engine/Plugins/Runtime/ModelViewViewModel/Source/ModelViewViewModelEditor/Private/Node/MVVMK2Node_LoadSoftResource.cpp
--- a/Engine/Plugins/Runtime/ModelViewViewModel/Source/ModelViewViewModelEditor/Private/Node/MVVMK2Node_LoadSoftResource.cpp
+++ b/Engine/Plugins/Runtime/ModelViewViewModel/Source/ModelViewViewModelEditor/Private/Node/MVVMK2Node_LoadSoftResource.cpp
@@ -20,6 +20,7 @@
 #include "K2Node_Self.h"
 
 #include "Engine/Texture2D.h"
+#include "InputAction.h"
 #include "Materials/MaterialInterface.h"
 
 
@@ -158,4 +159,29 @@ UClass* UMVVMK2Node_LoadSoftMaterial::GetInputResourceClass() const
 	return UMaterialInterface::StaticClass();
 }
 
-#undef LOCTEXT_NAMESPACE
+
+//////////////////////////////////////////////////////////////////////////
+
+
+FText UMVVMK2Node_LoadSoftInputAction::GetTooltipText() const
+{
+	return FText(LOCTEXT("UMVVMK2Node_LoadSoftInputActionGetTooltipText", "Asynchronously loads a Soft Input Action Reference and returns that Input Action on successful load"));
+}
+
+FText UMVVMK2Node_LoadSoftInputAction::GetNodeTitle(ENodeTitleType::Type TitleType) const
+{
+	return FText(LOCTEXT("UMVVMK2Node_LoadSoftInputActionGetNodeTitle", "Load From Soft Input Action"));
+}
+
+const FName& UMVVMK2Node_LoadSoftInputAction::GetInputPinName() const
+{
+	static const FName InputPinName("InputAction");
+	return InputPinName;
+}
+
+UClass* UMVVMK2Node_LoadSoftInputAction::GetInputResourceClass() const
+{
+	return UInputAction::StaticClass();
+}
+
+#undef LOCTEXT_NAMESPACE
\ No newline at end of file
diff --git a/Engine/Plugins/Runtime/ModelViewViewModel/Source/ModelViewViewModelEditor/Private/Node/MVVMK2Node_LoadSoftResource.h b/Engine/Plugins/Runtime/ModelViewViewModel/Source/ModelViewViewModelEditor/Private/Node/MVVMK2Node_LoadSoftResource.h
--- a/Engine/Plugins/Runtime/ModelViewViewModel/Source/ModelViewViewModelEditor/Private/Node/MVVMK2Node_LoadSoftResource.h
+++ b/Engine/Plugins/Runtime/ModelViewViewModel/Source/ModelViewViewModelEditor/Private/Node/MVVMK2Node_LoadSoftResource.h
@@ -96,3 +96,24 @@ protected:
 	virtual UClass* GetInputResourceClass() const override;
 	//~ End UMVVMK2Node_LoadSoftResource Interface.
 };
+
+UCLASS(MinimalAPI)
+class UMVVMK2Node_LoadSoftInputAction : public UMVVMK2Node_LoadSoftResource
+{
+	GENERATED_BODY()
+
+public:
+	//~ Begin UEdGraphNode Interface.
+	virtual FText GetTooltipText() const override;
+	virtual FText GetNodeTitle(ENodeTitleType::Type TitleType) const override;
+	//~ End UEdGraphNode Interface.
+
+protected:
+	//~ Begin UK2Node_LoadAsset Interface.
+	virtual const FName& GetInputPinName() const;
+	//~ End UK2Node_LoadAsset Interface.
+
+	//~ Begin UMVVMK2Node_LoadSoftResource Interface.
+	virtual UClass* GetInputResourceClass() const override;
+	//~ End UMVVMK2Node_LoadSoftResource Interface.
+};
