Fix Transaction Undo in Bindings List, and UI Data synchronization issues

#rnx
#tests

diff --git a/Engine/Plugins/Runtime/ModelViewViewModel/Source/ModelViewViewModelBlueprint/Private/MVVMBlueprintView.cpp b/Engine/Plugins/Runtime/ModelViewViewModel/Source/ModelViewViewModelBlueprint/Private/MVVMBlueprintView.cpp
--- a/Engine/Plugins/Runtime/ModelViewViewModel/Source/ModelViewViewModelBlueprint/Private/MVVMBlueprintView.cpp
+++ b/Engine/Plugins/Runtime/ModelViewViewModel/Source/ModelViewViewModelBlueprint/Private/MVVMBlueprintView.cpp
@@ -296,6 +296,7 @@ UMVVMBlueprintViewEvent* UMVVMBlueprintView::DuplicateEvent(UMVVMBlueprintViewEv
 	Events.Add(NewEvent);
 
 	OnBindingsAdded.Broadcast();
+	OnBindingsUpdated.Broadcast();
 
 	FBlueprintEditorUtils::MarkBlueprintAsModified(GetOuterUMVVMWidgetBlueprintExtension_View()->GetWidgetBlueprint());
 
@@ -664,6 +665,16 @@ void UMVVMBlueprintView::PostEditChangeChainProperty(FPropertyChangedChainEvent&
 	}
 }
 
+void UMVVMBlueprintView::PostEditUndo()
+{
+	Super::PostEditUndo();
+
+	OnBindingsUpdated.Broadcast();
+	OnViewModelsUpdated.Broadcast();
+	OnEventsUpdated.Broadcast();
+	OnConditionsUpdated.Broadcast();
+}
+
 void UMVVMBlueprintView::AddAssetTags(TArray<FAssetRegistryTag>& OutTags) const
 {
 }
diff --git a/Engine/Plugins/Runtime/ModelViewViewModel/Source/ModelViewViewModelBlueprint/Public/MVVMBlueprintView.h b/Engine/Plugins/Runtime/ModelViewViewModel/Source/ModelViewViewModelBlueprint/Public/MVVMBlueprintView.h
--- a/Engine/Plugins/Runtime/ModelViewViewModel/Source/ModelViewViewModelBlueprint/Public/MVVMBlueprintView.h
+++ b/Engine/Plugins/Runtime/ModelViewViewModel/Source/ModelViewViewModelBlueprint/Public/MVVMBlueprintView.h
@@ -183,6 +183,7 @@ public:
 	virtual void PreSave(FObjectPreSaveContext Context) override;
 	virtual void PostEditChangeProperty(FPropertyChangedEvent& PropertyChangedEvent) override;
 	virtual void PostEditChangeChainProperty(FPropertyChangedChainEvent& PropertyChainEvent) override;
+	virtual void PostEditUndo() override;
 
 	void AddAssetTags(FAssetRegistryTagsContext Context) const;
 	UE_DEPRECATED(5.4, "Implement the version that takes FAssetRegistryTagsContext instead.")
diff --git a/Engine/Plugins/Runtime/ModelViewViewModel/Source/ModelViewViewModelEditor/Private/MVVMEditorSubsystem.cpp b/Engine/Plugins/Runtime/ModelViewViewModel/Source/ModelViewViewModelEditor/Private/MVVMEditorSubsystem.cpp
--- a/Engine/Plugins/Runtime/ModelViewViewModel/Source/ModelViewViewModelEditor/Private/MVVMEditorSubsystem.cpp
+++ b/Engine/Plugins/Runtime/ModelViewViewModel/Source/ModelViewViewModelEditor/Private/MVVMEditorSubsystem.cpp
@@ -434,6 +434,10 @@ bool UMVVMEditorSubsystem::ReparentViewModel(UWidgetBlueprint* WidgetBlueprint,
 FMVVMBlueprintViewBinding& UMVVMEditorSubsystem::AddBinding(UWidgetBlueprint* WidgetBlueprint)
 {
 	UMVVMBlueprintView* View = RequestView(WidgetBlueprint);
+
+	FScopedTransaction Transaction(LOCTEXT("MVVMAddBinding", "MVVM Add Binding"));
+	View->Modify();
+
 	return View->AddDefaultBinding();
 }
 
@@ -441,6 +445,9 @@ void UMVVMEditorSubsystem::RemoveBinding(UWidgetBlueprint* WidgetBlueprint, cons
 {
 	if (UMVVMBlueprintView* View = GetView(WidgetBlueprint))
 	{
+		FScopedTransaction Transaction(LOCTEXT("MVVMRemoveBinding", "MVVM Remove Binding"));
+		View->Modify();
+
 		View->RemoveBinding(&Binding);
 	}
 }
@@ -450,6 +457,10 @@ UMVVMBlueprintViewEvent* UMVVMEditorSubsystem::AddEvent(UWidgetBlueprint* Widget
 	if (GetDefault<UMVVMDeveloperProjectSettings>()->bAllowBindingEvent)
 	{
 		UMVVMBlueprintView* View = RequestView(WidgetBlueprint);
+
+		FScopedTransaction Transaction(LOCTEXT("MVVMRemoveEvent", "MVVM Add Event"));
+		View->Modify();
+
 		return View->AddDefaultEvent();
 	}
 	return nullptr;
@@ -459,6 +470,9 @@ void UMVVMEditorSubsystem::RemoveEvent(UWidgetBlueprint* WidgetBlueprint, UMVVMB
 {
 	if (UMVVMBlueprintView* View = GetView(WidgetBlueprint))
 	{
+		FScopedTransaction Transaction(LOCTEXT("MVVMRemoveEvent", "MVVM Remove Event"));
+		View->Modify();
+
 		View->RemoveEvent(Event);
 	}
 }
@@ -469,7 +483,12 @@ UMVVMBlueprintViewCondition* UMVVMEditorSubsystem::AddCondition(UWidgetBlueprint
 	{
 		return nullptr;
 	}
+
 	UMVVMBlueprintView* View = RequestView(WidgetBlueprint);
+
+	FScopedTransaction Transaction(LOCTEXT("MVVMAddCondition", "MVVM Add Condition"));
+	View->Modify();
+
 	return View->AddDefaultCondition();
 }
 
@@ -477,6 +496,9 @@ void UMVVMEditorSubsystem::RemoveCondition(UWidgetBlueprint* WidgetBlueprint, UM
 {
 	if (UMVVMBlueprintView* View = GetView(WidgetBlueprint))
 	{
+		FScopedTransaction Transaction(LOCTEXT("MVVMRemoveCondition", "MVVM Remove Condition"));
+		View->Modify();
+
 		View->RemoveCondition(Condition);
 	}
 }
