MVVM: only show blueprint-visible properties when browsing viewmodel properties through the bind button in details panel.

diff --git a/Engine/Plugins/Runtime/ModelViewViewModel/Source/ModelViewViewModelEditor/Private/Customizations/MVVMPropertyBindingExtension.cpp b/Engine/Plugins/Runtime/ModelViewViewModel/Source/ModelViewViewModelEditor/Private/Customizations/MVVMPropertyBindingExtension.cpp
--- a/Engine/Plugins/Runtime/ModelViewViewModel/Source/ModelViewViewModelEditor/Private/Customizations/MVVMPropertyBindingExtension.cpp
+++ b/Engine/Plugins/Runtime/ModelViewViewModel/Source/ModelViewViewModelEditor/Private/Customizations/MVVMPropertyBindingExtension.cpp
@@ -143,12 +143,36 @@ void FMVVMPropertyBindingExtension::ExtendBindingsMenu(FMenuBuilder& MenuBuilder
 				{
 					const UClass* ViewModelClass = ViewModel.GetViewModelClass();
 					const FProperty* Property = WidgetPropertyHandle->GetProperty();
+					auto IsPropertyVisible = [](const FProperty* VMProperty) -> bool
+						{
+							static FName NAME_BlueprintPrivate = "BlueprintPrivate";
+							if (VMProperty->HasMetaData(NAME_BlueprintPrivate))
+							{
+								return false;
+							}
+							static FName NAME_BlueprintProtected = "BlueprintProtected";
+							if (VMProperty->HasMetaData(NAME_BlueprintProtected))
+							{
+								return false;
+							}
+							if (!VMProperty->HasAnyPropertyFlags(CPF_BlueprintVisible | CPF_BlueprintAssignable))
+							{
+								return false;
+							}
+
+							return true;
+						};
 
 					MenuBuilder.BeginSection("ValidDataTypes", LOCTEXT("ValidDataTypes", "Valid Data Types"));
 
 					for (TFieldIterator<FProperty> It(ViewModelClass); It; ++It)
 					{
 						const FProperty* VMProperty = *It;
+						if (!IsPropertyVisible(VMProperty))
+						{
+							continue;
+						}
+
 						if (VMProperty->GetClass() != Property->GetClass())
 						{
 							continue;
@@ -169,6 +193,11 @@ void FMVVMPropertyBindingExtension::ExtendBindingsMenu(FMenuBuilder& MenuBuilder
 					for (TFieldIterator<FProperty> It(ViewModelClass); It; ++It)
 					{
 						const FProperty* VMProperty = *It;
+						if (!IsPropertyVisible(VMProperty))
+						{
+							continue;
+						}
+
 						if (VMProperty->GetClass() == Property->GetClass())
 						{
 							continue;
