UE-251567 Fix localization key for entry

#tests compiled
#rnx

diff --git a/Engine/Plugins/Runtime/ModelViewViewModel/Source/ModelViewViewModelEditor/Private/MVVMEditorSubsystem.cpp b/Engine/Plugins/Runtime/ModelViewViewModel/Source/ModelViewViewModelEditor/Private/MVVMEditorSubsystem.cpp
--- a/Engine/Plugins/Runtime/ModelViewViewModel/Source/ModelViewViewModelEditor/Private/MVVMEditorSubsystem.cpp
+++ b/Engine/Plugins/Runtime/ModelViewViewModel/Source/ModelViewViewModelEditor/Private/MVVMEditorSubsystem.cpp
@@ -458,7 +458,7 @@ UMVVMBlueprintViewEvent* UMVVMEditorSubsystem::AddEvent(UWidgetBlueprint* Widget
 	{
 		UMVVMBlueprintView* View = RequestView(WidgetBlueprint);
 
-		FScopedTransaction Transaction(LOCTEXT("MVVMRemoveEvent", "MVVM Add Event"));
+		FScopedTransaction Transaction(LOCTEXT("MVVMAddEvent", "MVVM Add Event"));
 		View->Modify();
 
 		return View->AddDefaultEvent();
