UMG: Context Menu action to Create bindings from selected widgets

#rnx

diff --git a/Engine/Plugins/Runtime/ModelViewViewModel/Source/ModelViewViewModelEditor/Private/Customizations/MVVMWidgetContextMenuExtension.cpp b/Engine/Plugins/Runtime/ModelViewViewModel/Source/ModelViewViewModelEditor/Private/Customizations/MVVMWidgetContextMenuExtension.cpp
new file mode 100644
--- /dev/null
+++ b/Engine/Plugins/Runtime/ModelViewViewModel/Source/ModelViewViewModelEditor/Private/Customizations/MVVMWidgetContextMenuExtension.cpp
@@ -0,0 +1,67 @@
+
+// Copyright Epic Games, Inc. All Rights Reserved.
+
+#include "Customizations/MVVMWidgetContextMenuExtension.h"
+#include "Framework/MVVMBindingEditorHelper.h"
+#include "MVVMEditorSubsystem.h"
+#include "WidgetBlueprint.h"
+#include "WidgetBlueprintEditor.h"
+
+#define LOCTEXT_NAMESPACE "FMVVMBindingEditorHelper"
+
+namespace UE::MVVM
+{
+namespace Private
+{
+
+void ExecuteCreateWidgetBindings(TSharedRef<FWidgetBlueprintEditor> BlueprintEditor)
+{
+	TSet<FWidgetReference> Widgets = BlueprintEditor->GetSelectedWidgets();
+	UWidgetBlueprint* Blueprint = BlueprintEditor->GetWidgetBlueprintObj();
+
+	TArray<FGuid> BindingIds;
+	UE::MVVM::FMVVMBindingEditorHelper::CreateWidgetBindings(Blueprint, Widgets, BindingIds);
+}
+
+bool CanCreateWidgetBindings(TSharedRef<FWidgetBlueprintEditor> BlueprintEditor)
+{
+	TSet<FWidgetReference> Widgets = BlueprintEditor->GetSelectedWidgets();
+	if (Widgets.IsEmpty())
+	{
+		return false;
+	}
+
+	UWidgetBlueprint* Blueprint = BlueprintEditor->GetWidgetBlueprintObj();
+	if (Blueprint == nullptr)
+	{
+		return false;
+	}
+
+	UMVVMEditorSubsystem* EditorSubsystem = GEditor->GetEditorSubsystem<UMVVMEditorSubsystem>();
+	if (EditorSubsystem == nullptr)
+	{
+		return false;
+	}
+
+	return true;
+}
+
+}
+
+void FWidgetContextMenuExtension::ExtendContextMenu(FMenuBuilder& MenuBuilder, TSharedRef<FWidgetBlueprintEditor> BlueprintEditor, FVector2D TargetLocation) const
+{
+	MenuBuilder.AddMenuEntry(
+		LOCTEXT("MVVM_CreateBinding", "Create Widget Binding"),
+		LOCTEXT("MVVM_CreateBindingTooltip", "Creates View Binding(s) for the currently selected widgets"),
+		FSlateIcon(),
+		FUIAction(
+			FExecuteAction::CreateStatic(&Private::ExecuteCreateWidgetBindings, BlueprintEditor),
+			FCanExecuteAction(),
+			FIsActionChecked(),
+			FIsActionButtonVisible::CreateStatic(&Private::CanCreateWidgetBindings, BlueprintEditor)
+		)
+	);
+}
+}
+
+#undef LOCTEXT_NAMESPACE
diff --git a/Engine/Plugins/Runtime/ModelViewViewModel/Source/ModelViewViewModelEditor/Private/Customizations/MVVMWidgetContextMenuExtension.h b/Engine/Plugins/Runtime/ModelViewViewModel/Source/ModelViewViewModelEditor/Private/Customizations/MVVMWidgetContextMenuExtension.h
new file mode 100644
--- /dev/null
+++ b/Engine/Plugins/Runtime/ModelViewViewModel/Source/ModelViewViewModelEditor/Private/Customizations/MVVMWidgetContextMenuExtension.h
@@ -0,0 +1,18 @@
+// Copyright Epic Games, Inc. All Rights Reserved.
+
+#pragma once
+
+#include "IHasWidgetContextMenuExtensibility.h"
+
+class FMenuBuilder;
+class FWidgetBlueprintEditor;
+
+namespace UE::MVVM
+{
+class FWidgetContextMenuExtension : public IWidgetContextMenuExtension
+{
+	//~ Begin IWidgetContextMenuExtension overrides
+	virtual void ExtendContextMenu(FMenuBuilder& MenuBuilder, TSharedRef<FWidgetBlueprintEditor> BlueprintEditor, FVector2D TargetLocation) const override;
+	//~ End IWidgetContextMenuExtension overrides
+};
+}
diff --git a/Engine/Plugins/Runtime/ModelViewViewModel/Source/ModelViewViewModelEditor/Private/Framework/MVVMBindingEditorHelper.cpp b/Engine/Plugins/Runtime/ModelViewViewModel/Source/ModelViewViewModelEditor/Private/Framework/MVVMBindingEditorHelper.cpp
new file mode 100644
--- /dev/null
+++ b/Engine/Plugins/Runtime/ModelViewViewModel/Source/ModelViewViewModelEditor/Private/Framework/MVVMBindingEditorHelper.cpp
@@ -0,0 +1,51 @@
+// Copyright Epic Games, Inc. All Rights Reserved.
+#include "MVVMBindingEditorHelper.h"
+
+#include "Framework/MultiBox/MultiBoxBuilder.h"
+#include "MVVMBlueprintViewBinding.h"
+#include "MVVMEditorSubsystem.h"
+#include "MVVMPropertyPath.h"
+#include "WidgetBlueprint.h"
+#include "WidgetBlueprintEditor.h"
+
+namespace UE::MVVM
+{
+
+bool FMVVMBindingEditorHelper::CreateWidgetBindings(UWidgetBlueprint* Blueprint, TSet<FWidgetReference> Widgets, TArray<FGuid>& OutBindingIds)
+{
+	if (Blueprint == nullptr)
+	{
+		return false;
+	}
+
+	UMVVMEditorSubsystem* EditorSubsystem = GEditor->GetEditorSubsystem<UMVVMEditorSubsystem>();
+	if (EditorSubsystem == nullptr)
+	{
+		return false;
+	}
+
+	for (const FWidgetReference& WidgetReference : Widgets)
+	{
+		if (WidgetReference.IsValid() && WidgetReference.GetTemplate())
+		{
+			FMVVMBlueprintViewBinding& Binding = EditorSubsystem->AddBinding(Blueprint);
+
+			FMVVMBlueprintPropertyPath Path;
+			if (WidgetReference.GetTemplate()->GetFName() == Blueprint->GetFName())
+			{
+				Path.SetSelfContext();
+			}
+			else
+			{
+				Path.SetWidgetName(WidgetReference.GetTemplate()->GetFName());
+			}
+
+			EditorSubsystem->SetDestinationPathForBinding(Blueprint, Binding, Path, false);
+			OutBindingIds.Add(Binding.BindingId);
+		}
+	}
+
+	return !OutBindingIds.IsEmpty();
+}
+
+} // namespace UE::MVVM
diff --git a/Engine/Plugins/Runtime/ModelViewViewModel/Source/ModelViewViewModelEditor/Private/Framework/MVVMBindingEditorHelper.h b/Engine/Plugins/Runtime/ModelViewViewModel/Source/ModelViewViewModelEditor/Private/Framework/MVVMBindingEditorHelper.h
new file mode 100644
--- /dev/null
+++ b/Engine/Plugins/Runtime/ModelViewViewModel/Source/ModelViewViewModelEditor/Private/Framework/MVVMBindingEditorHelper.h
@@ -0,0 +1,20 @@
+// Copyright Epic Games, Inc. All Rights Reserved.
+#pragma once
+
+#include "CoreMinimal.h"
+
+struct FWidgetReference;
+
+class FMenuBuilder;
+class FWidgetBlueprintEditor;
+class UWidgetBlueprint;
+
+namespace UE::MVVM
+{
+struct FMVVMBindingEditorHelper
+{
+public:
+	static bool CreateWidgetBindings(UWidgetBlueprint* Blueprint, TSet<FWidgetReference> Widgets, TArray<FGuid>& OutBindingIds);
+
+};
+} // namespace UE::MVVM
\ No newline at end of file
diff --git a/Engine/Plugins/Runtime/ModelViewViewModel/Source/ModelViewViewModelEditor/Private/ModelViewViewModelEditorModule.cpp b/Engine/Plugins/Runtime/ModelViewViewModel/Source/ModelViewViewModelEditor/Private/ModelViewViewModelEditorModule.cpp
--- a/Engine/Plugins/Runtime/ModelViewViewModel/Source/ModelViewViewModelEditor/Private/ModelViewViewModelEditorModule.cpp
+++ b/Engine/Plugins/Runtime/ModelViewViewModel/Source/ModelViewViewModelEditor/Private/ModelViewViewModelEditorModule.cpp
@@ -12,6 +12,7 @@
 #include "Customizations/MVVMListViewBaseExtensionCustomizationExtender.h"
 #include "Customizations/MVVMPanelWidgetExtensionCustomizationExtender.h"
 #include "Customizations/MVVMPropertyBindingExtension.h"
+#include "Customizations/MVVMWidgetContextMenuExtension.h"
 #include "Framework/Application/SlateApplication.h"
 #include "Framework/Docking/LayoutExtender.h"
 #include "Kismet2/BlueprintEditorUtils.h"
@@ -54,6 +55,9 @@ void FModelViewViewModelEditorModule::StartupModule()
 	DragDropExtension = MakeShared<UE::MVVM::FWidgetDragDropExtension>();
 	UMGEditorModule.GetWidgetDragDropExtensibilityManager()->AddExtension(DragDropExtension.ToSharedRef());
 
+	WidgetContextMenuCustomization = MakeShared<UE::MVVM::FWidgetContextMenuExtension>();
+	UMGEditorModule.GetWidgetContextMenuExtensibilityManager()->AddExtension(WidgetContextMenuCustomization.ToSharedRef());
+
 	ListViewBaseCustomizationExtender = UE::MVVM::FMVVMListViewBaseExtensionCustomizationExtender::MakeInstance();
 	UMGEditorModule.AddWidgetCustomizationExtender(ListViewBaseCustomizationExtender.ToSharedRef());
 
@@ -109,6 +113,7 @@ void FModelViewViewModelEditorModule::ShutdownModule()
 		UMGEditorModule->GetDesignerExtensibilityManager()->RemoveDesignerExtensionFactory(BlueprintViewDesignerExtensionFactory.ToSharedRef());
 		UMGEditorModule->OnRegisterTabsForEditor().RemoveAll(this);
 		UMGEditorModule->GetWidgetDragDropExtensibilityManager()->RemoveExtension(DragDropExtension.ToSharedRef());
+		UMGEditorModule->GetWidgetContextMenuExtensibilityManager()->RemoveExtension(WidgetContextMenuCustomization.ToSharedRef());
 		UMGEditorModule->GetClipboardExtensibilityManager()->RemoveExtension(ClipboardExtension.ToSharedRef());
 		UMGEditorModule->GetPropertyBindingExtensibilityManager()->RemoveExtension(PropertyBindingExtension.ToSharedRef());
 
diff --git a/Engine/Plugins/Runtime/ModelViewViewModel/Source/ModelViewViewModelEditor/Private/ModelViewViewModelEditorModule.h b/Engine/Plugins/Runtime/ModelViewViewModel/Source/ModelViewViewModelEditor/Private/ModelViewViewModelEditorModule.h
--- a/Engine/Plugins/Runtime/ModelViewViewModel/Source/ModelViewViewModelEditor/Private/ModelViewViewModelEditorModule.h
+++ b/Engine/Plugins/Runtime/ModelViewViewModel/Source/ModelViewViewModelEditor/Private/ModelViewViewModelEditorModule.h
@@ -11,6 +11,7 @@ namespace UE::MVVM
 	class FBlueprintViewDesignerExtensionFactory;
 	class FClipboardExtension;
 	class FWidgetDragDropExtension;
+	class FWidgetContextMenuExtension;
 	class FMVVMListViewBaseExtensionCustomizationExtender;
 	class FMVVMPanelWidgetExtensionCustomizationExtender;
 	class FMVVMPropertyBindingExtension;
@@ -51,6 +52,7 @@ private:
 	TSharedPtr<UE::MVVM::FMVVMPropertyBindingExtension> PropertyBindingExtension;
 	TSharedPtr<UE::MVVM::FClipboardExtension> ClipboardExtension;
 	TSharedPtr<UE::MVVM::FWidgetDragDropExtension> DragDropExtension;
+	TSharedPtr<UE::MVVM::FWidgetContextMenuExtension> WidgetContextMenuCustomization;
 	TSharedPtr<UE::MVVM::FAssetTypeActions_ViewModelBlueprint> ViewModelBlueprintActions;
 	TSharedPtr<UE::MVVM::FMVVMListViewBaseExtensionCustomizationExtender> ListViewBaseCustomizationExtender;
 	TSharedPtr<UE::MVVM::FMVVMPanelWidgetExtensionCustomizationExtender> PanelWidgetCustomizationExtender;
diff --git a/Engine/Plugins/Runtime/ModelViewViewModel/Source/ModelViewViewModelEditor/Private/Widgets/SMVVMViewBindingPanel.cpp b/Engine/Plugins/Runtime/ModelViewViewModel/Source/ModelViewViewModelEditor/Private/Widgets/SMVVMViewBindingPanel.cpp
--- a/Engine/Plugins/Runtime/ModelViewViewModel/Source/ModelViewViewModelEditor/Private/Widgets/SMVVMViewBindingPanel.cpp
+++ b/Engine/Plugins/Runtime/ModelViewViewModel/Source/ModelViewViewModelEditor/Private/Widgets/SMVVMViewBindingPanel.cpp
@@ -17,6 +17,7 @@
 
 #include "Customizations/MVVMConversionPathCustomization.h"
 #include "Customizations/MVVMPropertyPathCustomization.h"
+#include "Framework/MVVMBindingEditorHelper.h"
 #include "Hierarchy/HierarchyWidgetDragDropOp.h"
 #include "IStructureDetailsView.h"
 #include "PropertyEditorModule.h"
@@ -338,25 +339,11 @@ void SBindingsPanel::AddBindingToWidgetList(const TSet<FWidgetReference>& Widget
 
 			if (AddBindingMode == EAddBindingMode::Selected)
 			{
-				for (const FWidgetReference& WidgetReference : WidgetsToAddBinding)
+				TArray<FGuid> BindingIds;
+				if (FMVVMBindingEditorHelper::CreateWidgetBindings(MVVMExtensionPtr->GetWidgetBlueprint(), WidgetsToAddBinding, BindingIds))
 				{
-					if (WidgetReference.IsValid() && WidgetReference.GetTemplate())
-					{
-						FMVVMBlueprintViewBinding& Binding = EditorSubsystem->AddBinding(MVVMExtensionPtr->GetWidgetBlueprint());
-						FMVVMBlueprintPropertyPath Path;
-						if (WidgetReference.GetTemplate()->GetFName() == MVVMExtensionPtr->GetWidgetBlueprint()->GetFName())
-						{
-							Path.SetSelfContext();
-						}
-						else
-						{
-							Path.SetWidgetName(WidgetReference.GetTemplate()->GetFName());
-						}
-						EditorSubsystem->SetDestinationPathForBinding(MVVMExtensionPtr->GetWidgetBlueprint(), Binding, Path, false);
-						AddedBindingId = Binding.BindingId;
-
-						bBindingAdded = true;
-					}
+					bBindingAdded = true;
+					AddedBindingId = BindingIds.Last();
 				}
 			}
 
diff --git a/Engine/Source/Editor/UMGEditor/Private/UMGEditorModule.cpp b/Engine/Source/Editor/UMGEditor/Private/UMGEditorModule.cpp
--- a/Engine/Source/Editor/UMGEditor/Private/UMGEditorModule.cpp
+++ b/Engine/Source/Editor/UMGEditor/Private/UMGEditorModule.cpp
@@ -102,6 +102,7 @@ public:
 		PropertyBindingExtensibilityManager = MakeShared<FPropertyBindingExtensibilityManager>();
 		ClipboardExtensibilityManager = MakeShared<FClipboardExtensibilityManager>();
 		WidgetDragDropExtensibilityManager = MakeShared<FWidgetDragDropExtensibilityManager>();
+		WidgetContextMenuExtensibilityManager = MakeShared<FWidgetContextMenuExtensibilityManager >();
 
 		UIComponentCustomizationExtender = FUIComponentCustomizationExtender::MakeInstance();
 		AddWidgetCustomizationExtender(UIComponentCustomizationExtender.ToSharedRef());
@@ -246,6 +247,7 @@ public:
 	virtual TSharedPtr<FPropertyBindingExtensibilityManager> GetPropertyBindingExtensibilityManager() override { return PropertyBindingExtensibilityManager; }
 	virtual TSharedPtr<FClipboardExtensibilityManager> GetClipboardExtensibilityManager() override { return ClipboardExtensibilityManager; }
 	virtual TSharedPtr<FWidgetDragDropExtensibilityManager> GetWidgetDragDropExtensibilityManager() override { return WidgetDragDropExtensibilityManager; }
+	virtual TSharedPtr<FWidgetContextMenuExtensibilityManager> GetWidgetContextMenuExtensibilityManager() override { return WidgetContextMenuExtensibilityManager; }
 
 	/** Register settings objects. */
 	void RegisterSettings()
@@ -515,6 +517,7 @@ private:
 	TSharedPtr<FPropertyBindingExtensibilityManager> PropertyBindingExtensibilityManager;
 	TSharedPtr<FClipboardExtensibilityManager> ClipboardExtensibilityManager;
 	TSharedPtr<FWidgetDragDropExtensibilityManager> WidgetDragDropExtensibilityManager;
+	TSharedPtr<FWidgetContextMenuExtensibilityManager> WidgetContextMenuExtensibilityManager;
 	TSharedPtr<FGraphPanelPinFactory> GraphPanelPinFactory;
 
 	FDelegateHandle SequenceEditorHandle;
diff --git a/Engine/Source/Editor/UMGEditor/Private/WidgetBlueprintEditorUtils.cpp b/Engine/Source/Editor/UMGEditor/Private/WidgetBlueprintEditorUtils.cpp
--- a/Engine/Source/Editor/UMGEditor/Private/WidgetBlueprintEditorUtils.cpp
+++ b/Engine/Source/Editor/UMGEditor/Private/WidgetBlueprintEditorUtils.cpp
@@ -531,6 +531,13 @@ void FWidgetBlueprintEditorUtils::CreateWidgetContextMenu(FMenuBuilder& MenuBuil
 		}
 	}
 	MenuBuilder.EndSection();
+
+	IUMGEditorModule& EditorModule = FModuleManager::LoadModuleChecked<IUMGEditorModule>("UMGEditor");
+	const TArrayView<const TSharedPtr<IWidgetContextMenuExtension>> ContextMenuExtensions = EditorModule.GetWidgetContextMenuExtensibilityManager()->GetExtensions();
+	for (const TSharedPtr<IWidgetContextMenuExtension>& ContextMenuExtension : ContextMenuExtensions)
+	{
+		ContextMenuExtension->ExtendContextMenu(MenuBuilder, BlueprintEditor, TargetLocation);
+	}
 }
 
 void FWidgetBlueprintEditorUtils::ExecuteOpenSelectedWidgetsForEdit( TSet<FWidgetReference> SelectedWidgets )
diff --git a/Engine/Source/Editor/UMGEditor/Public/IHasWidgetContextMenuExtensibility.h b/Engine/Source/Editor/UMGEditor/Public/IHasWidgetContextMenuExtensibility.h
new file mode 100644
--- /dev/null
+++ b/Engine/Source/Editor/UMGEditor/Public/IHasWidgetContextMenuExtensibility.h
@@ -0,0 +1,52 @@
+// Copyright Epic Games, Inc. All Rights Reserved.
+
+#pragma once
+
+#include "CoreMinimal.h"
+
+class FMenuBuilder;
+class FWidgetBlueprintEditor;
+
+class UMGEDITOR_API IWidgetContextMenuExtension
+{
+public:
+	virtual ~IWidgetContextMenuExtension() { }
+
+	virtual void ExtendContextMenu(FMenuBuilder& MenuBuilder, TSharedRef<FWidgetBlueprintEditor> BlueprintEditor, FVector2D TargetLocation) const = 0;
+};
+
+/**
+ * Widget Context Menu extensibility manager holds a list of registered widget context menu extensions.
+ */
+class UMGEDITOR_API FWidgetContextMenuExtensibilityManager
+{
+public:
+	void AddExtension(const TSharedRef<IWidgetContextMenuExtension>& Extension)
+	{
+		if (ensure(!Extensions.Contains(Extension)))
+		{
+			Extensions.Add(Extension);
+		}
+	}
+
+	void RemoveExtension(const TSharedRef<IWidgetContextMenuExtension>& Extension)
+	{
+		int32 NumRemoved = Extensions.RemoveSingleSwap(Extension);
+		ensure(NumRemoved == 1);
+	}
+
+	TArrayView<const TSharedPtr<IWidgetContextMenuExtension>> GetExtensions() const
+	{
+		return Extensions;
+	}
+
+private:
+	TArray<TSharedPtr<IWidgetContextMenuExtension>> Extensions;
+};
+
+/** Indicates that a class can extend drag & drop functionality */
+class IHasWidgetContextMenuExtensibility
+{
+public:
+	virtual TSharedPtr<FWidgetContextMenuExtensibilityManager> GetWidgetContextMenuExtensibilityManager() = 0;
+};
diff --git a/Engine/Source/Editor/UMGEditor/Public/UMGEditorModule.h b/Engine/Source/Editor/UMGEditor/Public/UMGEditorModule.h
--- a/Engine/Source/Editor/UMGEditor/Public/UMGEditorModule.h
+++ b/Engine/Source/Editor/UMGEditor/Public/UMGEditorModule.h
@@ -8,6 +8,7 @@
 #include "IHasDesignerExtensibility.h"
 #include "IHasPropertyBindingExtensibility.h"
 #include "IHasWidgetDragDropExtensibility.h"
+#include "IHasWidgetContextMenuExtensibility.h"
 #include "UObject/TopLevelAssetPath.h"
 
 extern const FName UMGEditorAppIdentifier;
@@ -26,7 +27,8 @@ class IUMGEditorModule :
 	public IHasDesignerExtensibility,
 	public IHasPropertyBindingExtensibility,
 	public IHasClipboardExtensibility,
-	public IHasWidgetDragDropExtensibility
+	public IHasWidgetDragDropExtensibility,
+	public IHasWidgetContextMenuExtensibility
 {
 public:
 	virtual class FWidgetBlueprintCompiler* GetRegisteredCompiler() = 0;
