UI Components: don't allow subclasses of components as viewmodels

diff --git a/Engine/Plugins/Runtime/ModelViewViewModel/Source/ModelViewViewModelEditor/Private/Widgets/SMVVMSelectViewModel.cpp b/Engine/Plugins/Runtime/ModelViewViewModel/Source/ModelViewViewModelEditor/Private/Widgets/SMVVMSelectViewModel.cpp
--- a/Engine/Plugins/Runtime/ModelViewViewModel/Source/ModelViewViewModelEditor/Private/Widgets/SMVVMSelectViewModel.cpp
+++ b/Engine/Plugins/Runtime/ModelViewViewModel/Source/ModelViewViewModelEditor/Private/Widgets/SMVVMSelectViewModel.cpp
@@ -3,6 +3,7 @@
 #include "Widgets/SMVVMSelectViewModel.h"
 
 #include "Components/Widget.h"
+#include "Extensions/UIComponent.h"
 #include "Modules/ModuleManager.h"
 #include "SPrimaryButton.h"
 #include "Types/MVVMBindingSource.h"
@@ -21,7 +22,7 @@ namespace Private
 {
 	bool IsValidViewModel(const UClass* InClass, EClassFlags DisallowedClassFlags)
 	{
-		if (InClass == nullptr || InClass->IsChildOf(UWidget::StaticClass()))
+		if (InClass == nullptr || InClass->IsChildOf(UWidget::StaticClass()) || InClass->IsChildOf(UUIComponent::StaticClass()))
 		{
 			return false;
 		}
