MVVM: Ensure unique conversion functions when duplicating bindings

#rnx

diff --git a/Engine/Plugins/Runtime/ModelViewViewModel/Source/ModelViewViewModelBlueprint/Private/MVVMBlueprintView.cpp b/Engine/Plugins/Runtime/ModelViewViewModel/Source/ModelViewViewModelBlueprint/Private/MVVMBlueprintView.cpp
--- a/Engine/Plugins/Runtime/ModelViewViewModel/Source/ModelViewViewModelBlueprint/Private/MVVMBlueprintView.cpp
+++ b/Engine/Plugins/Runtime/ModelViewViewModel/Source/ModelViewViewModelBlueprint/Private/MVVMBlueprintView.cpp
@@ -217,10 +217,24 @@ const FMVVMBlueprintViewBinding* UMVVMBlueprintView::DuplicateBinding(const FMVV
 	NewBinding = *Binding;
 	NewBinding.BindingId = FGuid::NewGuid();
 
+	UWidgetBlueprint* WidgetBlueprint = GetOuterUMVVMWidgetBlueprintExtension_View()->GetWidgetBlueprint();
+
+	if (UMVVMBlueprintViewConversionFunction* ConversionFunction = Binding->Conversion.SourceToDestinationConversion)
+	{
+		NewBinding.Conversion.SourceToDestinationConversion = DuplicateObject<UMVVMBlueprintViewConversionFunction>(ConversionFunction, WidgetBlueprint);
+		NewBinding.Conversion.SourceToDestinationConversion->RecreateWrapperGraph(WidgetBlueprint);
+	}
+
+	if (UMVVMBlueprintViewConversionFunction* ConversionFunction = Binding->Conversion.DestinationToSourceConversion)
+	{
+		NewBinding.Conversion.DestinationToSourceConversion = DuplicateObject<UMVVMBlueprintViewConversionFunction>(ConversionFunction, WidgetBlueprint);
+		NewBinding.Conversion.DestinationToSourceConversion->RecreateWrapperGraph(WidgetBlueprint);
+	}
+
 	OnBindingsAdded.Broadcast();
 	OnBindingsUpdated.Broadcast();
 
-	FBlueprintEditorUtils::MarkBlueprintAsModified(GetOuterUMVVMWidgetBlueprintExtension_View()->GetWidgetBlueprint());
+	FBlueprintEditorUtils::MarkBlueprintAsModified(WidgetBlueprint);
 
 	return &NewBinding;
 }
diff --git a/Engine/Plugins/Runtime/ModelViewViewModel/Source/ModelViewViewModelBlueprint/Private/MVVMBlueprintViewConversionFunction.cpp b/Engine/Plugins/Runtime/ModelViewViewModel/Source/ModelViewViewModelBlueprint/Private/MVVMBlueprintViewConversionFunction.cpp
--- a/Engine/Plugins/Runtime/ModelViewViewModel/Source/ModelViewViewModelBlueprint/Private/MVVMBlueprintViewConversionFunction.cpp
+++ b/Engine/Plugins/Runtime/ModelViewViewModel/Source/ModelViewViewModelBlueprint/Private/MVVMBlueprintViewConversionFunction.cpp
@@ -375,15 +375,20 @@ void UMVVMBlueprintViewConversionFunction::SetCachedWrapperGraph(UBlueprint* Con
 	}
 }
 
+FName UMVVMBlueprintViewConversionFunction::GenerateUniqueGraphName() const
+{
+	TStringBuilder<256> StringBuilder;
+	StringBuilder << UE::MVVM::Private::DefaultConversionFunctionName.Resolve();
+	StringBuilder << FGuid::NewDeterministicGuid(GetFullName()).ToString(EGuidFormats::DigitsWithHyphensLower);
+	return StringBuilder.ToString();
+}
+
 void UMVVMBlueprintViewConversionFunction::CreateWrapperGraphName()
 {
 	ensure(!GraphName.IsNone());
 	if (GraphName.IsNone())
 	{
-		TStringBuilder<256> StringBuilder;
-		StringBuilder << UE::MVVM::Private::DefaultConversionFunctionName.Resolve();
-		StringBuilder << FGuid::NewDeterministicGuid(GetFullName()).ToString(EGuidFormats::DigitsWithHyphensLower);
-		GraphName = StringBuilder.ToString();
+		GraphName = GenerateUniqueGraphName();
 	}
 }
 
@@ -456,6 +461,15 @@ UEdGraph* UMVVMBlueprintViewConversionFunction::GetOrCreateWrapperGraph(UBluepri
 	return nullptr;
 }
 
+void UMVVMBlueprintViewConversionFunction::RecreateWrapperGraph(UBlueprint* Blueprint)
+{
+	if (IsValid(Blueprint))
+	{
+		GraphName = GenerateUniqueGraphName();
+		CreateWrapperGraphInternal(Blueprint);
+	}
+}
+
 UEdGraphPin* UMVVMBlueprintViewConversionFunction::GetOrCreateGraphPin(UBlueprint* Blueprint, const FMVVMBlueprintPinId& PinId)
 {
 	GetOrCreateWrapperGraph(Blueprint);
diff --git a/Engine/Plugins/Runtime/ModelViewViewModel/Source/ModelViewViewModelBlueprint/Public/MVVMBlueprintViewConversionFunction.h b/Engine/Plugins/Runtime/ModelViewViewModel/Source/ModelViewViewModelBlueprint/Public/MVVMBlueprintViewConversionFunction.h
--- a/Engine/Plugins/Runtime/ModelViewViewModel/Source/ModelViewViewModelBlueprint/Public/MVVMBlueprintViewConversionFunction.h
+++ b/Engine/Plugins/Runtime/ModelViewViewModel/Source/ModelViewViewModelBlueprint/Public/MVVMBlueprintViewConversionFunction.h
@@ -46,7 +46,7 @@ public:
 
 	/** Set the function. Generate a Graph. */
 	void Initialize(UBlueprint* SelfContext, FName GraphName, FMVVMBlueprintFunctionReference Function);
-	
+
 	/** Set the function. Generate a Graph. */
 	void InitializeFromFunction(UBlueprint* SelfContext, FName GraphName, const UFunction* Function);
 
@@ -113,6 +113,8 @@ public:
 	/** If needed, create the graph and all the nodes for that graph. */
 	UEdGraph* GetOrCreateWrapperGraph(UBlueprint* Blueprint);
 
+	void RecreateWrapperGraph(UBlueprint* Blueprint);
+
 	/**
 	 * The conversion function is going to be removed from the Blueprint.
 	 * Do any cleanup that is needed.
@@ -152,6 +154,7 @@ private:
 	UEdGraph* CreateWrapperGraphInternal(UBlueprint* Blueprint);
 	bool NeedsWrapperGraphInternal(const UClass* SkeletalSelfContext) const;
 	void LoadPinValuesInternal(UBlueprint* Blueprint);
+	FName GenerateUniqueGraphName() const;
 	void CreateWrapperGraphName();
 	void Reset();
 
