MVVM: Fix Drag/Drop tooltip, add additional drag/drop support.
1. Add support for drag/drop conversion function parameters
2. When users drop properties into an empty bindings.
2.a ViewModel properties will create binding with property as source
2.b Widget Details properties will create binding with property as source

diff --git a/Engine/Plugins/Runtime/ModelViewViewModel/Source/ModelViewViewModelEditor/Private/Customizations/MVVMPropertyBindingExtension.cpp b/Engine/Plugins/Runtime/ModelViewViewModel/Source/ModelViewViewModelEditor/Private/Customizations/MVVMPropertyBindingExtension.cpp
--- a/Engine/Plugins/Runtime/ModelViewViewModel/Source/ModelViewViewModelEditor/Private/Customizations/MVVMPropertyBindingExtension.cpp
+++ b/Engine/Plugins/Runtime/ModelViewViewModel/Source/ModelViewViewModelEditor/Private/Customizations/MVVMPropertyBindingExtension.cpp
@@ -411,7 +411,15 @@ bool FMVVMPropertyBindingExtension::CanExtend(const UWidgetBlueprint* WidgetBlue
 		return false;
 	}
 	const UMVVMBlueprintView* MVVMBlueprintView = MVVMExtensionPtr->GetBlueprintView();
-	return MVVMBlueprintView ? MVVMBlueprintView->GetViewModels().Num() > 0 : false;
+	if (MVVMBlueprintView != nullptr)
+	{
+		if (MVVMBlueprintView->GetViewModels().Num() > 0 || MVVMBlueprintView->GetNumBindings() > 0)
+		{
+			return true;
+		}
+	}
+
+	return false;
 }
 
 IPropertyBindingExtension::EDropResult FMVVMPropertyBindingExtension::OnDrop(const FGeometry& Geometry, const FDragDropEvent& DragDropEvent, UWidgetBlueprint* WidgetBlueprint, UWidget* Widget, TSharedPtr<IPropertyHandle> WidgetPropertyHandle)
diff --git a/Engine/Plugins/Runtime/ModelViewViewModel/Source/ModelViewViewModelEditor/Private/Framework/MVVMRowHelper.cpp b/Engine/Plugins/Runtime/ModelViewViewModel/Source/ModelViewViewModelEditor/Private/Framework/MVVMRowHelper.cpp
--- a/Engine/Plugins/Runtime/ModelViewViewModel/Source/ModelViewViewModelEditor/Private/Framework/MVVMRowHelper.cpp
+++ b/Engine/Plugins/Runtime/ModelViewViewModel/Source/ModelViewViewModelEditor/Private/Framework/MVVMRowHelper.cpp
@@ -240,7 +240,7 @@ void FRowHelper::ShowBlueprintGraph(FBlueprintEditor* BlueprintEditor, UWidgetBl
 	}
 }
 
-TOptional<FMVVMBlueprintPropertyPath> FRowHelper::DropFieldSelector(const UWidgetBlueprint* WidgetBlueprint, const FDragDropEvent& DragDropEvent)
+TOptional<FMVVMBlueprintPropertyPath> FRowHelper::DropFieldSelector(const UWidgetBlueprint* WidgetBlueprint, const FDragDropEvent& DragDropEvent, bool bIsSource)
 {
 	if (WidgetBlueprint == nullptr)
 	{
@@ -253,89 +253,101 @@ TOptional<FMVVMBlueprintPropertyPath> FRowHelper::DropFieldSelector(const UWidge
 		return TOptional<FMVVMBlueprintPropertyPath>();
 	}
 
-	// Accept all drag-drop operations that are widget properties, but only accept view model fields when we are dropping into the Source box.
-	if (!DragDropOp->IsOfType<FWidgetPropertyDragDropOp>() && (!DragDropOp->IsOfType<FViewModelFieldDragDropOp>()))
+	if (DragDropOp->IsOfType<FViewModelFieldDragDropOp>())
 	{
-		return TOptional<FMVVMBlueprintPropertyPath>();
-	}
+		// Accept valid view model fields when we are dropping into the Source box.
+		TSharedPtr<FViewModelFieldDragDropOp> ViewModelFieldDragDropOp = DragDropEvent.GetOperationAs<FViewModelFieldDragDropOp>();
+		if (!bIsSource || !ViewModelFieldDragDropOp->ViewModelId.IsValid())
+		{
+			return TOptional<FMVVMBlueprintPropertyPath>();
+		}
 
-	TSharedPtr<FViewModelFieldDragDropOp> ViewModelFieldDragDropOp = DragDropEvent.GetOperationAs<FViewModelFieldDragDropOp>();
-	TSharedPtr<FWidgetPropertyDragDropOp> WidgetPropertyDragDropOp = DragDropEvent.GetOperationAs<FWidgetPropertyDragDropOp>();
-	bool bIsViewModelProperty = ViewModelFieldDragDropOp != nullptr;
+		UWidgetBlueprint* DragDropWidgetBP = ViewModelFieldDragDropOp->WidgetBP.Get();
+		if (WidgetBlueprint != DragDropWidgetBP)
+		{
+			return TOptional<FMVVMBlueprintPropertyPath>();
+		}
 
-	UWidgetBlueprint* DragDropWidgetBP = bIsViewModelProperty ? ViewModelFieldDragDropOp->WidgetBP.Get() : WidgetPropertyDragDropOp->WidgetBP.Get();
-	if (WidgetBlueprint != DragDropWidgetBP)
-	{
-		return TOptional<FMVVMBlueprintPropertyPath>();
-	}
+		TArray<FFieldVariant> FieldPath = ViewModelFieldDragDropOp->DraggedField;
+		FMVVMBlueprintPropertyPath PropertyPath;
+		for (const FFieldVariant& Field : FieldPath)
+		{
+			PropertyPath.AppendPropertyPath(WidgetBlueprint, FMVVMConstFieldVariant(Field));
+		}
 
-	TArray<FFieldVariant> FieldPath = bIsViewModelProperty ? ViewModelFieldDragDropOp->DraggedField : WidgetPropertyDragDropOp->DraggedPropertyPath;
-	FMVVMBlueprintPropertyPath PropertyPath;
-	for (const FFieldVariant& Field : FieldPath)
-	{
-		PropertyPath.AppendPropertyPath(WidgetBlueprint, FMVVMConstFieldVariant(Field));
+		UMVVMEditorSubsystem* Subsystem = GEditor->GetEditorSubsystem<UMVVMEditorSubsystem>();
+		PropertyPath.SetViewModelId(ViewModelFieldDragDropOp->ViewModelId);
+		return PropertyPath;
 	}
-
-	UMVVMEditorSubsystem* Subsystem = GEditor->GetEditorSubsystem<UMVVMEditorSubsystem>();
-	if (bIsViewModelProperty)
+	else if (DragDropOp->IsOfType<FWidgetPropertyDragDropOp>())
 	{
-		if (ViewModelFieldDragDropOp->ViewModelId.IsValid())
+		TSharedPtr<FWidgetPropertyDragDropOp> WidgetPropertyDragDropOp = DragDropEvent.GetOperationAs<FWidgetPropertyDragDropOp>();
+		
+		UWidgetBlueprint* DragDropWidgetBP = WidgetPropertyDragDropOp->WidgetBP.Get();
+		UWidget* OwnerWidgetPtr = WidgetPropertyDragDropOp->OwnerWidget.Get();
+		if (WidgetBlueprint != DragDropWidgetBP || OwnerWidgetPtr == nullptr)
 		{
-			PropertyPath.SetViewModelId(ViewModelFieldDragDropOp->ViewModelId);
+			return TOptional<FMVVMBlueprintPropertyPath>();
 		}
-		else
+
+		TArray<FFieldVariant> FieldPath = WidgetPropertyDragDropOp->DraggedPropertyPath;
+		FMVVMBlueprintPropertyPath PropertyPath;
+		for (const FFieldVariant& Field : FieldPath)
 		{
-			return TOptional<FMVVMBlueprintPropertyPath>();
+			PropertyPath.AppendPropertyPath(WidgetBlueprint, FMVVMConstFieldVariant(Field));
 		}
-	}
-	else
-	{
-		if (UWidget* OwnerWidgetPtr = WidgetPropertyDragDropOp->OwnerWidget.Get())
+
+		if (OwnerWidgetPtr->GetClass() == WidgetBlueprint->GeneratedClass)
 		{
-			// Search for the widget by its name in the widget tree
-			// If the widget is not found, we know it is the root preview widget so we use the blueprint name.
-			if (WidgetBlueprint->WidgetTree->FindWidget(OwnerWidgetPtr->GetFName()) && WidgetBlueprint->GetFName() != OwnerWidgetPtr->GetFName())
-			{
-				PropertyPath.SetWidgetName(OwnerWidgetPtr->GetFName());
-			}
-			else
-			{
-				PropertyPath.SetSelfContext();
-			}
+			PropertyPath.SetSelfContext();
+		}
+		else
+		{
+			PropertyPath.SetWidgetName(OwnerWidgetPtr->GetFName());
 		}
+
+		return PropertyPath;
 	}
 
-	return PropertyPath;
+	return TOptional<FMVVMBlueprintPropertyPath>();
 }
 
-void FRowHelper::DragEnterFieldSelector(const UWidgetBlueprint* WidgetBlueprint, const FDragDropEvent& DragDropEvent)
+FReply FRowHelper::DragOverFieldSelector(const UWidgetBlueprint* WidgetBlueprint, const FDragDropEvent& DragDropEvent, bool bIsSource)
 {
 	TSharedPtr<FDecoratedDragDropOp> DragDropOp = DragDropEvent.GetOperationAs<FDecoratedDragDropOp>();
-	if (DragDropOp.IsValid())
+	if (!DragDropOp.IsValid())
 	{
-		// Accept all drag-drop operations that are widget properties, but only accept view model fields when we are dropping into the Source box.
-		if (DragDropOp->IsOfType<FWidgetPropertyDragDropOp>() || DragDropOp->IsOfType<FViewModelFieldDragDropOp>())
-		{
-			TSharedPtr<FViewModelFieldDragDropOp> ViewModelFieldDragDropOp = DragDropEvent.GetOperationAs<FViewModelFieldDragDropOp>();
-			TSharedPtr<FWidgetPropertyDragDropOp> WidgetPropertyDragDropOp = DragDropEvent.GetOperationAs<FWidgetPropertyDragDropOp>();
-			bool IsViewModelProperty = false;
+		return FReply::Unhandled();
+	}
 
-			if (ViewModelFieldDragDropOp)
-			{
-				IsViewModelProperty = true;
-			}
+	// Accept view model fields when we are dropping into the Source box.
+	if (DragDropOp->IsOfType<FViewModelFieldDragDropOp>() && bIsSource)
+	{
+		TSharedPtr<FViewModelFieldDragDropOp> ViewModelFieldDragDropOp = DragDropEvent.GetOperationAs<FViewModelFieldDragDropOp>();
 
-			UWidgetBlueprint* DragDropWidgetBP = IsViewModelProperty ? ViewModelFieldDragDropOp->WidgetBP.Get() : WidgetPropertyDragDropOp->WidgetBP.Get();
-			if (DragDropWidgetBP && DragDropWidgetBP == WidgetBlueprint)
-			{
-				DragDropOp->CurrentIconBrush = FAppStyle::GetBrush(TEXT("Graph.ConnectorFeedback.OK"));
-			}
+		UWidgetBlueprint* DragDropWidgetBP = ViewModelFieldDragDropOp->WidgetBP.Get();
+		if (DragDropWidgetBP == WidgetBlueprint)
+		{
+			DragDropOp->CurrentIconBrush = FAppStyle::GetBrush(TEXT("Graph.ConnectorFeedback.OK"));
 		}
-		else
+	}
+	else if (DragDropOp->IsOfType<FWidgetPropertyDragDropOp>())
+	{
+		TSharedPtr<FWidgetPropertyDragDropOp> WidgetPropertyDragDropOp = DragDropEvent.GetOperationAs<FWidgetPropertyDragDropOp>();
+
+		UWidgetBlueprint* DragDropWidgetBP = WidgetPropertyDragDropOp->WidgetBP.Get();
+		UWidget* OwnerWidgetPtr = WidgetPropertyDragDropOp->OwnerWidget.Get();
+		if (WidgetBlueprint == DragDropWidgetBP && OwnerWidgetPtr != nullptr)
 		{
-			DragDropOp->CurrentIconBrush = FAppStyle::GetBrush(TEXT("Graph.ConnectorFeedback.Error"));
+			DragDropOp->CurrentIconBrush = FAppStyle::GetBrush(TEXT("Graph.ConnectorFeedback.OK"));
 		}
 	}
+	else
+	{
+		DragDropOp->CurrentIconBrush = FAppStyle::GetBrush(TEXT("Graph.ConnectorFeedback.Error"));
+	}
+
+	return FReply::Handled();
 }
 
 namespace Private
diff --git a/Engine/Plugins/Runtime/ModelViewViewModel/Source/ModelViewViewModelEditor/Private/Framework/MVVMRowHelper.h b/Engine/Plugins/Runtime/ModelViewViewModel/Source/ModelViewViewModelEditor/Private/Framework/MVVMRowHelper.h
--- a/Engine/Plugins/Runtime/ModelViewViewModel/Source/ModelViewViewModelEditor/Private/Framework/MVVMRowHelper.h
+++ b/Engine/Plugins/Runtime/ModelViewViewModel/Source/ModelViewViewModelEditor/Private/Framework/MVVMRowHelper.h
@@ -32,8 +32,8 @@ struct FRowHelper
 
 	static void ShowBlueprintGraph(FBlueprintEditor* Editor, UWidgetBlueprint* WidgetBlueprint, UMVVMBlueprintView* BlueprintView, TArrayView<const TSharedPtr<FBindingEntry>> Selection);
 
-	static TOptional<FMVVMBlueprintPropertyPath> DropFieldSelector(const UWidgetBlueprint* WidgetBlueprint, const FDragDropEvent& DragDropEvent);
-	static void DragEnterFieldSelector(const UWidgetBlueprint* WidgetBlueprint, const FDragDropEvent& DragDropEvent);
+	static TOptional<FMVVMBlueprintPropertyPath> DropFieldSelector(const UWidgetBlueprint* WidgetBlueprint, const FDragDropEvent& DragDropEvent, bool bIsSource);
+	static FReply DragOverFieldSelector(const UWidgetBlueprint* WidgetBlueprint, const FDragDropEvent& DragDropEvent, bool bIsSource);
 
 	static FMenuBuilder CreateContextMenu(UWidgetBlueprint* WidgetBlueprint, UMVVMBlueprintView* View, TConstArrayView<TSharedPtr<FBindingEntry>> Entries, FOnContextMenuEntryCallback OnSelectionChanged = FOnContextMenuEntryCallback());
 };
diff --git a/Engine/Plugins/Runtime/ModelViewViewModel/Source/ModelViewViewModelEditor/Private/Widgets/BindingEntry/SMVVMBindingRow.cpp b/Engine/Plugins/Runtime/ModelViewViewModel/Source/ModelViewViewModelEditor/Private/Widgets/BindingEntry/SMVVMBindingRow.cpp
--- a/Engine/Plugins/Runtime/ModelViewViewModel/Source/ModelViewViewModelEditor/Private/Widgets/BindingEntry/SMVVMBindingRow.cpp
+++ b/Engine/Plugins/Runtime/ModelViewViewModel/Source/ModelViewViewModelEditor/Private/Widgets/BindingEntry/SMVVMBindingRow.cpp
@@ -124,7 +124,7 @@ TSharedRef<SWidget> SBindingRow::BuildRowWidget()
 					.OnSelectionChanged(this, &SBindingRow::HandleFieldSelectionChanged, false)
 					.OnGetSelectionContext(this, &SBindingRow::GetSelectedSelectionContext, false)
 					.OnDrop(this, &SBindingRow::HandleFieldSelectorDrop, false)
-					.OnDragEnter(this, &SBindingRow::HandleFieldSelectorDragEnter, false)
+					.OnDragOver(this, &SBindingRow::HandleFieldSelectorDragOver, false)
 					.ShowContext(false)
 					.CanCreateEvent(true)
 					.ToolTipText(FText::Join(FText::FromString("\n"), LOCTEXT("MVVMDestTooltipHeader", "Destination"), LOCTEXT("MVVMDestTooltipBody", "The property to hook data into.")))
@@ -171,7 +171,7 @@ TSharedRef<SWidget> SBindingRow::BuildRowWidget()
 					.OnSelectionChanged(this, &SBindingRow::HandleFieldSelectionChanged, true)
 					.OnGetSelectionContext(this, &SBindingRow::GetSelectedSelectionContext, true)
 					.OnDrop(this, &SBindingRow::HandleFieldSelectorDrop, true)
-					.OnDragEnter(this, &SBindingRow::HandleFieldSelectorDragEnter, true)
+					.OnDragOver(this, &SBindingRow::HandleFieldSelectorDragOver, true)
 					.ToolTipText(FText::Join(FText::FromString("\n"), LOCTEXT("MVVMSrcTooltipHeader", "Source"), LOCTEXT("MVVMSrcTooltipBody", "The source of data to bind to.")))
 				]
 			]
@@ -510,10 +510,10 @@ FReply SBindingRow::HandleFieldSelectorDrop(const FGeometry& MyGeometry, const F
 		return FReply::Unhandled();
 	}
 
-	TOptional<FMVVMBlueprintPropertyPath> PropertyPath = BindingEntry::FRowHelper::DropFieldSelector(WidgetBlueprint, DragDropEvent);
+	TOptional<FMVVMBlueprintPropertyPath> PropertyPath = BindingEntry::FRowHelper::DropFieldSelector(WidgetBlueprint, DragDropEvent, bSource);
 	if (!PropertyPath.IsSet())
 	{
-		return FReply::Unhandled();
+		return FReply::Handled();
 	}
 
 	UMVVMEditorSubsystem* Subsystem = GetEditorSubsystem();
@@ -528,9 +528,9 @@ FReply SBindingRow::HandleFieldSelectorDrop(const FGeometry& MyGeometry, const F
 	return FReply::Handled();
 }
 
-void SBindingRow::HandleFieldSelectorDragEnter(const FGeometry& MyGeometry, const FDragDropEvent& DragDropEvent, bool bSource)
+FReply SBindingRow::HandleFieldSelectorDragOver(const FGeometry& MyGeometry, const FDragDropEvent& DragDropEvent, bool bSource)
 {
-	BindingEntry::FRowHelper::DragEnterFieldSelector(GetBlueprint(), DragDropEvent);
+	return BindingEntry::FRowHelper::DragOverFieldSelector(GetBlueprint(), DragDropEvent, bSource);
 }
 
 ECheckBoxState SBindingRow::IsExecutionModeOverrideChecked() const
diff --git a/Engine/Plugins/Runtime/ModelViewViewModel/Source/ModelViewViewModelEditor/Private/Widgets/BindingEntry/SMVVMBindingRow.h b/Engine/Plugins/Runtime/ModelViewViewModel/Source/ModelViewViewModelEditor/Private/Widgets/BindingEntry/SMVVMBindingRow.h
--- a/Engine/Plugins/Runtime/ModelViewViewModel/Source/ModelViewViewModelEditor/Private/Widgets/BindingEntry/SMVVMBindingRow.h
+++ b/Engine/Plugins/Runtime/ModelViewViewModel/Source/ModelViewViewModelEditor/Private/Widgets/BindingEntry/SMVVMBindingRow.h
@@ -59,7 +59,7 @@ private:
 
 	FReply HandleFieldSelectorDrop(const FGeometry& MyGeometry, const FDragDropEvent& DragDropEvent, bool bSource);
 
-	void HandleFieldSelectorDragEnter(const FGeometry& MyGeometry, const FDragDropEvent& DragDropEvent, bool bSource);
+	FReply HandleFieldSelectorDragOver(const FGeometry& MyGeometry, const FDragDropEvent& DragDropEvent, bool bSource);
 
 	ECheckBoxState IsExecutionModeOverrideChecked() const;
 
diff --git a/Engine/Plugins/Runtime/ModelViewViewModel/Source/ModelViewViewModelEditor/Private/Widgets/BindingEntry/SMVVMConditionRow.cpp b/Engine/Plugins/Runtime/ModelViewViewModel/Source/ModelViewViewModelEditor/Private/Widgets/BindingEntry/SMVVMConditionRow.cpp
--- a/Engine/Plugins/Runtime/ModelViewViewModel/Source/ModelViewViewModelEditor/Private/Widgets/BindingEntry/SMVVMConditionRow.cpp
+++ b/Engine/Plugins/Runtime/ModelViewViewModel/Source/ModelViewViewModelEditor/Private/Widgets/BindingEntry/SMVVMConditionRow.cpp
@@ -122,7 +122,7 @@ TSharedRef<SWidget> SConditionRow::BuildRowWidget()
 						.OnSelectionChanged(this, &SConditionRow::HandleFieldSelectionChanged, true)
 						.OnGetSelectionContext(this, &SConditionRow::GetSelectedSelectionContext, true)
 						.OnDrop(this, &SConditionRow::HandleFieldSelectorDrop, true)
-						.OnDragEnter(this, &SConditionRow::HandleFieldSelectorDragEnter, true)
+						.OnDragOver(this, &SConditionRow::HandleFieldSelectorOver, true)
 				]
 			]
 			+ SHorizontalBox::Slot()
@@ -224,7 +224,7 @@ TSharedRef<SWidget> SConditionRow::BuildRowWidget()
 					.OnSelectionChanged(this, &SConditionRow::HandleFieldSelectionChanged, false)
 					.OnGetSelectionContext(this, &SConditionRow::GetSelectedSelectionContext, false)
 					.OnDrop(this, &SConditionRow::HandleFieldSelectorDrop, false)
-					.OnDragEnter(this, &SConditionRow::HandleFieldSelectorDragEnter, false)
+					.OnDragOver(this, &SConditionRow::HandleFieldSelectorOver, false)
 					.IsBindingToEvent(true)
 				]
 			]
@@ -617,10 +617,10 @@ FReply SConditionRow::HandleFieldSelectorDrop(const FGeometry& MyGeometry, const
 		return FReply::Unhandled();
 	}
 
-	TOptional<FMVVMBlueprintPropertyPath> PropertyPath = BindingEntry::FRowHelper::DropFieldSelector(GetBlueprint(), DragDropEvent);
+	TOptional<FMVVMBlueprintPropertyPath> PropertyPath = BindingEntry::FRowHelper::DropFieldSelector(GetBlueprint(), DragDropEvent, bCondition);
 	if (!PropertyPath.IsSet())
 	{
-		return FReply::Unhandled();
+		return FReply::Handled();
 	}
 
 	if (bCondition)
@@ -634,9 +634,9 @@ FReply SConditionRow::HandleFieldSelectorDrop(const FGeometry& MyGeometry, const
 	return FReply::Handled();
 }
 
-void SConditionRow::HandleFieldSelectorDragEnter(const FGeometry& MyGeometry, const FDragDropEvent& DragDropEvent, bool bEvent)
+FReply SConditionRow::HandleFieldSelectorOver(const FGeometry& MyGeometry, const FDragDropEvent& DragDropEvent, bool bCondition)
 {
-	BindingEntry::FRowHelper::DragEnterFieldSelector(GetBlueprint(), DragDropEvent);
+	return BindingEntry::FRowHelper::DragOverFieldSelector(GetBlueprint(), DragDropEvent, bCondition);
 }
 
 TSharedRef<SWidget> SConditionRow::HandleContextMenu() const
diff --git a/Engine/Plugins/Runtime/ModelViewViewModel/Source/ModelViewViewModelEditor/Private/Widgets/BindingEntry/SMVVMConditionRow.h b/Engine/Plugins/Runtime/ModelViewViewModel/Source/ModelViewViewModelEditor/Private/Widgets/BindingEntry/SMVVMConditionRow.h
--- a/Engine/Plugins/Runtime/ModelViewViewModel/Source/ModelViewViewModelEditor/Private/Widgets/BindingEntry/SMVVMConditionRow.h
+++ b/Engine/Plugins/Runtime/ModelViewViewModel/Source/ModelViewViewModelEditor/Private/Widgets/BindingEntry/SMVVMConditionRow.h
@@ -51,15 +51,15 @@ private:
 
 	void OnIsConditionCompileChanged(ECheckBoxState NewState);
 
-	FMVVMLinkedPinValue GetFieldSelectedValue(bool bSource) const;
+	FMVVMLinkedPinValue GetFieldSelectedValue(bool bCondition) const;
 
-	void HandleFieldSelectionChanged(FMVVMLinkedPinValue Value, SFieldSelectorMenu::ESelectionType SelectionType, bool bSource);
+	void HandleFieldSelectionChanged(FMVVMLinkedPinValue Value, SFieldSelectorMenu::ESelectionType SelectionType, bool bCondition);
 
-	FFieldSelectionContext GetSelectedSelectionContext(bool bSource) const;
+	FFieldSelectionContext GetSelectedSelectionContext(bool bCondition) const;
 
-	FReply HandleFieldSelectorDrop(const FGeometry& MyGeometry, const FDragDropEvent& DragDropEvent, bool bSource);
+	FReply HandleFieldSelectorDrop(const FGeometry& MyGeometry, const FDragDropEvent& DragDropEvent, bool bCondition);
 
-	void HandleFieldSelectorDragEnter(const FGeometry& MyGeometry, const FDragDropEvent& DragDropEvent, bool bSource);
+	FReply HandleFieldSelectorOver(const FGeometry& MyGeometry, const FDragDropEvent& DragDropEvent, bool bCondition);
 
 	TSharedRef<SWidget> HandleContextMenu() const;
 	void HandleShowBlueprintGraph() const;
diff --git a/Engine/Plugins/Runtime/ModelViewViewModel/Source/ModelViewViewModelEditor/Private/Widgets/BindingEntry/SMVVMEventRow.cpp b/Engine/Plugins/Runtime/ModelViewViewModel/Source/ModelViewViewModelEditor/Private/Widgets/BindingEntry/SMVVMEventRow.cpp
--- a/Engine/Plugins/Runtime/ModelViewViewModel/Source/ModelViewViewModelEditor/Private/Widgets/BindingEntry/SMVVMEventRow.cpp
+++ b/Engine/Plugins/Runtime/ModelViewViewModel/Source/ModelViewViewModelEditor/Private/Widgets/BindingEntry/SMVVMEventRow.cpp
@@ -87,7 +87,7 @@ TSharedRef<SWidget> SEventRow::BuildRowWidget()
 					.OnSelectionChanged(this, &SEventRow::HandleFieldSelectionChanged, true)
 					.OnGetSelectionContext(this, &SEventRow::GetSelectedSelectionContext, true)
 					.OnDrop(this, &SEventRow::HandleFieldSelectorDrop, true)
-					.OnDragEnter(this, &SEventRow::HandleFieldSelectorDragEnter, true)
+					.OnDragOver(this, &SEventRow::HandleFieldSelectorDragOver, true)
 					.ShowContext(false)
 					.CanCreateEvent(true)
 				]
@@ -124,7 +124,7 @@ TSharedRef<SWidget> SEventRow::BuildRowWidget()
 					.OnSelectionChanged(this, &SEventRow::HandleFieldSelectionChanged, false)
 					.OnGetSelectionContext(this, &SEventRow::GetSelectedSelectionContext, false)
 					.OnDrop(this, &SEventRow::HandleFieldSelectorDrop, false)
-					.OnDragEnter(this, &SEventRow::HandleFieldSelectorDragEnter, false)
+					.OnDragOver(this, &SEventRow::HandleFieldSelectorDragOver, false)
 					.IsBindingToEvent(true)
 				]
 			]
@@ -371,10 +371,10 @@ FReply SEventRow::HandleFieldSelectorDrop(const FGeometry& MyGeometry, const FDr
 		return FReply::Unhandled();
 	}
 
-	TOptional<FMVVMBlueprintPropertyPath> PropertyPath = BindingEntry::FRowHelper::DropFieldSelector(GetBlueprint(), DragDropEvent);
+	TOptional<FMVVMBlueprintPropertyPath> PropertyPath = BindingEntry::FRowHelper::DropFieldSelector(GetBlueprint(), DragDropEvent, bEvent);
 	if (!PropertyPath.IsSet())
 	{
-		return FReply::Unhandled();
+		return FReply::Handled();
 	}
 
 	if (bEvent)
@@ -388,9 +388,9 @@ FReply SEventRow::HandleFieldSelectorDrop(const FGeometry& MyGeometry, const FDr
 	return FReply::Handled();
 }
 
-void SEventRow::HandleFieldSelectorDragEnter(const FGeometry& MyGeometry, const FDragDropEvent& DragDropEvent, bool bEvent)
+FReply SEventRow::HandleFieldSelectorDragOver(const FGeometry& MyGeometry, const FDragDropEvent& DragDropEvent, bool bEvent)
 {
-	BindingEntry::FRowHelper::DragEnterFieldSelector(GetBlueprint(), DragDropEvent);
+	return BindingEntry::FRowHelper::DragOverFieldSelector(GetBlueprint(), DragDropEvent, bEvent);
 }
 
 TSharedRef<SWidget> SEventRow::HandleContextMenu() const
diff --git a/Engine/Plugins/Runtime/ModelViewViewModel/Source/ModelViewViewModelEditor/Private/Widgets/BindingEntry/SMVVMEventRow.h b/Engine/Plugins/Runtime/ModelViewViewModel/Source/ModelViewViewModelEditor/Private/Widgets/BindingEntry/SMVVMEventRow.h
--- a/Engine/Plugins/Runtime/ModelViewViewModel/Source/ModelViewViewModelEditor/Private/Widgets/BindingEntry/SMVVMEventRow.h
+++ b/Engine/Plugins/Runtime/ModelViewViewModel/Source/ModelViewViewModelEditor/Private/Widgets/BindingEntry/SMVVMEventRow.h
@@ -51,15 +51,15 @@ private:
 
 	void OnIsEventCompileChanged(ECheckBoxState NewState);
 
-	FMVVMLinkedPinValue GetFieldSelectedValue(bool bSource) const;
+	FMVVMLinkedPinValue GetFieldSelectedValue(bool bEvent) const;
 
-	void HandleFieldSelectionChanged(FMVVMLinkedPinValue Value, SFieldSelectorMenu::ESelectionType SelectionType, bool bSource);
+	void HandleFieldSelectionChanged(FMVVMLinkedPinValue Value, SFieldSelectorMenu::ESelectionType SelectionType, bool bEvent);
 
-	FFieldSelectionContext GetSelectedSelectionContext(bool bSource) const;
+	FFieldSelectionContext GetSelectedSelectionContext(bool bEvent) const;
 
-	FReply HandleFieldSelectorDrop(const FGeometry& MyGeometry, const FDragDropEvent& DragDropEvent, bool bSource);
+	FReply HandleFieldSelectorDrop(const FGeometry& MyGeometry, const FDragDropEvent& DragDropEvent, bool bEvent);
 
-	void HandleFieldSelectorDragEnter(const FGeometry& MyGeometry, const FDragDropEvent& DragDropEvent, bool bSource);
+	FReply HandleFieldSelectorDragOver(const FGeometry& MyGeometry, const FDragDropEvent& DragDropEvent, bool bEvent);
 
 	TSharedRef<SWidget> HandleContextMenu() const;
 	void HandleShowBlueprintGraph() const;
diff --git a/Engine/Plugins/Runtime/ModelViewViewModel/Source/ModelViewViewModelEditor/Private/Widgets/SMVVMFieldSelector.cpp b/Engine/Plugins/Runtime/ModelViewViewModel/Source/ModelViewViewModelEditor/Private/Widgets/SMVVMFieldSelector.cpp
--- a/Engine/Plugins/Runtime/ModelViewViewModel/Source/ModelViewViewModelEditor/Private/Widgets/SMVVMFieldSelector.cpp
+++ b/Engine/Plugins/Runtime/ModelViewViewModel/Source/ModelViewViewModelEditor/Private/Widgets/SMVVMFieldSelector.cpp
@@ -32,8 +32,8 @@ void SFieldSelector::Construct(const FArguments& InArgs, const UWidgetBlueprint*
 
 	OnSelectionChanged = InArgs._OnSelectionChanged;
 	OnGetSelectionContext = InArgs._OnGetSelectionContext;
-	OnDragEnterEvent = InArgs._OnDragEnter;
 	OnDropEvent = InArgs._OnDrop;
+	OnDragOverEvent = InArgs._OnDragOver;
 	bIsBindingToEvent = InArgs._IsBindingToEvent;
 	bCanCreateEvent = InArgs._CanCreateEvent;
 
@@ -104,20 +104,22 @@ void SFieldSelector::HandleFieldSelectionChanged(FMVVMLinkedPinValue LinkedValue
 	}
 }
 
-void SFieldSelector::OnDragEnter(const FGeometry& MyGeometry, const FDragDropEvent& DragDropEvent)
+void SFieldSelector::OnDragLeave(const FDragDropEvent& DragDropEvent)
 {
-	if (OnDragEnterEvent.IsBound())
+	if (TSharedPtr<FDecoratedDragDropOp> DecoratedDragDropOp = DragDropEvent.GetOperationAs<FDecoratedDragDropOp>())
 	{
-		OnDragEnterEvent.Execute(MyGeometry, DragDropEvent);
+		DecoratedDragDropOp->ResetToDefaultToolTip();
 	}
 }
 
-void SFieldSelector::OnDragLeave(const FDragDropEvent& DragDropEvent)
+FReply SFieldSelector::OnDragOver(const FGeometry& MyGeometry, const FDragDropEvent& DragDropEvent)
 {
-	if (TSharedPtr<FDecoratedDragDropOp> DecoratedDragDropOp = DragDropEvent.GetOperationAs<FDecoratedDragDropOp>())
+	if (OnDragOverEvent.IsBound())
 	{
-		DecoratedDragDropOp->ResetToDefaultToolTip();
+		return OnDragOverEvent.Execute(MyGeometry, DragDropEvent);
 	}
+
+	return FReply::Unhandled();
 }
 
 FReply SFieldSelector::OnDrop(const FGeometry& MyGeometry, const FDragDropEvent& DragDropEvent)
diff --git a/Engine/Plugins/Runtime/ModelViewViewModel/Source/ModelViewViewModelEditor/Private/Widgets/SMVVMFieldSelector.h b/Engine/Plugins/Runtime/ModelViewViewModel/Source/ModelViewViewModelEditor/Private/Widgets/SMVVMFieldSelector.h
--- a/Engine/Plugins/Runtime/ModelViewViewModel/Source/ModelViewViewModelEditor/Private/Widgets/SMVVMFieldSelector.h
+++ b/Engine/Plugins/Runtime/ModelViewViewModel/Source/ModelViewViewModelEditor/Private/Widgets/SMVVMFieldSelector.h
@@ -21,7 +21,7 @@ class SFieldSelector : public SCompoundWidget
 public:
 	DECLARE_DELEGATE_RetVal(FFieldSelectionContext, FOnGetSelectionContext);
 	DECLARE_DELEGATE_RetVal_TwoParams(FReply, FOnDrop, const FGeometry&, const FDragDropEvent&);
-	DECLARE_DELEGATE_TwoParams(FOnDragEnter, const FGeometry&, const FDragDropEvent&);
+	DECLARE_DELEGATE_RetVal_TwoParams(FReply, FOnDragOver, const FGeometry&, const FDragDropEvent&);
 
 	SLATE_BEGIN_ARGS(SFieldSelector)
 		: _TextStyle(&FCoreStyle::Get().GetWidgetStyle<FTextBlockStyle>("NormalText"))
@@ -35,14 +35,14 @@ public:
 		SLATE_EVENT(SFieldSelectorMenu::FOnLinkedValueSelected, OnSelectionChanged)
 		SLATE_EVENT(FOnGetSelectionContext, OnGetSelectionContext)
 		SLATE_EVENT(FOnDrop, OnDrop)
-		SLATE_EVENT(FOnDragEnter, OnDragEnter)
+		SLATE_EVENT(FOnDragOver, OnDragOver)
 		SLATE_ARGUMENT_DEFAULT(bool, CanCreateEvent) = false;
 	SLATE_END_ARGS()
 
 	void Construct(const FArguments& InArgs, const UWidgetBlueprint* InWidgetBlueprint);
 	virtual FReply OnDrop(const FGeometry& MyGeometry, const FDragDropEvent& DragDropEvent) override;
-	virtual void OnDragEnter(const FGeometry& MyGeometry, const FDragDropEvent& DragDropEvent) override;
 	virtual void OnDragLeave(const FDragDropEvent& DragDropEvent) override;
+	virtual FReply OnDragOver(const FGeometry& MyGeometry, const FDragDropEvent& DragDropEvent) override;
 
 private:
 	TSharedRef<SWidget> CreateSourcePanel();
@@ -59,7 +59,7 @@ private:
 	SFieldSelectorMenu::FOnLinkedValueSelected OnSelectionChanged;
 	FOnGetSelectionContext OnGetSelectionContext;
 	FOnDrop OnDropEvent;
-	FOnDragEnter OnDragEnterEvent;
+	FOnDragOver OnDragOverEvent;
 	bool bIsBindingToEvent = false;
 	bool bCanCreateEvent = false;
 }; 
diff --git a/Engine/Plugins/Runtime/ModelViewViewModel/Source/ModelViewViewModelEditor/Private/Widgets/SMVVMFunctionParameter.cpp b/Engine/Plugins/Runtime/ModelViewViewModel/Source/ModelViewViewModelEditor/Private/Widgets/SMVVMFunctionParameter.cpp
--- a/Engine/Plugins/Runtime/ModelViewViewModel/Source/ModelViewViewModelEditor/Private/Widgets/SMVVMFunctionParameter.cpp
+++ b/Engine/Plugins/Runtime/ModelViewViewModel/Source/ModelViewViewModelEditor/Private/Widgets/SMVVMFunctionParameter.cpp
@@ -3,6 +3,7 @@
 #include "SMVVMFunctionParameter.h" 
 #include "EdGraphSchema_K2.h"
 #include "Editor.h"
+#include "Framework/MVVMRowHelper.h"
 #include "MVVMBlueprintFunctionReference.h"
 #include "MVVMBlueprintView.h"
 #include "MVVMBlueprintViewConversionFunction.h"
@@ -103,6 +104,8 @@ void SFunctionParameter::Construct(const FArguments& InArgs, UWidgetBlueprint* I
 				+ SWidgetSwitcher::Slot()
 				[
 					SNew(SFieldSelector, WidgetBlueprint.Get())
+					.OnDrop(this, &SFunctionParameter::HandleFieldSelectorDrop)
+					.OnDragOver(this, &SFunctionParameter::HandleFieldSelectorDragOver)
 					.OnGetLinkedValue(this, &SFunctionParameter::OnGetSelectedField)
 					.OnSelectionChanged(this, &SFunctionParameter::HandleFieldSelectionChanged)
 					.OnGetSelectionContext(this, &SFunctionParameter::GetSelectedSelectionContext)
@@ -248,5 +251,38 @@ FFieldSelectionContext SFunctionParameter::GetSelectedSelectionContext() const
 	return Result;
 }
 
+FReply SFunctionParameter::HandleFieldSelectorDrop(const FGeometry& MyGeometry, const FDragDropEvent& DragDropEvent)
+{
+	if (const UWidgetBlueprint* WidgetBlueprintPtr = WidgetBlueprint.Get())
+	{
+		const UMVVMEditorSubsystem* Subsystem = GEditor->GetEditorSubsystem<UMVVMEditorSubsystem>();
+		if (UMVVMBlueprintView* View = Subsystem->GetView(WidgetBlueprintPtr))
+		{
+			if (FMVVMBlueprintViewBinding* Binding = View->GetBinding(BindingId))
+			{
+				TOptional<FMVVMBlueprintPropertyPath> PropertyPath = BindingEntry::FRowHelper::DropFieldSelector(WidgetBlueprintPtr, DragDropEvent, bSourceToDestination);
+				if (PropertyPath.IsSet())
+				{
+					Subsystem->SetPathForConversionFunctionArgument(WidgetBlueprint.Get(), *Binding, ParameterId, PropertyPath.GetValue(), bSourceToDestination);
+				}
+				return FReply::Handled();
+			}
+		}
+	}
+
+	return FReply::Unhandled();
+}
+
+FReply SFunctionParameter::HandleFieldSelectorDragOver(const FGeometry& MyGeometry, const FDragDropEvent& DragDropEvent)
+{
+	UWidgetBlueprint* WidgetBlueprintPtr = WidgetBlueprint.Get();
+	if (WidgetBlueprintPtr == nullptr)
+	{
+		return FReply::Unhandled();
+	}
+
+	return BindingEntry::FRowHelper::DragOverFieldSelector(WidgetBlueprintPtr, DragDropEvent, bSourceToDestination);
+}
+
 }//namespace
 #undef LOCTEXT_NAMESPACE
diff --git a/Engine/Plugins/Runtime/ModelViewViewModel/Source/ModelViewViewModelEditor/Private/Widgets/SMVVMFunctionParameter.h b/Engine/Plugins/Runtime/ModelViewViewModel/Source/ModelViewViewModelEditor/Private/Widgets/SMVVMFunctionParameter.h
--- a/Engine/Plugins/Runtime/ModelViewViewModel/Source/ModelViewViewModelEditor/Private/Widgets/SMVVMFunctionParameter.h
+++ b/Engine/Plugins/Runtime/ModelViewViewModel/Source/ModelViewViewModelEditor/Private/Widgets/SMVVMFunctionParameter.h
@@ -44,6 +44,9 @@ private:
 	void HandleFieldSelectionChanged(FMVVMLinkedPinValue Value, SFieldSelectorMenu::ESelectionType SelectionType);
 	FFieldSelectionContext GetSelectedSelectionContext() const;
 
+	FReply HandleFieldSelectorDrop(const FGeometry& MyGeometry, const FDragDropEvent& DragDropEvent);
+	FReply HandleFieldSelectorDragOver(const FGeometry& MyGeometry, const FDragDropEvent& DragDropEvent);
+
 	int32 GetCurrentWidgetIndex() const;
 
 private:
diff --git a/Engine/Plugins/Runtime/ModelViewViewModel/Source/ModelViewViewModelEditor/Private/Widgets/SMVVMViewBindingPanel.cpp b/Engine/Plugins/Runtime/ModelViewViewModel/Source/ModelViewViewModelEditor/Private/Widgets/SMVVMViewBindingPanel.cpp
--- a/Engine/Plugins/Runtime/ModelViewViewModel/Source/ModelViewViewModelEditor/Private/Widgets/SMVVMViewBindingPanel.cpp
+++ b/Engine/Plugins/Runtime/ModelViewViewModel/Source/ModelViewViewModelEditor/Private/Widgets/SMVVMViewBindingPanel.cpp
@@ -17,6 +17,7 @@
 
 #include "Customizations/MVVMConversionPathCustomization.h"
 #include "Customizations/MVVMPropertyPathCustomization.h"
+#include "Details/WidgetPropertyDragDropOp.h"
 #include "Framework/MVVMBindingEditorHelper.h"
 #include "Hierarchy/HierarchyWidgetDragDropOp.h"
 #include "IStructureDetailsView.h"
@@ -39,6 +40,8 @@
 #include "Widgets/Layout/SSplitter.h"
 #include "Widgets/SMVVMViewModelPanel.h"
 #include "Widgets/SMVVMViewBindingListView.h"
+#include "Widgets/ViewModelFieldDragDropOp.h"
+
 #include "SPositiveActionButton.h"
 
 #define LOCTEXT_NAMESPACE "BindingPanel"
@@ -696,6 +699,30 @@ FReply SBindingsPanel::OnDragOver(const FGeometry& MyGeometry, const FDragDropEv
 			}
 			HierarchyDragDropOp->CurrentIconBrush = FAppStyle::GetBrush(TEXT("Graph.ConnectorFeedback.Error"));
 		}
+		else if (TSharedPtr<FViewModelFieldDragDropOp> ViewModelFieldDragDropOp = DragDropEvent.GetOperationAs<FViewModelFieldDragDropOp>())
+		{
+			UWidgetBlueprint* DragWidgetBlueprint = ViewModelFieldDragDropOp->WidgetBP.Get();
+			UWidgetBlueprint* CurrentWidgetBlueprint = WeakBlueprintEditor.Pin()->GetWidgetBlueprintObj();
+			if (CurrentWidgetBlueprint == DragWidgetBlueprint && ViewModelFieldDragDropOp->ViewModelId.IsValid())
+			{
+				ViewModelFieldDragDropOp->CurrentIconBrush = FAppStyle::GetBrush(TEXT("Graph.ConnectorFeedback.OK"));	
+				return FReply::Handled();
+			}
+
+			ViewModelFieldDragDropOp->CurrentIconBrush = FAppStyle::GetBrush(TEXT("Graph.ConnectorFeedback.Error"));
+		}
+		else if (TSharedPtr<FWidgetPropertyDragDropOp> WidgetPropertyDragDropOp = DragDropEvent.GetOperationAs<FWidgetPropertyDragDropOp>())
+		{
+			UWidgetBlueprint* DragWidgetBlueprint = WidgetPropertyDragDropOp->WidgetBP.Get();
+			UWidgetBlueprint* CurrentWidgetBlueprint = WeakBlueprintEditor.Pin()->GetWidgetBlueprintObj();
+			if (CurrentWidgetBlueprint == DragWidgetBlueprint)
+			{
+				WidgetPropertyDragDropOp->CurrentIconBrush = FAppStyle::GetBrush(TEXT("Graph.ConnectorFeedback.OK"));
+				return FReply::Handled();
+			}
+
+			WidgetPropertyDragDropOp->CurrentIconBrush = FAppStyle::GetBrush(TEXT("Graph.ConnectorFeedback.Error"));
+		}
 		else
 		{
 			TSharedPtr<FDecoratedDragDropOp> DecoratedDragDropOp = nullptr;
@@ -737,6 +764,76 @@ FReply SBindingsPanel::OnDrop(const FGeometry& MyGeometry, const FDragDropEvent&
 			}
 		}
 	}
+	else if (TSharedPtr<FViewModelFieldDragDropOp> ViewModelFieldDragDropOp = DragDropEvent.GetOperationAs<FViewModelFieldDragDropOp>())
+	{
+		UWidgetBlueprint* DragWidgetBlueprint = ViewModelFieldDragDropOp->WidgetBP.Get();
+		UWidgetBlueprint* CurrentWidgetBlueprint = WeakBlueprintEditor.Pin()->GetWidgetBlueprintObj();
+		if (CurrentWidgetBlueprint == DragWidgetBlueprint && ViewModelFieldDragDropOp->ViewModelId.IsValid())
+		{
+			FMVVMBlueprintPropertyPath PropertyPath;
+			for (const FFieldVariant& Field : ViewModelFieldDragDropOp->DraggedField)
+			{
+				PropertyPath.AppendPropertyPath(DragWidgetBlueprint, FMVVMConstFieldVariant(Field));
+			}
+
+			PropertyPath.SetViewModelId(ViewModelFieldDragDropOp->ViewModelId);
+
+			if (PropertyPath.IsValid())
+			{
+				UMVVMEditorSubsystem* Subsystem = GEditor->GetEditorSubsystem<UMVVMEditorSubsystem>();
+				FMVVMBlueprintViewBinding& ViewBinding = Subsystem->AddBinding(DragWidgetBlueprint);
+
+				Subsystem->SetSourcePathForBinding(DragWidgetBlueprint, ViewBinding, PropertyPath);
+
+				if (BindingsList)
+				{
+					BindingsList->RequestNavigateToBinding(ViewBinding.BindingId);
+				}
+				return FReply::Handled();
+			}
+		}
+	}
+	else if (TSharedPtr<FWidgetPropertyDragDropOp> WidgetPropertyDragDropOp = DragDropEvent.GetOperationAs<FWidgetPropertyDragDropOp>())
+	{
+		UWidgetBlueprint* DragWidgetBlueprint = WidgetPropertyDragDropOp->WidgetBP.Get();
+		UWidgetBlueprint* CurrentWidgetBlueprint = WeakBlueprintEditor.Pin()->GetWidgetBlueprintObj();
+		if (CurrentWidgetBlueprint == DragWidgetBlueprint)
+		{
+			FMVVMBlueprintPropertyPath PropertyPath;
+			for (const FFieldVariant& Field : WidgetPropertyDragDropOp->DraggedPropertyPath)
+			{
+				PropertyPath.AppendPropertyPath(DragWidgetBlueprint, FMVVMConstFieldVariant(Field));
+			}
+
+			if (UWidget* OwnerWidgetPtr = WidgetPropertyDragDropOp->OwnerWidget.Get())
+			{
+				if (WeakBlueprintEditor.Pin()->GetPreview() == OwnerWidgetPtr)
+				{
+					PropertyPath.SetSelfContext();
+				}
+				else
+				{
+					PropertyPath.SetWidgetName(OwnerWidgetPtr->GetFName());
+				}
+			}
+
+			if (PropertyPath.IsValid())
+			{
+				UMVVMEditorSubsystem* Subsystem = GEditor->GetEditorSubsystem<UMVVMEditorSubsystem>();
+				FMVVMBlueprintViewBinding& ViewBinding = Subsystem->AddBinding(DragWidgetBlueprint);
+
+				Subsystem->SetDestinationPathForBinding(DragWidgetBlueprint, ViewBinding, PropertyPath, false);
+
+				if (BindingsList)
+				{
+					BindingsList->RequestNavigateToBinding(ViewBinding.BindingId);
+				}
+				return FReply::Handled();
+			}
+		}
+	}
+
+
 	return FReply::Unhandled();
 }
 
