MVVM: fix function filter not resetting properly

diff --git a/Engine/Plugins/Runtime/ModelViewViewModel/Source/ModelViewViewModelBlueprint/Private/MVVMViewBlueprintCompiler.cpp b/Engine/Plugins/Runtime/ModelViewViewModel/Source/ModelViewViewModelBlueprint/Private/MVVMViewBlueprintCompiler.cpp
--- a/Engine/Plugins/Runtime/ModelViewViewModel/Source/ModelViewViewModelBlueprint/Private/MVVMViewBlueprintCompiler.cpp
+++ b/Engine/Plugins/Runtime/ModelViewViewModel/Source/ModelViewViewModelBlueprint/Private/MVVMViewBlueprintCompiler.cpp
@@ -717,6 +717,8 @@ void FMVVMViewBlueprintCompiler::CreateVariables(const FWidgetBlueprintCompilerC
 	// Public function permissions
 	{
 		FPathPermissionList& FunctionPermissions = GetMutableDefault<UBlueprintEditorSettings>()->GetFunctionPermissions();
+		FPathPermissionList& GeneratedFunctionPermissions = GetMutableDefault<UMVVMDeveloperProjectSettings>()->GetGeneratedFunctionPermissions();
+
 		FName ContextName;
 		{
 			TStringBuilder<512> ContextNameStr;
@@ -725,16 +727,20 @@ void FMVVMViewBlueprintCompiler::CreateVariables(const FWidgetBlueprintCompilerC
 		}
 
 		FunctionPermissions.UnregisterOwner(ContextName);
+		GeneratedFunctionPermissions.UnregisterOwner(ContextName);
 
 		const bool bHasFiltering = GetMutableDefault<UBlueprintEditorSettings>()->GetFunctionPermissions().HasFiltering();
-		if (bHasFiltering)
+		for (FName FunctionName : FunctionPermissionsToAdd)
 		{
-			for (FName FunctionName : FunctionPermissionsToAdd)
+			TStringBuilder<512> FunctionPath;
+			FunctionPath << ContextName;
+			FunctionPath << TEXT(':');
+			FunctionPath << FunctionName;
+
+			GeneratedFunctionPermissions.AddAllowListItem(ContextName, FunctionPath);
+
+			if (bHasFiltering)
 			{
-				TStringBuilder<512> FunctionPath;
-				FunctionPath << ContextName;
-				FunctionPath << TEXT(':');
-				FunctionPath << FunctionName;
 				FunctionPermissions.AddAllowListItem(ContextName, FunctionPath);
 			}
 		}
diff --git a/Engine/Plugins/Runtime/ModelViewViewModel/Source/ModelViewViewModelBlueprint/Public/MVVMDeveloperProjectSettings.h b/Engine/Plugins/Runtime/ModelViewViewModel/Source/ModelViewViewModelBlueprint/Public/MVVMDeveloperProjectSettings.h
--- a/Engine/Plugins/Runtime/ModelViewViewModel/Source/ModelViewViewModelBlueprint/Public/MVVMDeveloperProjectSettings.h
+++ b/Engine/Plugins/Runtime/ModelViewViewModel/Source/ModelViewViewModelBlueprint/Public/MVVMDeveloperProjectSettings.h
@@ -3,6 +3,7 @@
 #pragma once
 
 #include "Engine/DeveloperSettings.h"
+#include "Misc/NamePermissionList.h"
 
 #include "MVVMDeveloperProjectSettings.generated.h"
 
@@ -103,6 +104,9 @@ public:
 	bool IsExtensionSupportedForPanelClass(TSubclassOf<UPanelWidget> ClassToSupport) const;
 	bool IsExtensionSupportedForListViewBaseClass(TSubclassOf<UListViewBase> ClassToSupport) const;
 
+	/** Get generated functions permission list. */
+	FPathPermissionList& GetGeneratedFunctionPermissions() { return GeneratedFunctionPermissions; }
+
 private:
 	/** Permission list for filtering which properties are visible in UI. */
 	UPROPERTY(EditAnywhere, config, Category = "UX", meta=(AllowAbstract = true))
@@ -204,4 +208,8 @@ public:
 	/** Resolver class to use as the default value when selecting resolver creation mode */
 	UPROPERTY(EditAnywhere, config, Category = "Defaults")
 	TSoftClassPtr<UMVVMViewModelContextResolver> DefaultResolverValue;
+
+private:
+	/** MVVM Generated function permission list */
+	FPathPermissionList GeneratedFunctionPermissions;
 };
