MVVM: guard against null/deleted function in property path of conversion function argument

diff --git a/Engine/Plugins/Runtime/ModelViewViewModel/Source/ModelViewViewModelBlueprint/Private/Bindings/MVVMConversionFunctionHelper.cpp b/Engine/Plugins/Runtime/ModelViewViewModel/Source/ModelViewViewModelBlueprint/Private/Bindings/MVVMConversionFunctionHelper.cpp
--- a/Engine/Plugins/Runtime/ModelViewViewModel/Source/ModelViewViewModelBlueprint/Private/Bindings/MVVMConversionFunctionHelper.cpp
+++ b/Engine/Plugins/Runtime/ModelViewViewModel/Source/ModelViewViewModelBlueprint/Private/Bindings/MVVMConversionFunctionHelper.cpp
@@ -533,15 +533,21 @@ namespace NamedNodes
 					const FString& MetaData = PreviousStructProperty->Struct->GetMetaData(FBlueprintMetadata::MD_NativeBreakFunction);
 					if (MetaData.Len() > 0)
 					{
-						const UFunction* Function = FindObject<UFunction>(nullptr, *MetaData, true);
-						ensure(Function);
-						FGraphNodeCreator<UK2Node_CallFunction> MakeStructCreator(*FunctionGraph);
-						UK2Node_CallFunction* FunctionNode = MakeStructCreator.CreateNode(false);
-						FunctionNode->SetFromFunction(Function);
-						MakeStructCreator.Finalize();
-
-						NewNode = FunctionNode;
-						NewPreviousProperty = UE::MVVM::BindingHelper::GetReturnProperty(Function);
+						if (const UFunction* Function = FindObject<UFunction>(nullptr, *MetaData, true))
+						{
+							FGraphNodeCreator<UK2Node_CallFunction> MakeStructCreator(*FunctionGraph);
+							UK2Node_CallFunction* FunctionNode = MakeStructCreator.CreateNode(false);
+							FunctionNode->SetFromFunction(Function);
+							MakeStructCreator.Finalize();
+
+							NewNode = FunctionNode;
+							NewPreviousProperty = UE::MVVM::BindingHelper::GetReturnProperty(Function);
+						}
+						else
+						{
+							ensureMsgf(false, TEXT("A function in path couldn't be resolved. Blueprint: %s Function: %s"), *Blueprint->GetName(), *MetaData);
+							return MakeError();
+						}
 					}
 					else
 					{
