MVVM: move the material functions to the blueprint module so that it's loaded with UncookedOnly instead of Editor.

diff --git a/Engine/Plugins/Runtime/ModelViewViewModel/Config/BaseModelViewViewModel.ini b/Engine/Plugins/Runtime/ModelViewViewModel/Config/BaseModelViewViewModel.ini
--- a/Engine/Plugins/Runtime/ModelViewViewModel/Config/BaseModelViewViewModel.ini
+++ b/Engine/Plugins/Runtime/ModelViewViewModel/Config/BaseModelViewViewModel.ini
@@ -1,6 +1,14 @@
 [CoreRedirects]
 +ClassRedirects=(OldName="/Script/ModelViewViewModelBlueprint.MVVMViewBlueprintListViewBaseExtension", NewName="/Script/ModelViewViewModelBlueprint.MVVMBlueprintViewExtension_ListViewBase")
 +ClassRedirects=(OldName="/Script/ModelViewViewModelBlueprint.MVVMViewListViewBaseExtension", NewName="/Script/ModelViewViewModelBlueprint.MVVMViewListViewBaseClassExtension")
++ClassRedirects=(OldName="/Script/ModelViewViewModelEditor.MVVMK2Node_MakeBrushFromSoftMaterial", NewName="/Script/ModelViewViewModelBlueprint.MVVMK2Node_MakeBrushFromSoftMaterial")
++ClassRedirects=(OldName="/Script/ModelViewViewModelEditor.MVVMK2Node_MakeBrushFromSoftResource", NewName="/Script/ModelViewViewModelBlueprint.MVVMK2Node_MakeBrushFromSoftResource")
++ClassRedirects=(OldName="/Script/ModelViewViewModelEditor.MVVMK2Node_MakeBrushFromSoftTexture", NewName="/Script/ModelViewViewModelBlueprint.MVVMK2Node_MakeBrushFromSoftTexture")
++ClassRedirects=(OldName="/Script/ModelViewViewModelEditor.MVVMK2Node_LoadSoftResource", NewName="/Script/ModelViewViewModelBlueprint.MVVMK2Node_LoadSoftResource")
++ClassRedirects=(OldName="/Script/ModelViewViewModelEditor.MVVMK2Node_LoadSoftTexture", NewName="/Script/ModelViewViewModelBlueprint.MVVMK2Node_LoadSoftTexture")
++ClassRedirects=(OldName="/Script/ModelViewViewModelEditor.MVVMK2Node_LoadSoftMaterial", NewName="/Script/ModelViewViewModelBlueprint.MVVMK2Node_LoadSoftMaterial")
++ClassRedirects=(OldName="/Script/ModelViewViewModelEditor.MVVMK2Node_LoadSoftInputAction", NewName="/Script/ModelViewViewModelBlueprint.MVVMK2Node_LoadSoftInputAction")
++ClassRedirects=(OldName="/Script/ModelViewViewModelEditor.MVVMK2Node_SetSoftTextureParameter", NewName="/Script/ModelViewViewModelBlueprint.MVVMK2Node_SetSoftTextureParameter")
 +StructRedirects=(OldName="MVVMViewModelPropertyPath",NewName="/Script/ModelViewViewModelBlueprint.MVVMBlueprintPropertyPath")
 +StructRedirects=(OldName="MVVMWidgetPropertyPath",NewName="/Script/ModelViewViewModelBlueprint.MVVMBlueprintPropertyPath")
 +FunctionRedirects=(OldName="/Script/ModelViewViewModel.MVVMSubsystem.GetViewFromUserWidget", NewName="/Script/ModelViewViewModel.MVVMSubsystem.K2_GetViewFromUserWidget")
diff --git a/Engine/Plugins/Runtime/ModelViewViewModel/Source/ModelViewViewModelBlueprint/ModelViewViewModelBlueprint.Build.cs b/Engine/Plugins/Runtime/ModelViewViewModel/Source/ModelViewViewModelBlueprint/ModelViewViewModelBlueprint.Build.cs
--- a/Engine/Plugins/Runtime/ModelViewViewModel/Source/ModelViewViewModelBlueprint/ModelViewViewModelBlueprint.Build.cs
+++ b/Engine/Plugins/Runtime/ModelViewViewModel/Source/ModelViewViewModelBlueprint/ModelViewViewModelBlueprint.Build.cs
@@ -23,6 +23,7 @@ public class ModelViewViewModelBlueprint : ModuleRules
 			new string[]
 			{
 				"BlueprintGraph",
+				"EnhancedInput",
 				"Kismet",
 				"KismetCompiler",
 				"PropertyEditor",
diff --git a/Engine/Plugins/Runtime/ModelViewViewModel/Source/ModelViewViewModelBlueprint/Private/Node/MVVMK2Node_LoadSoftResource.cpp b/Engine/Plugins/Runtime/ModelViewViewModel/Source/ModelViewViewModelBlueprint/Private/Node/MVVMK2Node_LoadSoftResource.cpp
new file mode 100644
--- /dev/null
+++ b/Engine/Plugins/Runtime/ModelViewViewModel/Source/ModelViewViewModelBlueprint/Private/Node/MVVMK2Node_LoadSoftResource.cpp
@@ -0,0 +1,187 @@
+// Copyright Epic Games, Inc. All Rights Reserved.
+
+#include "Node/MVVMK2Node_LoadSoftResource.h"
+
+#include "BlueprintCompiledStatement.h"
+#include "BlueprintNodeSpawner.h"
+#include "Blueprint/WidgetBlueprintLibrary.h"
+#include "EdGraph/EdGraphPin.h"
+#include "EdGraphSchema_K2.h"
+#include "Kismet2/CompilerResultsLog.h"
+#include "KismetCompiledFunctionContext.h"
+#include "KismetCompiler.h"
+#include "MVVMSubsystem.h"
+#include "Subsystems/SubsystemBlueprintLibrary.h"
+#include "View/MVVMView.h"
+
+#include "K2Node_CallFunction.h"
+#include "K2Node_ClassDynamicCast.h"
+#include "K2Node_IfThenElse.h"
+#include "K2Node_Self.h"
+
+#include "Engine/Texture2D.h"
+#include "InputAction.h"
+#include "Materials/MaterialInterface.h"
+
+
+#define LOCTEXT_NAMESPACE "MVVMK2Node_LoadSoftResource"
+
+//////////////////////////////////////////////////////////////////////////
+
+void UMVVMK2Node_LoadSoftResource::AllocateDefaultPins()
+{
+	const UEdGraphSchema_K2* K2Schema = GetDefault<UEdGraphSchema_K2>();
+
+	CreatePin(EGPD_Input, UEdGraphSchema_K2::PC_Exec, UEdGraphSchema_K2::PN_Execute);
+
+	// The immediate continue pin
+	CreatePin(EGPD_Output, UEdGraphSchema_K2::PC_Exec, UEdGraphSchema_K2::PN_Then);
+
+	// The delayed completed pin, this used to be called Then
+	CreatePin(EGPD_Output, UEdGraphSchema_K2::PC_Exec, UEdGraphSchema_K2::PN_Completed);
+
+	CreatePin(EGPD_Input, UEdGraphSchema_K2::PC_SoftObject, GetInputResourceClass(), GetInputPinName());
+	CreatePin(EGPD_Output, UEdGraphSchema_K2::PC_Object, GetInputResourceClass(), GetOutputPinName());
+
+	UK2Node::AllocateDefaultPins();
+}
+
+void UMVVMK2Node_LoadSoftResource::ExpandNode(FKismetCompilerContext& CompilerContext, UEdGraph* SourceGraph)
+{
+	UK2Node::ExpandNode(CompilerContext, SourceGraph);
+	
+	const UEdGraphSchema_K2* Schema = CompilerContext.GetSchema();
+
+	UK2Node_LoadAsset* CallLoadAssetNode = CompilerContext.SpawnIntermediateNode<UK2Node_LoadAsset>(this, SourceGraph);
+	{
+		CallLoadAssetNode->AllocateDefaultPins();
+	}
+	UEdGraphPin* LoadAssetInput = CallLoadAssetNode->FindPinChecked(CallLoadAssetNode->GetInputPinName());
+	UEdGraphPin* LoadAssetOutput = CallLoadAssetNode->FindPinChecked(CallLoadAssetNode->GetOutputPinName());
+
+	UK2Node_DynamicCast* CastNode = CompilerContext.SpawnIntermediateNode<UK2Node_DynamicCast>(this, SourceGraph);
+	{
+		CastNode->TargetType = GetInputResourceClass();
+		CastNode->AllocateDefaultPins();
+	}
+	UEdGraphPin* CastInput = CastNode->GetCastSourcePin();
+	UEdGraphPin* CastOutput = CastNode->GetCastResultPin();
+
+	// move this.exec to CallLoadAssetNode.exec
+	CompilerContext.MovePinLinksToIntermediate(*GetExecPin(), *CallLoadAssetNode->GetExecPin());
+	// move this.then to CallLoadAssetNode.then
+	CompilerContext.MovePinLinksToIntermediate(*GetThenPin(), *CallLoadAssetNode->GetThenPin());
+	// CallLoadAssetNode.completed to CastNode.exec
+	ensure(Schema->TryCreateConnection(CallLoadAssetNode->FindPinChecked(UEdGraphSchema_K2::PN_Completed), CastNode->GetExecPin()));
+	// move this.completed to CastNode.then
+	CompilerContext.MovePinLinksToIntermediate(*GetCompletedPin(), *CastNode->GetThenPin());
+
+	// move this.resource to CallLoadAssetNode.arg
+	CompilerContext.MovePinLinksToIntermediate(*FindPinChecked(GetInputPinName()), *LoadAssetInput);
+	// CallLoadAssetNode.result to CastNode.Input
+	ensure(Schema->TryCreateConnection(LoadAssetOutput, CastInput));
+	// move this.result to CastNode.Output
+	CompilerContext.MovePinLinksToIntermediate(*FindPinChecked(GetOutputPinName()), *CastOutput);
+
+	BreakAllNodeLinks();
+}
+
+UEdGraphPin* UMVVMK2Node_LoadSoftResource::GetThenPin() const
+{
+	return FindPinChecked(UEdGraphSchema_K2::PN_Then);
+}
+
+UEdGraphPin* UMVVMK2Node_LoadSoftResource::GetCompletedPin() const
+{
+	return FindPinChecked(UEdGraphSchema_K2::PN_Completed);
+}
+
+const FName& UMVVMK2Node_LoadSoftResource::GetOutputPinName() const
+{
+	static const FName OutputPinName("Result");
+	return OutputPinName;
+}
+
+UClass* UMVVMK2Node_LoadSoftResource::GetInputResourceClass() const
+{
+	return nullptr;
+}
+
+
+//////////////////////////////////////////////////////////////////////////
+
+
+FText UMVVMK2Node_LoadSoftTexture::GetTooltipText() const
+{
+	return FText(LOCTEXT("UMVVMK2Node_LoadSoftTextureGetTooltipText", "Asynchronously loads a Soft Texture Reference and returns that texture on successful load"));
+}
+
+FText UMVVMK2Node_LoadSoftTexture::GetNodeTitle(ENodeTitleType::Type TitleType) const
+{
+	return FText(LOCTEXT("UMVVMK2Node_LoadSoftTextureGetNodeTitle", "Load From Soft Texture"));
+}
+
+const FName& UMVVMK2Node_LoadSoftTexture::GetInputPinName() const
+{
+	// Note: Our pin input should match the input of the specified MakeBrush node
+	static const FName InputPinName("Texture");
+	return InputPinName;
+}
+
+UClass* UMVVMK2Node_LoadSoftTexture::GetInputResourceClass() const
+{
+	return UTexture2D::StaticClass();
+}
+
+
+//////////////////////////////////////////////////////////////////////////
+
+
+FText UMVVMK2Node_LoadSoftMaterial::GetTooltipText() const
+{
+	return FText(LOCTEXT("UMVVMK2Node_LoadSoftMaterialGetTooltipText", "Asynchronously loads a Soft Material Reference and returns that material on successful load"));
+}
+
+FText UMVVMK2Node_LoadSoftMaterial::GetNodeTitle(ENodeTitleType::Type TitleType) const
+{
+	return FText(LOCTEXT("UMVVMK2Node_LoadSoftMaterialGetNodeTitle", "Load From Soft Material"));
+}
+
+const FName& UMVVMK2Node_LoadSoftMaterial::GetInputPinName() const
+{
+	// Note: Our pin input should match the input of the specified MakeBrush node
+	static const FName InputPinName("Material");
+	return InputPinName;
+}
+
+UClass* UMVVMK2Node_LoadSoftMaterial::GetInputResourceClass() const
+{
+	return UMaterialInterface::StaticClass();
+}
+
+
+//////////////////////////////////////////////////////////////////////////
+
+
+FText UMVVMK2Node_LoadSoftInputAction::GetTooltipText() const
+{
+	return FText(LOCTEXT("UMVVMK2Node_LoadSoftInputActionGetTooltipText", "Asynchronously loads a Soft Input Action Reference and returns that Input Action on successful load"));
+}
+
+FText UMVVMK2Node_LoadSoftInputAction::GetNodeTitle(ENodeTitleType::Type TitleType) const
+{
+	return FText(LOCTEXT("UMVVMK2Node_LoadSoftInputActionGetNodeTitle", "Load From Soft Input Action"));
+}
+
+const FName& UMVVMK2Node_LoadSoftInputAction::GetInputPinName() const
+{
+	static const FName InputPinName("InputAction");
+	return InputPinName;
+}
+
+UClass* UMVVMK2Node_LoadSoftInputAction::GetInputResourceClass() const
+{
+	return UInputAction::StaticClass();
+}
+
+#undef LOCTEXT_NAMESPACE
\ No newline at end of file
diff --git a/Engine/Plugins/Runtime/ModelViewViewModel/Source/ModelViewViewModelBlueprint/Private/Node/MVVMK2Node_LoadSoftResource.h b/Engine/Plugins/Runtime/ModelViewViewModel/Source/ModelViewViewModelBlueprint/Private/Node/MVVMK2Node_LoadSoftResource.h
new file mode 100644
--- /dev/null
+++ b/Engine/Plugins/Runtime/ModelViewViewModel/Source/ModelViewViewModelBlueprint/Private/Node/MVVMK2Node_LoadSoftResource.h
@@ -0,0 +1,119 @@
+// Copyright Epic Games, Inc. All Rights Reserved.
+
+#pragma once
+
+#include "CoreMinimal.h"
+#include "K2Node.h"
+#include "K2Node_LoadAsset.h"
+#include "View/MVVMViewTypes.h"
+
+#include "MVVMK2Node_LoadSoftResource.generated.h"
+
+class UEdGraphPin;
+class FKismetCompilerContext;
+
+
+/**
+ * Utility async node to load a type-specific a soft resource (Texture, Material, etc). 
+ * Needed since we current do not support automatic dynamic casting in MVVM
+ */
+UCLASS(MinimalAPI, Abstract, Hidden, HideDropDown)
+class UMVVMK2Node_LoadSoftResource : public UK2Node_LoadAsset
+{
+	GENERATED_BODY()
+
+public:
+	//~ Begin UEdGraphNode Interface.
+	virtual void AllocateDefaultPins() override;
+	//~ End UEdGraphNode Interface.
+	
+	//~ Begin K2Node Interface.
+	virtual void ExpandNode(FKismetCompilerContext& CompilerContext, UEdGraph* SourceGraph) override;
+	//~ End K2Node Interface.
+
+	/** Get the then output pin */
+	UEdGraphPin* GetThenPin() const;
+
+	/** Get the then output pin */
+	UEdGraphPin* GetCompletedPin() const;
+
+protected:
+
+	//~ Begin UK2Node_LoadAsset Interface.
+	virtual const FName& GetOutputPinName() const;
+	//~ End UK2Node_LoadAsset Interface.
+
+	/** Get the type of the input arg pin */
+	virtual UClass* GetInputResourceClass() const;
+};
+
+
+//////////////////////////////////////////////////////////////////////////
+
+
+UCLASS(MinimalAPI)
+class UMVVMK2Node_LoadSoftTexture : public UMVVMK2Node_LoadSoftResource
+{
+	GENERATED_BODY()
+
+public:
+	//~ Begin UEdGraphNode Interface.
+	virtual FText GetTooltipText() const override;
+	virtual FText GetNodeTitle(ENodeTitleType::Type TitleType) const override;
+	//~ End UEdGraphNode Interface.
+
+protected:
+	//~ Begin UK2Node_LoadAsset Interface.
+	virtual const FName& GetInputPinName() const;
+	//~ End UK2Node_LoadAsset Interface.
+
+	//~ Begin UMVVMK2Node_LoadSoftResource Interface.
+	virtual UClass* GetInputResourceClass() const override;
+	//~ End UMVVMK2Node_LoadSoftResource Interface.
+};
+
+
+//////////////////////////////////////////////////////////////////////////
+
+
+UCLASS(MinimalAPI)
+class UMVVMK2Node_LoadSoftMaterial: public UMVVMK2Node_LoadSoftResource
+{
+	GENERATED_BODY()
+
+public:
+	//~ Begin UEdGraphNode Interface.
+	virtual FText GetTooltipText() const override;
+	virtual FText GetNodeTitle(ENodeTitleType::Type TitleType) const override;
+	//~ End UEdGraphNode Interface.
+
+protected:
+	//~ Begin UK2Node_LoadAsset Interface.
+	virtual const FName& GetInputPinName() const;
+	//~ End UK2Node_LoadAsset Interface.
+
+	//~ Begin UMVVMK2Node_LoadSoftResource Interface.
+	virtual UClass* GetInputResourceClass() const override;
+	//~ End UMVVMK2Node_LoadSoftResource Interface.
+};
+
+UCLASS(MinimalAPI)
+class UMVVMK2Node_LoadSoftInputAction : public UMVVMK2Node_LoadSoftResource
+{
+	GENERATED_BODY()
+
+public:
+	//~ Begin UEdGraphNode Interface.
+	virtual FText GetTooltipText() const override;
+	virtual FText GetNodeTitle(ENodeTitleType::Type TitleType) const override;
+	//~ End UEdGraphNode Interface.
+
+protected:
+	//~ Begin UK2Node_LoadAsset Interface.
+	virtual const FName& GetInputPinName() const;
+	//~ End UK2Node_LoadAsset Interface.
+
+	//~ Begin UMVVMK2Node_LoadSoftResource Interface.
+	virtual UClass* GetInputResourceClass() const override;
+	//~ End UMVVMK2Node_LoadSoftResource Interface.
+};
diff --git a/Engine/Plugins/Runtime/ModelViewViewModel/Source/ModelViewViewModelBlueprint/Private/Node/MVVMK2Node_MakeBrushFromSoftResource.cpp b/Engine/Plugins/Runtime/ModelViewViewModel/Source/ModelViewViewModelBlueprint/Private/Node/MVVMK2Node_MakeBrushFromSoftResource.cpp
new file mode 100644
--- /dev/null
+++ b/Engine/Plugins/Runtime/ModelViewViewModel/Source/ModelViewViewModelBlueprint/Private/Node/MVVMK2Node_MakeBrushFromSoftResource.cpp
@@ -0,0 +1,204 @@
+// Copyright Epic Games, Inc. All Rights Reserved.
+
+#include "Node/MVVMK2Node_MakeBrushFromSoftResource.h"
+
+#include "BlueprintCompiledStatement.h"
+#include "BlueprintNodeSpawner.h"
+#include "Blueprint/WidgetBlueprintLibrary.h"
+#include "EdGraph/EdGraphPin.h"
+#include "EdGraphSchema_K2.h"
+#include "Kismet2/CompilerResultsLog.h"
+#include "KismetCompiledFunctionContext.h"
+#include "KismetCompiler.h"
+#include "MVVMSubsystem.h"
+#include "Subsystems/SubsystemBlueprintLibrary.h"
+#include "View/MVVMView.h"
+
+#include "K2Node_CallFunction.h"
+#include "K2Node_ClassDynamicCast.h"
+#include "K2Node_IfThenElse.h"
+#include "K2Node_Self.h"
+
+#include "Engine/Texture2D.h"
+#include "Materials/MaterialInterface.h"
+
+
+#define LOCTEXT_NAMESPACE "MVVMK2Node_MakeBrushFromSoftResource"
+
+namespace UE::MVVM::Private
+{
+	namespace PinNames
+	{
+		static const FName Width = "Width";
+		static const FName Height = "Height";
+	}
+}
+
+//////////////////////////////////////////////////////////////////////////
+
+void UMVVMK2Node_MakeBrushFromSoftResource::AllocateDefaultPins()
+{
+	const UEdGraphSchema_K2* K2Schema = GetDefault<UEdGraphSchema_K2>();
+
+	CreatePin(EGPD_Input, UEdGraphSchema_K2::PC_Exec, UEdGraphSchema_K2::PN_Execute);
+
+	// The immediate continue pin
+	CreatePin(EGPD_Output, UEdGraphSchema_K2::PC_Exec, UEdGraphSchema_K2::PN_Then);
+
+	// The delayed completed pin, this used to be called Then
+	CreatePin(EGPD_Output, UEdGraphSchema_K2::PC_Exec, UEdGraphSchema_K2::PN_Completed);
+
+	CreatePin(EGPD_Input, UEdGraphSchema_K2::PC_SoftObject, GetInputResourceClass(), GetInputPinName());
+	CreatePin(EGPD_Input, UEdGraphSchema_K2::PC_Int, nullptr, UE::MVVM::Private::PinNames::Width);
+	CreatePin(EGPD_Input, UEdGraphSchema_K2::PC_Int, nullptr, UE::MVVM::Private::PinNames::Height);
+	CreatePin(EGPD_Output, UEdGraphSchema_K2::PC_Struct, FSlateBrush::StaticStruct(), GetOutputPinName());
+
+	UK2Node::AllocateDefaultPins();
+}
+
+void UMVVMK2Node_MakeBrushFromSoftResource::ExpandNode(FKismetCompilerContext& CompilerContext, UEdGraph* SourceGraph)
+{
+	UK2Node::ExpandNode(CompilerContext, SourceGraph);
+	
+	const UEdGraphSchema_K2* Schema = CompilerContext.GetSchema();
+
+	UK2Node_LoadAsset* CallLoadAssetNode = CompilerContext.SpawnIntermediateNode<UK2Node_LoadAsset>(this, SourceGraph);
+	{
+		CallLoadAssetNode->AllocateDefaultPins();
+	}
+	UEdGraphPin* LoadAssetInput = CallLoadAssetNode->FindPinChecked(CallLoadAssetNode->GetInputPinName());
+	UEdGraphPin* LoadAssetOutput = CallLoadAssetNode->FindPinChecked(CallLoadAssetNode->GetOutputPinName());
+
+	UK2Node_CallFunction* CallMakeBrushFrom = CompilerContext.SpawnIntermediateNode<UK2Node_CallFunction>(this, SourceGraph);
+	{
+		CallMakeBrushFrom->FunctionReference.SetExternalMember(GetMakeBrushFunctionName(), UWidgetBlueprintLibrary::StaticClass());
+		CallMakeBrushFrom->AllocateDefaultPins();
+	}
+	// Note: Our pin input should match the input of the specified MakeBrush node
+	UEdGraphPin* CallMakeBrushFromResourcePin = CallMakeBrushFrom->FindPinChecked(GetInputPinName());
+	UEdGraphPin* CallMakeBrushFromWidthPin = CallMakeBrushFrom->FindPinChecked(UE::MVVM::Private::PinNames::Width);
+	UEdGraphPin* CallMakeBrushFromHeightPin = CallMakeBrushFrom->FindPinChecked(UE::MVVM::Private::PinNames::Height);
+	UEdGraphPin* CallMakeBrushFromResultPin = CallMakeBrushFrom->GetReturnValuePin();
+
+	UK2Node_DynamicCast* CastNode = CompilerContext.SpawnIntermediateNode<UK2Node_DynamicCast>(this, SourceGraph);
+	{
+		CastNode->TargetType = GetInputResourceClass();
+		CastNode->AllocateDefaultPins();
+	}
+	UEdGraphPin* CastInput = CastNode->GetCastSourcePin();
+	UEdGraphPin* CastOutput = CastNode->GetCastResultPin();
+
+	// move this.exec to CallLoadAssetNode.exec
+	CompilerContext.MovePinLinksToIntermediate(*GetExecPin(), *CallLoadAssetNode->GetExecPin());
+	// move this.then to CallLoadAssetNode.then
+	CompilerContext.MovePinLinksToIntermediate(*GetThenPin(), *CallLoadAssetNode->GetThenPin());
+	// CallLoadAssetNode.completed to CastNode.exec
+	ensure(Schema->TryCreateConnection(CallLoadAssetNode->FindPinChecked(UEdGraphSchema_K2::PN_Completed), CastNode->GetExecPin()));
+	// move this.completed to CastNode.then
+	CompilerContext.MovePinLinksToIntermediate(*GetCompletedPin(), *CastNode->GetThenPin());
+
+	// move this.resource to CallLoadAssetNode.arg
+	CompilerContext.MovePinLinksToIntermediate(*FindPinChecked(GetInputPinName()), *LoadAssetInput);
+	// move this.width to CallMakeBrushFrom.width
+	CompilerContext.MovePinLinksToIntermediate(*FindPinChecked(UE::MVVM::Private::PinNames::Width), *CallMakeBrushFromWidthPin);
+	// move this.height to CallMakeBrushFrom.height
+	CompilerContext.MovePinLinksToIntermediate(*FindPinChecked(UE::MVVM::Private::PinNames::Height), *CallMakeBrushFromHeightPin);
+	// CallLoadAssetNode.result to CastNode.Input
+	ensure(Schema->TryCreateConnection(LoadAssetOutput, CastInput));
+	// CastNode.Output to CallMakeBrushFrom.Resource
+	ensure(Schema->TryCreateConnection(CastOutput, CallMakeBrushFromResourcePin));
+	// move this.result to CallLoadAssetNode.result
+	CompilerContext.MovePinLinksToIntermediate(*FindPinChecked(GetOutputPinName()), *CallMakeBrushFromResultPin);
+
+	BreakAllNodeLinks();
+}
+
+UEdGraphPin* UMVVMK2Node_MakeBrushFromSoftResource::GetThenPin() const
+{
+	return FindPinChecked(UEdGraphSchema_K2::PN_Then);
+}
+
+UEdGraphPin* UMVVMK2Node_MakeBrushFromSoftResource::GetCompletedPin() const
+{
+	return FindPinChecked(UEdGraphSchema_K2::PN_Completed);
+}
+
+const FName& UMVVMK2Node_MakeBrushFromSoftResource::GetOutputPinName() const
+{
+	static const FName OutputPinName("SlateBrush");
+	return OutputPinName;
+}
+
+UClass* UMVVMK2Node_MakeBrushFromSoftResource::GetInputResourceClass() const
+{
+	return nullptr;
+}
+
+FName UMVVMK2Node_MakeBrushFromSoftResource::GetMakeBrushFunctionName() const
+{
+	return NAME_None;
+}
+
+
+//////////////////////////////////////////////////////////////////////////
+
+
+FText UMVVMK2Node_MakeBrushFromSoftTexture::GetTooltipText() const
+{
+	return FText(LOCTEXT("UMVVMK2Node_MakeBrushFromSoftTextureGetTooltipText", "Asynchronously loads a Soft Texture Reference and returns a slate brush using that texture on successful load"));
+}
+
+FText UMVVMK2Node_MakeBrushFromSoftTexture::GetNodeTitle(ENodeTitleType::Type TitleType) const
+{
+	return FText(LOCTEXT("UMVVMK2Node_MakeBrushFromSoftTextureGetNodeTitle", "Make Brush From Soft Texture"));
+}
+
+const FName& UMVVMK2Node_MakeBrushFromSoftTexture::GetInputPinName() const
+{
+	// Note: Our pin input should match the input of the specified MakeBrush node
+	static const FName InputPinName("Texture");
+	return InputPinName;
+}
+
+UClass* UMVVMK2Node_MakeBrushFromSoftTexture::GetInputResourceClass() const
+{
+	return UTexture2D::StaticClass();
+}
+
+FName UMVVMK2Node_MakeBrushFromSoftTexture::GetMakeBrushFunctionName() const
+{
+	return GET_FUNCTION_NAME_CHECKED(UWidgetBlueprintLibrary, MakeBrushFromTexture);
+}
+
+
+//////////////////////////////////////////////////////////////////////////
+
+
+FText UMVVMK2Node_MakeBrushFromSoftMaterial::GetTooltipText() const
+{
+	return FText(LOCTEXT("UMVVMK2Node_MakeBrushFromSoftMaterialGetTooltipText", "Asynchronously loads a Soft Material Reference and returns a slate brush using that material on successful load"));
+}
+
+FText UMVVMK2Node_MakeBrushFromSoftMaterial::GetNodeTitle(ENodeTitleType::Type TitleType) const
+{
+	return FText(LOCTEXT("UMVVMK2Node_MakeBrushFromSoftMaterialGetNodeTitle", "Make Brush From Soft Material"));
+}
+
+const FName& UMVVMK2Node_MakeBrushFromSoftMaterial::GetInputPinName() const
+{
+	// Note: Our pin input should match the input of the specified MakeBrush node
+	static const FName InputPinName("Material");
+	return InputPinName;
+}
+
+UClass* UMVVMK2Node_MakeBrushFromSoftMaterial::GetInputResourceClass() const
+{
+	return UMaterialInterface::StaticClass();
+}
+
+FName UMVVMK2Node_MakeBrushFromSoftMaterial::GetMakeBrushFunctionName() const
+{
+	return GET_FUNCTION_NAME_CHECKED(UWidgetBlueprintLibrary, MakeBrushFromMaterial);
+}
+
+#undef LOCTEXT_NAMESPACE
diff --git a/Engine/Plugins/Runtime/ModelViewViewModel/Source/ModelViewViewModelBlueprint/Private/Node/MVVMK2Node_MakeBrushFromSoftResource.h b/Engine/Plugins/Runtime/ModelViewViewModel/Source/ModelViewViewModelBlueprint/Private/Node/MVVMK2Node_MakeBrushFromSoftResource.h
new file mode 100644
--- /dev/null
+++ b/Engine/Plugins/Runtime/ModelViewViewModel/Source/ModelViewViewModelBlueprint/Private/Node/MVVMK2Node_MakeBrushFromSoftResource.h
@@ -0,0 +1,103 @@
+// Copyright Epic Games, Inc. All Rights Reserved.
+
+#pragma once
+
+#include "CoreMinimal.h"
+#include "K2Node.h"
+#include "K2Node_LoadAsset.h"
+#include "View/MVVMViewTypes.h"
+
+#include "MVVMK2Node_MakeBrushFromSoftResource.generated.h"
+
+class UEdGraphPin;
+class FKismetCompilerContext;
+
+
+/**
+ * Utility async node to create a slate brush from a soft resource (Texture, Material, etc). 
+ * Needed since we cannot nest conversion functions (Ex: LoadAsset->MakeBrushFromTexture).
+ */
+UCLASS(MinimalAPI, Abstract, Hidden, HideDropDown)
+class UMVVMK2Node_MakeBrushFromSoftResource : public UK2Node_LoadAsset
+{
+	GENERATED_BODY()
+
+public:
+	//~ Begin UEdGraphNode Interface.
+	virtual void AllocateDefaultPins() override;
+	//~ End UEdGraphNode Interface.
+	
+	//~ Begin K2Node Interface.
+	virtual void ExpandNode(FKismetCompilerContext& CompilerContext, UEdGraph* SourceGraph) override;
+	//~ End K2Node Interface.
+
+	/** Get the then output pin */
+	UEdGraphPin* GetThenPin() const;
+
+	/** Get the completed output pin */
+	UEdGraphPin* GetCompletedPin() const;
+
+protected:
+
+	//~ Begin UK2Node_LoadAsset Interface.
+	virtual const FName& GetOutputPinName() const;
+	//~ End UK2Node_LoadAsset Interface.
+
+	/** Get the type of the input arg pin */
+	virtual UClass* GetInputResourceClass() const;
+
+	/** Get the name of the function*/
+	virtual FName GetMakeBrushFunctionName() const;
+};
+
+
+//////////////////////////////////////////////////////////////////////////
+
+
+UCLASS(MinimalAPI)
+class UMVVMK2Node_MakeBrushFromSoftTexture : public UMVVMK2Node_MakeBrushFromSoftResource
+{
+	GENERATED_BODY()
+
+public:
+	//~ Begin UEdGraphNode Interface.
+	virtual FText GetTooltipText() const override;
+	virtual FText GetNodeTitle(ENodeTitleType::Type TitleType) const override;
+	//~ End UEdGraphNode Interface.
+
+protected:
+	//~ Begin UK2Node_LoadAsset Interface.
+	virtual const FName& GetInputPinName() const;
+	//~ End UK2Node_LoadAsset Interface.
+
+	//~ Begin UMVVMK2Node_MakeBrushFromSoftResource Interface.
+	virtual UClass* GetInputResourceClass() const override;
+	virtual FName GetMakeBrushFunctionName() const override;
+	//~ End UMVVMK2Node_MakeBrushFromSoftResource Interface.
+};
+
+
+//////////////////////////////////////////////////////////////////////////
+
+
+UCLASS(MinimalAPI)
+class UMVVMK2Node_MakeBrushFromSoftMaterial: public UMVVMK2Node_MakeBrushFromSoftResource
+{
+	GENERATED_BODY()
+
+public:
+	//~ Begin UEdGraphNode Interface.
+	virtual FText GetTooltipText() const override;
+	virtual FText GetNodeTitle(ENodeTitleType::Type TitleType) const override;
+	//~ End UEdGraphNode Interface.
+
+protected:
+	//~ Begin UK2Node_LoadAsset Interface.
+	virtual const FName& GetInputPinName() const;
+	//~ End UK2Node_LoadAsset Interface.
+
+	//~ Begin UMVVMK2Node_MakeBrushFromSoftResource Interface.
+	virtual UClass* GetInputResourceClass() const override;
+	virtual FName GetMakeBrushFunctionName() const override;
+	//~ End UMVVMK2Node_MakeBrushFromSoftResource Interface.
+};
diff --git a/Engine/Plugins/Runtime/ModelViewViewModel/Source/ModelViewViewModelBlueprint/Private/Node/MVVMK2Node_SetSoftTextureParameter.cpp b/Engine/Plugins/Runtime/ModelViewViewModel/Source/ModelViewViewModelBlueprint/Private/Node/MVVMK2Node_SetSoftTextureParameter.cpp
new file mode 100644
--- /dev/null
+++ b/Engine/Plugins/Runtime/ModelViewViewModel/Source/ModelViewViewModelBlueprint/Private/Node/MVVMK2Node_SetSoftTextureParameter.cpp
@@ -0,0 +1,134 @@
+// Copyright Epic Games, Inc. All Rights Reserved.
+
+#include "Node/MVVMK2Node_SetSoftTextureParameter.h"
+
+#include "Bindings/ConversionLibraries/MVVMSlateBrushConversionLibrary.h"
+#include "KismetCompiler.h"
+#include "K2Node_CallFunction.h"
+#include "K2Node_ClassDynamicCast.h"
+
+#include "Engine/Texture2D.h"
+
+#define LOCTEXT_NAMESPACE "MVVMK2Node_SetSoftTextureParameter"
+
+namespace UE::MVVM::Private
+{
+	namespace PinNames
+	{
+		static const FName TargetBrushName("TargetBrush");
+		static const FName InputPinName("Texture");
+		static const FName ParamName("ParameterName");
+		static const FName OutputPinName("SlateBrush");
+	}
+}
+
+//////////////////////////////////////////////////////////////////////////
+
+void UMVVMK2Node_SetSoftTextureParameter::AllocateDefaultPins()
+{
+	const UEdGraphSchema_K2* K2Schema = GetDefault<UEdGraphSchema_K2>();
+
+	CreatePin(EGPD_Input, UEdGraphSchema_K2::PC_Exec, UEdGraphSchema_K2::PN_Execute);
+
+	// The immediate continue pin
+	CreatePin(EGPD_Output, UEdGraphSchema_K2::PC_Exec, UEdGraphSchema_K2::PN_Then);
+
+	// The delayed completed pin, this used to be called Then
+	CreatePin(EGPD_Output, UEdGraphSchema_K2::PC_Exec, UEdGraphSchema_K2::PN_Completed);
+	
+	CreatePin(EGPD_Input, UEdGraphSchema_K2::PC_Struct, FSlateBrush::StaticStruct(), UE::MVVM::Private::PinNames::TargetBrushName);
+	CreatePin(EGPD_Input, UEdGraphSchema_K2::PC_Name, nullptr, UE::MVVM::Private::PinNames::ParamName);
+	CreatePin(EGPD_Input, UEdGraphSchema_K2::PC_SoftObject, UTexture2D::StaticClass(), UE::MVVM::Private::PinNames::InputPinName);
+	CreatePin(EGPD_Output, UEdGraphSchema_K2::PC_Struct, FSlateBrush::StaticStruct(), UE::MVVM::Private::PinNames::OutputPinName);
+
+	UK2Node::AllocateDefaultPins();
+}
+
+void UMVVMK2Node_SetSoftTextureParameter::ExpandNode(FKismetCompilerContext& CompilerContext, UEdGraph* SourceGraph)
+{
+	UK2Node::ExpandNode(CompilerContext, SourceGraph);
+	
+	const UEdGraphSchema_K2* Schema = CompilerContext.GetSchema();
+
+	UK2Node_LoadAsset* CallLoadAssetNode = CompilerContext.SpawnIntermediateNode<UK2Node_LoadAsset>(this, SourceGraph);
+	{
+		CallLoadAssetNode->AllocateDefaultPins();
+	}
+	UEdGraphPin* LoadAssetInput = CallLoadAssetNode->FindPinChecked(CallLoadAssetNode->GetInputPinName());
+	UEdGraphPin* LoadAssetOutput = CallLoadAssetNode->FindPinChecked(CallLoadAssetNode->GetOutputPinName());
+
+	UK2Node_CallFunction* CallSetTextureParameter = CompilerContext.SpawnIntermediateNode<UK2Node_CallFunction>(this, SourceGraph);
+	{
+		CallSetTextureParameter->FunctionReference.SetExternalMember(GET_FUNCTION_NAME_CHECKED(UMVVMSlateBrushConversionLibrary, Conv_SetTextureParameter), UMVVMSlateBrushConversionLibrary::StaticClass());
+		CallSetTextureParameter->AllocateDefaultPins();
+	}
+
+	UEdGraphPin* CallSetTextureParameterTexturePin = CallSetTextureParameter->FindPinChecked(FName("Value"));
+	UEdGraphPin* CallSetTextureParameterBrushPin = CallSetTextureParameter->FindPinChecked(FName("Brush"));
+	UEdGraphPin* CallSetTextureParameterParameterPin = CallSetTextureParameter->FindPinChecked(FName("ParameterName"));
+	UEdGraphPin* CallSetTextureParameterResultPin = CallSetTextureParameter->GetReturnValuePin();
+
+	UK2Node_DynamicCast* CastNode = CompilerContext.SpawnIntermediateNode<UK2Node_DynamicCast>(this, SourceGraph);
+	{
+		CastNode->TargetType = UTexture2D::StaticClass();
+		CastNode->AllocateDefaultPins();
+	}
+	UEdGraphPin* CastInput = CastNode->GetCastSourcePin();
+	UEdGraphPin* CastOutput = CastNode->GetCastResultPin();
+
+	// move this.exec to CallLoadAssetNode.exec
+	CompilerContext.MovePinLinksToIntermediate(*GetExecPin(), *CallLoadAssetNode->GetExecPin());
+	// move this.then to CallLoadAssetNode.then
+	CompilerContext.MovePinLinksToIntermediate(*GetThenPin(), *CallLoadAssetNode->GetThenPin());
+	// CallLoadAssetNode.completed to CastNode.exec
+	ensure(Schema->TryCreateConnection(CallLoadAssetNode->FindPinChecked(UEdGraphSchema_K2::PN_Completed), CastNode->GetExecPin()));
+	// move this.completed to CastNode.then
+	CompilerContext.MovePinLinksToIntermediate(*GetCompletedPin(), *CastNode->GetThenPin());
+
+	// move this.resource to CallLoadAssetNode.arg
+	CompilerContext.MovePinLinksToIntermediate(*FindPinChecked(UE::MVVM::Private::PinNames::InputPinName), *LoadAssetInput);
+	// move the target brush pin to the SetTextureParameter brush pin
+	CompilerContext.MovePinLinksToIntermediate(*FindPinChecked(UE::MVVM::Private::PinNames::TargetBrushName), *CallSetTextureParameterBrushPin);
+	// set the parameter name on the SetTextureParameter name pin
+	CallSetTextureParameterParameterPin->DefaultValue = FindPinChecked(UE::MVVM::Private::PinNames::ParamName)->DefaultValue;
+	// CallLoadAssetNode.result to CastNode.Input
+	ensure(Schema->TryCreateConnection(LoadAssetOutput, CastInput));
+	// CastNode.Output to SetTextureParameter.Value
+	ensure(Schema->TryCreateConnection(CastOutput, CallSetTextureParameterTexturePin));
+	// move this.result to SetTextureParameter.result
+	CompilerContext.MovePinLinksToIntermediate(*FindPinChecked(GetOutputPinName()), *CallSetTextureParameterResultPin);
+
+	BreakAllNodeLinks();
+}
+
+UEdGraphPin* UMVVMK2Node_SetSoftTextureParameter::GetThenPin() const
+{
+	return FindPinChecked(UEdGraphSchema_K2::PN_Then);
+}
+
+UEdGraphPin* UMVVMK2Node_SetSoftTextureParameter::GetCompletedPin() const
+{
+	return FindPinChecked(UEdGraphSchema_K2::PN_Completed);
+}
+
+const FName& UMVVMK2Node_SetSoftTextureParameter::GetOutputPinName() const
+{
+	return UE::MVVM::Private::PinNames::OutputPinName;
+}
+
+const FName& UMVVMK2Node_SetSoftTextureParameter::GetInputPinName() const
+{
+	return UE::MVVM::Private::PinNames::InputPinName;
+}
+
+FText UMVVMK2Node_SetSoftTextureParameter::GetTooltipText() const
+{
+	return FText(LOCTEXT("UMVVMK2Node_SetSoftTextureParameterGetTooltipText", "Asynchronously loads a Soft Texture Reference and sets the texture property on the slate brush"));
+}
+
+FText UMVVMK2Node_SetSoftTextureParameter::GetNodeTitle(ENodeTitleType::Type TitleType) const
+{
+	return FText(LOCTEXT("UMVVMK2Node_SetSoftTextureParameterGetNodeTitle", "Set Soft Texture Parameter"));
+}
+
+#undef LOCTEXT_NAMESPACE
diff --git a/Engine/Plugins/Runtime/ModelViewViewModel/Source/ModelViewViewModelBlueprint/Private/Node/MVVMK2Node_SetSoftTextureParameter.h b/Engine/Plugins/Runtime/ModelViewViewModel/Source/ModelViewViewModelBlueprint/Private/Node/MVVMK2Node_SetSoftTextureParameter.h
new file mode 100644
--- /dev/null
+++ b/Engine/Plugins/Runtime/ModelViewViewModel/Source/ModelViewViewModelBlueprint/Private/Node/MVVMK2Node_SetSoftTextureParameter.h
@@ -0,0 +1,54 @@
+// Copyright Epic Games, Inc. All Rights Reserved.
+
+#pragma once
+
+#include "CoreMinimal.h"
+#include "K2Node.h"
+#include "K2Node_LoadAsset.h"
+#include "View/MVVMViewTypes.h"
+
+#include "MVVMK2Node_SetSoftTextureParameter.generated.h"
+
+class UEdGraphPin;
+class FKismetCompilerContext;
+
+
+/**
+ * Utility async node to create a slate brush from a soft resource (Texture, Material, etc). 
+ * Needed since we cannot nest conversion functions (Ex: LoadAsset->MakeBrushFromTexture).
+ */
+UCLASS(MinimalAPI, Category = "Widget", meta = (MVVMBindToDestination = "TargetBrush"))
+class UMVVMK2Node_SetSoftTextureParameter : public UK2Node_LoadAsset
+{
+	GENERATED_BODY()
+
+public:
+	//~ Begin UEdGraphNode Interface.
+	virtual void AllocateDefaultPins() override;
+	//~ End UEdGraphNode Interface.
+	
+	//~ Begin K2Node Interface.
+	virtual void ExpandNode(FKismetCompilerContext& CompilerContext, UEdGraph* SourceGraph) override;
+	//~ End K2Node Interface.
+
+	/** Get the then output pin */
+	UEdGraphPin* GetThenPin() const;
+
+	/** Get the completed output pin */
+	UEdGraphPin* GetCompletedPin() const;
+
+protected:
+
+	//~ Begin UK2Node_LoadAsset Interface.
+	virtual const FName& GetOutputPinName() const override;
+	virtual const FName& GetInputPinName() const override;
+	//~ End UK2Node_LoadAsset Interface.
+
+public:
+	//~ Begin UEdGraphNode Interface.
+	virtual FText GetTooltipText() const override;
+	virtual FText GetNodeTitle(ENodeTitleType::Type TitleType) const override;
+	//~ End UEdGraphNode Interface.
+};
+
+
diff --git a/Engine/Plugins/Runtime/ModelViewViewModel/Source/ModelViewViewModelEditor/Private/Node/MVVMK2Node_LoadSoftResource.cpp b/Engine/Plugins/Runtime/ModelViewViewModel/Source/ModelViewViewModelEditor/Private/Node/MVVMK2Node_LoadSoftResource.cpp
deleted file mode 100644
diff --git a/Engine/Plugins/Runtime/ModelViewViewModel/Source/ModelViewViewModelEditor/Private/Node/MVVMK2Node_LoadSoftResource.h b/Engine/Plugins/Runtime/ModelViewViewModel/Source/ModelViewViewModelEditor/Private/Node/MVVMK2Node_LoadSoftResource.h
deleted file mode 100644
diff --git a/Engine/Plugins/Runtime/ModelViewViewModel/Source/ModelViewViewModelEditor/Private/Node/MVVMK2Node_MakeBrushFromSoftResource.cpp b/Engine/Plugins/Runtime/ModelViewViewModel/Source/ModelViewViewModelEditor/Private/Node/MVVMK2Node_MakeBrushFromSoftResource.cpp
deleted file mode 100644
diff --git a/Engine/Plugins/Runtime/ModelViewViewModel/Source/ModelViewViewModelEditor/Private/Node/MVVMK2Node_MakeBrushFromSoftResource.h b/Engine/Plugins/Runtime/ModelViewViewModel/Source/ModelViewViewModelEditor/Private/Node/MVVMK2Node_MakeBrushFromSoftResource.h
deleted file mode 100644
diff --git a/Engine/Plugins/Runtime/ModelViewViewModel/Source/ModelViewViewModelEditor/Private/Node/MVVMK2Node_SetSoftTextureParameter.cpp b/Engine/Plugins/Runtime/ModelViewViewModel/Source/ModelViewViewModelEditor/Private/Node/MVVMK2Node_SetSoftTextureParameter.cpp
deleted file mode 100644
diff --git a/Engine/Plugins/Runtime/ModelViewViewModel/Source/ModelViewViewModelEditor/Private/Node/MVVMK2Node_SetSoftTextureParameter.h b/Engine/Plugins/Runtime/ModelViewViewModel/Source/ModelViewViewModelEditor/Private/Node/MVVMK2Node_SetSoftTextureParameter.h
deleted file mode 100644
