MVVM: Bindings List Copy/Paste

#rnx

diff --git a/Engine/Plugins/Runtime/ModelViewViewModel/Source/ModelViewViewModelBlueprint/Private/MVVMBlueprintView.cpp b/Engine/Plugins/Runtime/ModelViewViewModel/Source/ModelViewViewModelBlueprint/Private/MVVMBlueprintView.cpp
--- a/Engine/Plugins/Runtime/ModelViewViewModel/Source/ModelViewViewModelBlueprint/Private/MVVMBlueprintView.cpp
+++ b/Engine/Plugins/Runtime/ModelViewViewModel/Source/ModelViewViewModelBlueprint/Private/MVVMBlueprintView.cpp
@@ -283,14 +283,19 @@ const FMVVMBlueprintViewBinding* UMVVMBlueprintView::GetBinding(FGuid Id) const
 UMVVMBlueprintViewEvent* UMVVMBlueprintView::AddDefaultEvent()
 {
 	UMVVMBlueprintViewEvent* Event = NewObject<UMVVMBlueprintViewEvent>(this);
+	AddEvent(Event);
+
+	return Event;
+}
+
+void UMVVMBlueprintView::AddEvent(UMVVMBlueprintViewEvent* Event)
+{
 	Events.Add(Event);
 
 	OnBindingsAdded.Broadcast();
 	OnBindingsUpdated.Broadcast();
 
 	FBlueprintEditorUtils::MarkBlueprintAsModified(GetOuterUMVVMWidgetBlueprintExtension_View()->GetWidgetBlueprint());
-
-	return Event;
 }
 
 void UMVVMBlueprintView::RemoveEvent(UMVVMBlueprintViewEvent* Event)
@@ -307,12 +312,7 @@ void UMVVMBlueprintView::RemoveEvent(UMVVMBlueprintViewEvent* Event)
 UMVVMBlueprintViewEvent* UMVVMBlueprintView::DuplicateEvent(UMVVMBlueprintViewEvent* Event)
 {
 	UMVVMBlueprintViewEvent* NewEvent = DuplicateObject<UMVVMBlueprintViewEvent>(Event, this);
-	Events.Add(NewEvent);
-
-	OnBindingsAdded.Broadcast();
-	OnBindingsUpdated.Broadcast();
-
-	FBlueprintEditorUtils::MarkBlueprintAsModified(GetOuterUMVVMWidgetBlueprintExtension_View()->GetWidgetBlueprint());
+	AddEvent(NewEvent);
 
 	return NewEvent;
 }
@@ -320,14 +320,19 @@ UMVVMBlueprintViewEvent* UMVVMBlueprintView::DuplicateEvent(UMVVMBlueprintViewEv
 UMVVMBlueprintViewCondition* UMVVMBlueprintView::AddDefaultCondition()
 {
 	UMVVMBlueprintViewCondition* Condition = NewObject<UMVVMBlueprintViewCondition>(this);
+	AddCondition(Condition);
+
+	return Condition;
+}
+
+void UMVVMBlueprintView::AddCondition(UMVVMBlueprintViewCondition* Condition)
+{
 	Conditions.Add(Condition);
 
 	OnBindingsAdded.Broadcast();
 	OnBindingsUpdated.Broadcast();
 
 	FBlueprintEditorUtils::MarkBlueprintAsModified(GetOuterUMVVMWidgetBlueprintExtension_View()->GetWidgetBlueprint());
-
-	return Condition;
 }
 
 void UMVVMBlueprintView::RemoveCondition(UMVVMBlueprintViewCondition* Condition)
@@ -344,12 +349,8 @@ void UMVVMBlueprintView::RemoveCondition(UMVVMBlueprintViewCondition* Condition)
 UMVVMBlueprintViewCondition* UMVVMBlueprintView::DuplicateCondition(UMVVMBlueprintViewCondition* Condition)
 {
 	UMVVMBlueprintViewCondition* NewCondition = DuplicateObject<UMVVMBlueprintViewCondition>(Condition, this);
-	Conditions.Add(NewCondition);
 
-	OnBindingsAdded.Broadcast();
-	OnBindingsUpdated.Broadcast();
-
-	FBlueprintEditorUtils::MarkBlueprintAsModified(GetOuterUMVVMWidgetBlueprintExtension_View()->GetWidgetBlueprint());
+	AddCondition(NewCondition);
 
 	return NewCondition;
 }
diff --git a/Engine/Plugins/Runtime/ModelViewViewModel/Source/ModelViewViewModelBlueprint/Public/MVVMBlueprintView.h b/Engine/Plugins/Runtime/ModelViewViewModel/Source/ModelViewViewModelBlueprint/Public/MVVMBlueprintView.h
--- a/Engine/Plugins/Runtime/ModelViewViewModel/Source/ModelViewViewModelBlueprint/Public/MVVMBlueprintView.h
+++ b/Engine/Plugins/Runtime/ModelViewViewModel/Source/ModelViewViewModelBlueprint/Public/MVVMBlueprintView.h
@@ -136,6 +136,7 @@ public:
 	}
 
 	UMVVMBlueprintViewEvent* AddDefaultEvent();
+	void AddEvent(UMVVMBlueprintViewEvent* Event);
 	void RemoveEvent(UMVVMBlueprintViewEvent* Event);
 	UMVVMBlueprintViewEvent* DuplicateEvent(UMVVMBlueprintViewEvent* Event);
 
@@ -150,6 +151,7 @@ public:
 	}
 
 	UMVVMBlueprintViewCondition* AddDefaultCondition();
+	void AddCondition(UMVVMBlueprintViewCondition* Condition);
 	void RemoveCondition(UMVVMBlueprintViewCondition* Condition);
 	UMVVMBlueprintViewCondition* DuplicateCondition(UMVVMBlueprintViewCondition* Condition);
 
diff --git a/Engine/Plugins/Runtime/ModelViewViewModel/Source/ModelViewViewModelEditor/ModelViewViewModelEditor.Build.cs b/Engine/Plugins/Runtime/ModelViewViewModel/Source/ModelViewViewModelEditor/ModelViewViewModelEditor.Build.cs
--- a/Engine/Plugins/Runtime/ModelViewViewModel/Source/ModelViewViewModelEditor/ModelViewViewModelEditor.Build.cs
+++ b/Engine/Plugins/Runtime/ModelViewViewModel/Source/ModelViewViewModelEditor/ModelViewViewModelEditor.Build.cs
@@ -28,6 +28,7 @@ public class ModelViewViewModelEditor : ModuleRules
 			new string[]
 			{
 				"AdvancedWidgets",
+				"ApplicationCore",
 				"AssetTools",
 				"EditorWidgets",
 				"GraphEditor",
diff --git a/Engine/Plugins/Runtime/ModelViewViewModel/Source/ModelViewViewModelEditor/Private/Framework/MVVMRowHelper.cpp b/Engine/Plugins/Runtime/ModelViewViewModel/Source/ModelViewViewModelEditor/Private/Framework/MVVMRowHelper.cpp
--- a/Engine/Plugins/Runtime/ModelViewViewModel/Source/ModelViewViewModelEditor/Private/Framework/MVVMRowHelper.cpp
+++ b/Engine/Plugins/Runtime/ModelViewViewModel/Source/ModelViewViewModelEditor/Private/Framework/MVVMRowHelper.cpp
@@ -10,11 +10,17 @@
 #include "Editor/EditorEngine.h"
 #include "EdGraphSchema_K2.h"
 #include "EdGraphUtilities.h"
+#include "Factories.h"
+#include "HAL/PlatformApplicationMisc.h"
 #include "Kismet2/BlueprintEditorUtils.h"
 #include "Kismet2/KismetEditorUtilities.h"
 #include "Misc/Guid.h"
 #include "ScopedTransaction.h"
+#include "UnrealExporter.h"
+#include "UObject/UObjectGlobals.h"
+#include "Exporters/Exporter.h"
 #include "Widgets/ViewModelFieldDragDropOp.h"
+#include "WidgetBlueprint.h"
 
 #include "Styling/AppStyle.h"
 
@@ -24,6 +30,8 @@
 #include "MVVMBlueprintViewConversionFunction.h"
 #include "MVVMBlueprintViewEvent.h"
 #include "MVVMBlueprintViewCondition.h"
+#include "MVVMMessageLog.h"
+#include "MVVMPropertyPath.h"
 #include "MVVMWidgetBlueprintExtension_View.h"
 #include "Misc/MessageDialog.h"
 #include "Types/MVVMBindingEntry.h"
@@ -36,6 +44,259 @@ extern UNREALED_API UEditorEngine* GEditor;
 namespace UE::MVVM::BindingEntry
 {
 
+namespace Private
+{
+// Text import/export uses the unit seperator ascii character code as a delimiter. 
+// Using a non-printable character because all objects and structs are exported as readable text.
+static constexpr TStaticArray<TCHAR, 2> BindingClipboardDelimiter = { 0x1f, TEXT('\0') };
+
+struct FBindingClipboardData
+{
+	using Type = TVariant<FMVVMBlueprintViewBinding, UMVVMBlueprintViewCondition*, UMVVMBlueprintViewEvent*>;
+	TArray<Type> Items;
+};
+
+bool ExportClipboardData(UMVVMBlueprintView* BlueprintView, const TArray<TSharedPtr<FBindingEntry>>& Entries, FBindingClipboardData& OutClipboardData)
+{
+	TArray<TSharedPtr<FBindingEntry>> EntriesToExport;
+	EntriesToExport.Append(Entries);
+
+	while (!EntriesToExport.IsEmpty())
+	{
+		TSharedPtr<FBindingEntry> Entry = EntriesToExport.Pop();
+
+		switch (Entry->GetRowType())
+		{
+		case FBindingEntry::ERowType::Binding:
+		{
+			if (FMVVMBlueprintViewBinding* Binding = Entry->GetBinding(BlueprintView))
+			{
+				FBindingClipboardData::Type Item;
+				Item.Set<FMVVMBlueprintViewBinding>(*Binding);
+				OutClipboardData.Items.Add(Item);
+			}
+			break;
+		}
+		case FBindingEntry::ERowType::Event:
+		{
+			FBindingClipboardData::Type Item;
+			Item.Set<UMVVMBlueprintViewEvent*>(Entry->GetEvent());
+			OutClipboardData.Items.Add(Item);
+			break;
+		}
+		case FBindingEntry::ERowType::Condition:
+		{
+			FBindingClipboardData::Type Item;
+			Item.Set<UMVVMBlueprintViewCondition*>(Entry->GetCondition());
+			OutClipboardData.Items.Add(Item);
+			break;
+		}
+		case FBindingEntry::ERowType::Group:
+		{
+			EntriesToExport.Append(Entry->GetFilteredChildren());
+			break;
+		}
+		case FBindingEntry::ERowType::BindingParameter:
+		{
+			UE_LOG(LogMVVM, Warning, TEXT("Failed to copy %s. Parameter Copy/Paste not supported."), *Entry->GetBindingParameterId().ToString());
+			break;
+		}
+		case FBindingEntry::ERowType::EventParameter:
+		{
+			UE_LOG(LogMVVM, Warning, TEXT("Failed to copy %s. Parameter Copy/Paste not supported."), *Entry->GetEventParameterId().ToString());
+			break;
+		}
+		case FBindingEntry::ERowType::ConditionParameter:
+		{
+			UE_LOG(LogMVVM, Warning, TEXT("Failed to copy %s. Parameter Copy/Paste not supported."), *Entry->GetConditionParameterId().ToString());
+			break;
+		}
+		default:
+		{
+			break;
+		}
+		}
+	}
+
+	return !OutClipboardData.Items.IsEmpty();
+}
+
+bool ImportClipboardData(UMVVMBlueprintView* BlueprintView, FBindingClipboardData& OutClipboardData)
+{
+	class FBindingEntryTextFactory : public FCustomizableTextObjectFactory
+	{
+	public:
+		FBindingEntryTextFactory()
+			: FCustomizableTextObjectFactory(GWarn)
+		{}
+
+		virtual bool CanCreateClass(UClass* ObjectClass, bool& bOmitSubObjs) const override final
+		{
+			return ObjectClass == UMVVMBlueprintViewEvent::StaticClass() || ObjectClass == UMVVMBlueprintViewCondition::StaticClass();
+		}
+
+		virtual void ProcessConstructedObject(UObject* NewObject) override final
+		{
+			Objects.Add(NewObject);
+		}
+
+		TArray<UObject*> Objects;
+	};
+
+	FString TextToImport;
+	FPlatformApplicationMisc::ClipboardPaste(TextToImport);
+
+	TArray<FString> Items;
+	TextToImport.ParseIntoArray(Items, BindingClipboardDelimiter.GetData());
+
+	for (const FString& ItemText : Items)
+	{
+		const FString StructId = TBaseStructure<FMVVMBlueprintViewBinding>::Get()->GetFName().ToString();
+		if (ItemText.StartsWith(StructId))
+		{
+			check(ItemText.Len() > StructId.Len());
+			const TCHAR* ItemTextData = &ItemText[StructId.Len()];
+
+			FMVVMBlueprintViewBinding BindingCopied;
+			FStringOutputDevice ErrorOutput;
+			FMVVMBlueprintViewBinding::StaticStruct()->ImportText(ItemTextData, &BindingCopied, BlueprintView, PPF_None, &ErrorOutput, StructId);
+
+			FBindingClipboardData::Type Item;
+			Item.Set<FMVVMBlueprintViewBinding>(BindingCopied);
+			OutClipboardData.Items.Add(Item);
+		}
+		else
+		{
+			FBindingEntryTextFactory Factory;
+			Factory.ProcessBuffer(BlueprintView, RF_Transactional, ItemText);
+			if (Factory.Objects.Num() != 1)
+			{
+				UE_LOG(LogMVVM, Warning, TEXT("Failed to import object from clipboard"));
+				continue;
+			}
+
+			UObject* Object = Factory.Objects.Last();
+			if (Object->GetClass() == UMVVMBlueprintViewCondition::StaticClass())
+			{
+				FBindingClipboardData::Type Item;
+				Item.Set<UMVVMBlueprintViewCondition*>(StaticCast<UMVVMBlueprintViewCondition*>(Object));
+				OutClipboardData.Items.Add(Item);
+			}
+			else if (Object->GetClass() == UMVVMBlueprintViewEvent::StaticClass())
+			{
+				FBindingClipboardData::Type Item;
+				Item.Set<UMVVMBlueprintViewEvent*>(StaticCast<UMVVMBlueprintViewEvent*>(Object));
+				OutClipboardData.Items.Add(Item);
+			}
+		}
+	}
+
+	return !OutClipboardData.Items.IsEmpty();
+}
+
+bool GatherPropertyPath(UWidgetBlueprint* WidgetBlueprint, UStruct* Struct, const FMVVMBlueprintPropertyPath& PathToFind, UClass* ContextClass, FMVVMBlueprintPropertyPath& OutputPath)
+{
+	for (const FMVVMBlueprintFieldPath& Path : PathToFind.GetFieldPaths())
+	{
+		bool bFound = false;
+
+		if (Struct != nullptr)
+		{
+			const FName PropertyName = Path.GetFieldName(ContextClass);
+			if (const FProperty* Property = Struct->FindPropertyByName(PropertyName))
+			{
+				bFound = true;
+				OutputPath.AppendPropertyPath(WidgetBlueprint, UE::MVVM::FMVVMConstFieldVariant(Property));
+
+				if (const FObjectPropertyBase* ObjectProperty = CastField<const FObjectPropertyBase>(Property))
+				{
+					Struct = ObjectProperty->PropertyClass.Get();
+				}
+				else if (const FStructProperty* StructProperty = CastField<const FStructProperty>(Property))
+				{
+					Struct = StructProperty->Struct.Get();
+				}
+			}
+			else if (const UClass* Class = dynamic_cast<UClass*>(Struct))
+			{
+				if (const UFunction* Function = Class->FindFunctionByName(PropertyName))
+				{
+					bFound = true;
+					OutputPath.AppendPropertyPath(WidgetBlueprint, UE::MVVM::FMVVMConstFieldVariant(Function));
+				}
+			}
+		}
+
+		if (!bFound)
+		{
+			return false;
+		}
+	}
+
+	return true;
+}
+
+/*
+	Attempt to find the desired property path under the target widget's ownership. If not found, defaults back to the original target path.
+*/
+FMVVMBlueprintPropertyPath TryGetCommonPropertyPath(UWidgetBlueprint* WidgetBlueprint, UMVVMBlueprintView* View, const FMVVMBlueprintPropertyPath& TargetPath, const FMVVMBlueprintPropertyPath& DesiredPath)
+{
+	FMVVMBlueprintPropertyPath NewPath = TargetPath;
+	NewPath.ResetPropertyPath();
+
+	UClass* ContextClass = WidgetBlueprint->SkeletonGeneratedClass ? WidgetBlueprint->SkeletonGeneratedClass : WidgetBlueprint->GeneratedClass;
+
+	switch (TargetPath.GetSource(WidgetBlueprint))
+	{
+	case EMVVMBlueprintFieldPathSource::SelfContext:
+	{
+		// Both paths are under that same Blueprint, so all relative paths are available
+		if (TargetPath.GetSource(WidgetBlueprint) == DesiredPath.GetSource(WidgetBlueprint))
+		{
+			for (const FMVVMBlueprintFieldPath& Path : DesiredPath.GetFieldPaths())
+			{
+				NewPath.AppendPropertyPath(WidgetBlueprint, Path.GetField(ContextClass));
+			}
+			break;
+		}
+		else
+		{
+			UClass* TargetClass = WidgetBlueprint->GeneratedClass.Get();
+			if (!GatherPropertyPath(WidgetBlueprint, TargetClass, DesiredPath, ContextClass, NewPath))
+			{
+				NewPath = TargetPath;
+			}
+		}
+		break;
+	}
+	case EMVVMBlueprintFieldPathSource::Widget:
+	{
+		if (UWidget* Widget = WidgetBlueprint->WidgetTree->FindWidget(TargetPath.GetWidgetName()))
+		{
+			UClass* TargetClass = Widget->GetClass();
+			if (!GatherPropertyPath(WidgetBlueprint, TargetClass, DesiredPath, ContextClass, NewPath))
+			{
+				NewPath = TargetPath;
+			}
+		}
+		break;
+	}
+	case EMVVMBlueprintFieldPathSource::ViewModel:
+	{
+		UE_LOG(LogMVVM, Error, TEXT("MVVM: View Bindings does not support ViewModel properties as destination bindings"));
+		break;
+	}
+	default:
+	{
+		check(false);
+		break;
+	}
+	}
+
+	return NewPath;
+}
+}
+
 void FRowHelper::GatherAllChildBindings(UMVVMBlueprintView* BlueprintView, const TConstArrayView<TSharedPtr<FBindingEntry>> Entries, TArray<FGuid>& OutBindingIds, TArray<UMVVMBlueprintViewEvent*>& OutEvents, TArray<UMVVMBlueprintViewCondition*>& OutConditions)
 {
 	for (const TSharedPtr<FBindingEntry>& Entry : Entries)
@@ -145,7 +406,7 @@ void FRowHelper::DuplicateEntries(const UWidgetBlueprint* WidgetBlueprint, UMVVM
 		{
 			return;
 		}
- 
+
 		OutSelection.Reserve(Selection.Num());
 
 		FScopedTransaction Transaction(LOCTEXT("DuplicateBindingsTransaction", "Duplicate Bindings"));
@@ -173,7 +434,7 @@ void FRowHelper::DuplicateEntries(const UWidgetBlueprint* WidgetBlueprint, UMVVM
 
 			OutSelection.Add(NewEntry);
 		}
-		
+
 		for (UMVVMBlueprintViewCondition* Condition : ConditionsToDuplicate)
 		{
 			UMVVMBlueprintViewCondition* NewCondition = BlueprintView->DuplicateCondition(Condition);
@@ -186,6 +447,256 @@ void FRowHelper::DuplicateEntries(const UWidgetBlueprint* WidgetBlueprint, UMVVM
 	}
 }
 
+void FRowHelper::CopyEntries(const UWidgetBlueprint* WidgetBlueprint, UMVVMBlueprintView* BlueprintView, const TArray<TSharedPtr<FBindingEntry>>& Entries)
+{
+	if (WidgetBlueprint == nullptr || BlueprintView == nullptr || Entries.IsEmpty())
+	{
+		return;
+	}
+		
+	Private::FBindingClipboardData ClipboardData;
+	if (!Private::ExportClipboardData(BlueprintView, Entries, ClipboardData))
+	{
+		return;
+	}
+	
+	FString CopyText;
+
+	for (const Private::FBindingClipboardData::Type& Item : ClipboardData.Items)
+	{
+		bool bFailed = false;
+		UScriptStruct* Struct = nullptr;
+		const void* StructData = nullptr;
+		UObject* Object = nullptr;
+
+		if (Item.IsType<FMVVMBlueprintViewBinding>())
+		{
+			StructData = &Item.Get<FMVVMBlueprintViewBinding>();
+			Struct = FMVVMBlueprintViewBinding::StaticStruct();
+		}
+		else if (Item.IsType<UMVVMBlueprintViewCondition*>())
+		{
+			Object = Item.Get<UMVVMBlueprintViewCondition*>();
+		}
+		else if (Item.IsType<UMVVMBlueprintViewEvent*>())
+		{
+			Object = Item.Get<UMVVMBlueprintViewEvent*>();
+		}
+
+		if (Struct != nullptr && StructData != nullptr)
+		{
+			FString Text;
+			Struct->ExportText(Text, StructData, nullptr, BlueprintView, PPF_None, nullptr);
+
+			if (Text.IsEmpty())
+			{
+				bFailed = true;
+				continue;
+			}
+
+			if (!CopyText.IsEmpty())
+			{
+				CopyText.Append(Private::BindingClipboardDelimiter);
+			}
+
+			// Prepend struct type to support explicit struct deserialization. 
+			CopyText.Append(Struct->GetFName().ToString());
+			CopyText.Append(Text);
+		}
+		else if (Object != nullptr)
+		{
+			FStringOutputDevice Output;
+			const FExportObjectInnerContext Context;
+
+			UnMarkAllObjects(EObjectMark(OBJECTMARK_TagExp | OBJECTMARK_TagImp));
+			UExporter::ExportToOutputDevice(&Context, Object, nullptr, Output, TEXT("copy"), 0, PPF_ExportsNotFullyQualified | PPF_Copy | PPF_Delimited, false, Object->GetOuter());
+
+			if (Output.IsEmpty())
+			{
+				bFailed = true;
+				continue;
+			}
+
+			if (!CopyText.IsEmpty())
+			{
+				CopyText.Append(Private::BindingClipboardDelimiter);
+			}
+
+			CopyText.Append(*Output);
+		}
+		
+		if (bFailed)
+		{
+			if (Item.IsType<FMVVMBlueprintViewBinding>())
+			{
+				const FMVVMBlueprintViewBinding& Binding = Item.Get<FMVVMBlueprintViewBinding>();
+				UE_LOG(LogMVVM, Warning, TEXT("Failed to export binding %s for copy"), *(Binding.GetDisplayNameString(WidgetBlueprint, true)));
+			}
+			else if (Item.IsType<UMVVMBlueprintViewCondition*>())
+			{
+				UMVVMBlueprintViewCondition* Condition = Item.Get<UMVVMBlueprintViewCondition*>();
+				UE_LOG(LogMVVM, Warning, TEXT("Failed to export condition %s for copy"), *(Condition->GetDisplayName(true).ToString()));
+			}
+			else if (Item.IsType<UMVVMBlueprintViewEvent*>())
+			{
+				UMVVMBlueprintViewEvent* Event = Item.Get<UMVVMBlueprintViewEvent*>();
+				UE_LOG(LogMVVM, Warning, TEXT("Failed to export event %s for copy"), *(Event->GetDisplayName(true).ToString()));
+			}
+		}
+	}
+
+	if (!CopyText.IsEmpty())
+	{
+		FPlatformApplicationMisc::ClipboardCopy(*CopyText);
+	}
+}
+
+void FRowHelper::PasteEntries(UWidgetBlueprint* WidgetBlueprint, UMVVMBlueprintView* BlueprintView, const TArray<TSharedPtr<FBindingEntry>>& Entries)
+{
+	// When pasting new bindings, this will attempt to preserve the destination path. 
+	// If the relative destination path that was copied exists under the selected source, the use this property path. Otherwise, fallback to previous path. 
+	// Treated as a special condition with copy/paste behavior to preserve the existing binding group hierarchy
+
+	if (WidgetBlueprint == nullptr || BlueprintView == nullptr)
+	{
+		return;
+	}
+
+	Private::FBindingClipboardData ClipboardData;
+	Private::ImportClipboardData(BlueprintView, ClipboardData);
+
+	if (Entries.IsEmpty())
+	{
+		// Duplicate the Bindings when pasting into nothing
+		FScopedTransaction Transaction(LOCTEXT("PasteBindingsTransaction", "Paste Bindings"));
+		BlueprintView->Modify();
+
+		for (const Private::FBindingClipboardData::Type& Item : ClipboardData.Items)
+		{
+			if (Item.IsType<FMVVMBlueprintViewBinding>())
+			{
+				const FMVVMBlueprintViewBinding& Binding = Item.Get<FMVVMBlueprintViewBinding>();
+				const FMVVMBlueprintViewBinding* NewBinding = BlueprintView->DuplicateBinding(&Binding);
+
+				TSharedPtr<FBindingEntry> NewEntry = MakeShared<FBindingEntry>();
+				NewEntry->SetBindingId(NewBinding->BindingId);
+			}
+			else if (Item.IsType<UMVVMBlueprintViewCondition*>())
+			{
+				UMVVMBlueprintViewCondition* Condition = Item.Get<UMVVMBlueprintViewCondition*>();
+				UMVVMBlueprintViewCondition* NewCondition = BlueprintView->DuplicateCondition(Condition);
+
+				TSharedPtr<FBindingEntry> NewEntry = MakeShared<FBindingEntry>();
+				NewEntry->SetCondition(NewCondition);
+			}
+			else if (Item.IsType<UMVVMBlueprintViewEvent*>())
+			{
+				UMVVMBlueprintViewEvent* Event = Item.Get<UMVVMBlueprintViewEvent*>();
+				UMVVMBlueprintViewEvent* NewEvent = BlueprintView->DuplicateEvent(Event);
+
+				TSharedPtr<FBindingEntry> NewEntry = MakeShared<FBindingEntry>();
+				NewEntry->SetEvent(NewEvent);
+			}
+		}
+		return;
+	}
+
+	if (Entries.Num() != 1 || ClipboardData.Items.Num() != 1)
+	{
+		UE_LOG(LogMVVM, Warning, TEXT("Copy/Pasting into multiple entries is not supported."));
+		return;
+	}
+
+	const TSharedPtr<FBindingEntry> Entry = Entries.Last();
+	const Private::FBindingClipboardData::Type& Item = ClipboardData.Items.Last();
+	UMVVMEditorSubsystem* EditorSubsystem = GEditor->GetEditorSubsystem<UMVVMEditorSubsystem>();
+
+	if (Item.IsType<FMVVMBlueprintViewBinding>())
+	{
+		if (Entry->GetRowType() != FBindingEntry::ERowType::Binding)
+		{
+			UE_LOG(LogMVVM, Error, TEXT("Failed to paste into Property Binding, mismatched types."));
+			return;
+		}
+
+		FScopedTransaction Transaction(LOCTEXT("PasteBindingTransaction", "Paste Binding"));
+		BlueprintView->Modify();
+
+		const FMVVMBlueprintViewBinding& BindingToPaste = Item.Get<FMVVMBlueprintViewBinding>();
+
+		FMVVMBlueprintViewBinding* BindingSelected = Entry->GetBinding(BlueprintView);
+		BindingSelected->BindingType = BindingToPaste.BindingType;
+		BindingSelected->SourcePath = BindingToPaste.SourcePath;
+
+		// Copy over conversion functions
+		if (UMVVMBlueprintViewConversionFunction* ConversionFunction = BindingToPaste.Conversion.SourceToDestinationConversion)
+		{
+			BindingSelected->Conversion.SourceToDestinationConversion = DuplicateObject<UMVVMBlueprintViewConversionFunction>(ConversionFunction, WidgetBlueprint);
+			BindingSelected->Conversion.SourceToDestinationConversion->RecreateWrapperGraph(WidgetBlueprint);
+		}
+
+		if (UMVVMBlueprintViewConversionFunction* ConversionFunction = BindingToPaste.Conversion.DestinationToSourceConversion)
+		{
+			BindingSelected->Conversion.DestinationToSourceConversion = DuplicateObject<UMVVMBlueprintViewConversionFunction>(ConversionFunction, WidgetBlueprint);
+			BindingSelected->Conversion.DestinationToSourceConversion->RecreateWrapperGraph(WidgetBlueprint);
+		}
+
+		const FMVVMBlueprintPropertyPath& DesiredPath = BindingToPaste.DestinationPath; 
+		const FMVVMBlueprintPropertyPath& TargetPath = BindingSelected->DestinationPath;
+		const FMVVMBlueprintPropertyPath NewPath = Private::TryGetCommonPropertyPath(WidgetBlueprint, BlueprintView, TargetPath, DesiredPath);
+
+		EditorSubsystem->SetDestinationPathForBinding(WidgetBlueprint, *BindingSelected, NewPath, false);
+	}
+	else if (Item.IsType<UMVVMBlueprintViewCondition*>())
+	{
+		if (Entry->GetRowType() != FBindingEntry::ERowType::Condition)
+		{
+			UE_LOG(LogMVVM, Error, TEXT("Failed to paste into Condition, mismatched types."));
+			return;
+		}
+
+		FScopedTransaction Transaction(LOCTEXT("PasteConditionTransaction", "Paste Condition"));
+		BlueprintView->Modify();
+
+		// Update the pasted condition's path and replace selected condition with the condition being pasted.
+		UMVVMBlueprintViewCondition* ConditionSelected = Entry->GetCondition();
+		UMVVMBlueprintViewCondition* ConditionToPaste = Item.Get<UMVVMBlueprintViewCondition*>();
+
+		const FMVVMBlueprintPropertyPath& DesiredPath = ConditionToPaste->GetConditionPath();
+		const FMVVMBlueprintPropertyPath& TargetPath = ConditionSelected->GetConditionPath();
+		const FMVVMBlueprintPropertyPath NewPath = Private::TryGetCommonPropertyPath(WidgetBlueprint, BlueprintView, TargetPath, DesiredPath);
+
+		EditorSubsystem->SetConditionPath(ConditionToPaste, NewPath, false);
+
+		BlueprintView->RemoveCondition(ConditionSelected);
+		BlueprintView->AddCondition(ConditionToPaste);
+	}
+	else if (Item.IsType<UMVVMBlueprintViewEvent*>())
+	{
+		if (Entry->GetRowType() != FBindingEntry::ERowType::Event)
+		{
+			UE_LOG(LogMVVM, Error, TEXT("Failed to paste into Event, mismatched types."));
+			return;
+		}
+
+		FScopedTransaction Transaction(LOCTEXT("PasteEventTransaction", "Paste Event"));
+		BlueprintView->Modify();
+
+		// Update the pasted event's path and replace selected event with the event being pasted.
+		UMVVMBlueprintViewEvent* EventSelected = Entry->GetEvent();
+		UMVVMBlueprintViewEvent* EventToPaste = Item.Get<UMVVMBlueprintViewEvent*>();
+
+		const FMVVMBlueprintPropertyPath& DesiredPath = EventToPaste->GetEventPath();
+		const FMVVMBlueprintPropertyPath& TargetPath = EventSelected->GetEventPath();
+		const FMVVMBlueprintPropertyPath NewPath = Private::TryGetCommonPropertyPath(WidgetBlueprint, BlueprintView, TargetPath, DesiredPath);
+
+		EditorSubsystem->SetEventPath(EventToPaste, NewPath, false);
+
+		BlueprintView->RemoveEvent(EventSelected);
+		BlueprintView->AddEvent(EventToPaste);
+	}
+}
+
 void FRowHelper::ShowBlueprintGraph(FBlueprintEditor* BlueprintEditor, UWidgetBlueprint* WidgetBlueprint, UMVVMBlueprintView* BlueprintView, TArrayView<const TSharedPtr<FBindingEntry>> Entries)
 {
 	auto ShowGraph = [WidgetBlueprint, BlueprintEditor](UEdGraph* Graph)
@@ -380,6 +891,16 @@ void HandleDuplicateEntries(UWidgetBlueprint* WidgetBlueprint, UMVVMBlueprintVie
 	OnSelectionChanged.ExecuteIfBound(NewSelection);
 }
 
+void HandleCopyEntries(UWidgetBlueprint* WidgetBlueprint, UMVVMBlueprintView* View, TSharedPtr<TArray<TWeakPtr<FBindingEntry>>> Entries)
+{
+	FRowHelper::CopyEntries(WidgetBlueprint, View, WeakToSharedPtr(Entries));
+}
+
+void HandlePasteEntries(UWidgetBlueprint* WidgetBlueprint, UMVVMBlueprintView* View, TSharedPtr<TArray<TWeakPtr<FBindingEntry>>> Entries)
+{
+	FRowHelper::PasteEntries(WidgetBlueprint, View, WeakToSharedPtr(Entries));
+}
+
 void HandleResetSelectedPin(UWidgetBlueprint* WidgetBlueprint, UMVVMBlueprintView* View, TSharedPtr<TArray<TWeakPtr<FBindingEntry>>> Entries)
 {
 	const UMVVMEditorSubsystem* EditorSubsystem = GEditor->GetEditorSubsystem<UMVVMEditorSubsystem>();
@@ -480,6 +1001,8 @@ FMenuBuilder FRowHelper::CreateContextMenu(UWidgetBlueprint* WidgetBlueprint, UM
 		const UMVVMEditorSubsystem* EditorSubsystem = GEditor->GetEditorSubsystem<UMVVMEditorSubsystem>();
 		{
 			bool bCanRemoveEntry = true;
+			bool bCanCopyEntry = true;
+			bool bCanPasteEntry = Entries.Num() <= 1;
 			bool bCanDuplicateEntry = true;
 			for (const TSharedPtr<FBindingEntry>& Entry : Entries)
 			{
@@ -487,12 +1010,15 @@ FMenuBuilder FRowHelper::CreateContextMenu(UWidgetBlueprint* WidgetBlueprint, UM
 				{
 				case FBindingEntry::ERowType::Group:
 					bCanDuplicateEntry = false;
-					break; 
+					bCanPasteEntry = false;
+					break;
 				case FBindingEntry::ERowType::Binding:
 				case FBindingEntry::ERowType::Event:
 				case FBindingEntry::ERowType::Condition:
 					break;
 				default:
+					bCanCopyEntry = false;
+					bCanPasteEntry = false;
 					bCanRemoveEntry = false;
 					bCanDuplicateEntry = false;
 					break;
@@ -514,6 +1040,22 @@ FMenuBuilder FRowHelper::CreateContextMenu(UWidgetBlueprint* WidgetBlueprint, UM
 				LOCTEXT("DuplicateBindingTooltip", "Duplicate bindings or events."),
 				FSlateIcon(FAppStyle::GetAppStyleSetName(), "Icons.Duplicate"),
 				DuplicateAction);
+
+			FUIAction CopyAction;
+			CopyAction.ExecuteAction = FExecuteAction::CreateStatic(&Private::HandleCopyEntries, WidgetBlueprint, View, CopiedEntries);
+			CopyAction.CanExecuteAction = FCanExecuteAction::CreateLambda([bCanCopyEntry]() { return bCanCopyEntry; });
+			MenuBuilder.AddMenuEntry(LOCTEXT("CopyBinding", "Copy"),
+				LOCTEXT("CopyBindingTooltip", "Copy bindings or events."),
+				FSlateIcon(FAppStyle::GetAppStyleSetName(), "Icons.Copy"),
+				CopyAction);
+
+			FUIAction PasteAction;
+			PasteAction.ExecuteAction = FExecuteAction::CreateStatic(&Private::HandlePasteEntries, WidgetBlueprint, View, CopiedEntries);
+			PasteAction.CanExecuteAction = FCanExecuteAction::CreateLambda([bCanPasteEntry]() { return bCanPasteEntry; });
+			MenuBuilder.AddMenuEntry(LOCTEXT("PasteBinding", "Paste"),
+				LOCTEXT("PasteBindingTooltip", "Paste bindings or events."),
+				FSlateIcon(FAppStyle::GetAppStyleSetName(), "Icons.Paste"),
+				PasteAction);
 		}
 
 		{
diff --git a/Engine/Plugins/Runtime/ModelViewViewModel/Source/ModelViewViewModelEditor/Private/Framework/MVVMRowHelper.h b/Engine/Plugins/Runtime/ModelViewViewModel/Source/ModelViewViewModelEditor/Private/Framework/MVVMRowHelper.h
--- a/Engine/Plugins/Runtime/ModelViewViewModel/Source/ModelViewViewModelEditor/Private/Framework/MVVMRowHelper.h
+++ b/Engine/Plugins/Runtime/ModelViewViewModel/Source/ModelViewViewModelEditor/Private/Framework/MVVMRowHelper.h
@@ -30,6 +30,9 @@ struct FRowHelper
 	static void DeleteEntries(const UWidgetBlueprint* WidgetBlueprint, UMVVMBlueprintView* BlueprintView, TArrayView<const TSharedPtr<FBindingEntry>> Selection);
 	static void DuplicateEntries(const UWidgetBlueprint* WidgetBlueprint, UMVVMBlueprintView* BlueprintView, TArrayView<const TSharedPtr<FBindingEntry>> Selection, TArray<const TSharedPtr<FBindingEntry>>& OutSelection);
 
+	static void CopyEntries(const UWidgetBlueprint* WidgetBlueprint, UMVVMBlueprintView* BlueprintView, const TArray<TSharedPtr<FBindingEntry>>& Entries);
+	static void PasteEntries(UWidgetBlueprint* WidgetBlueprint, UMVVMBlueprintView* BlueprintView, const TArray<TSharedPtr<FBindingEntry>>& Entries);
+
 	static void ShowBlueprintGraph(FBlueprintEditor* Editor, UWidgetBlueprint* WidgetBlueprint, UMVVMBlueprintView* BlueprintView, TArrayView<const TSharedPtr<FBindingEntry>> Selection);
 
 	static TOptional<FMVVMBlueprintPropertyPath> DropFieldSelector(const UWidgetBlueprint* WidgetBlueprint, const FDragDropEvent& DragDropEvent, bool bIsSource);
diff --git a/Engine/Plugins/Runtime/ModelViewViewModel/Source/ModelViewViewModelEditor/Private/Widgets/SMVVMViewBindingListView.cpp b/Engine/Plugins/Runtime/ModelViewViewModel/Source/ModelViewViewModelEditor/Private/Widgets/SMVVMViewBindingListView.cpp
--- a/Engine/Plugins/Runtime/ModelViewViewModel/Source/ModelViewViewModelEditor/Private/Widgets/SMVVMViewBindingListView.cpp
+++ b/Engine/Plugins/Runtime/ModelViewViewModel/Source/ModelViewViewModelEditor/Private/Widgets/SMVVMViewBindingListView.cpp
@@ -817,6 +817,28 @@ FReply SBindingsList::OnKeyDown(const FGeometry& MyGeometry, const FKeyEvent& In
 		return FReply::Handled();
 	}
 
+	else if (InKeyEvent.GetModifierKeys().IsControlDown())
+	{
+			if (UMVVMWidgetBlueprintExtension_View* MVVMExtensionPtr = MVVMExtension.Get())
+			{
+				UMVVMBlueprintView* BlueprintView = MVVMExtensionPtr->GetBlueprintView();
+				UWidgetBlueprint* WidgetBlueprint = MVVMExtensionPtr->GetWidgetBlueprint();
+				TArray<TSharedPtr<FBindingEntry>> Selection = TreeView->GetSelectedItems();
+
+				// Only support single selection for copy/paste behavior
+				if (InKeyEvent.GetKey() == EKeys::C)
+				{
+					BindingEntry::FRowHelper::CopyEntries(WidgetBlueprint, BlueprintView, Selection);
+				}
+				else if (InKeyEvent.GetKey() == EKeys::V)
+				{
+					BindingEntry::FRowHelper::PasteEntries(WidgetBlueprint, BlueprintView, Selection);
+				}
+			}
+			return FReply::Handled();
+		
+	}
+
 	return FReply::Unhandled();
 }
 
