MVVM: Duplicate Keybinding for bindings. Keybinding tooltip for binding context menu

#rnx

diff --git a/Engine/Plugins/Runtime/ModelViewViewModel/Source/ModelViewViewModelEditor/Private/Framework/MVVMRowHelper.cpp b/Engine/Plugins/Runtime/ModelViewViewModel/Source/ModelViewViewModelEditor/Private/Framework/MVVMRowHelper.cpp
--- a/Engine/Plugins/Runtime/ModelViewViewModel/Source/ModelViewViewModelEditor/Private/Framework/MVVMRowHelper.cpp
+++ b/Engine/Plugins/Runtime/ModelViewViewModel/Source/ModelViewViewModelEditor/Private/Framework/MVVMRowHelper.cpp
@@ -1039,7 +1039,11 @@ FMenuBuilder FRowHelper::CreateContextMenu(UWidgetBlueprint* WidgetBlueprint, UM
 			MenuBuilder.AddMenuEntry(LOCTEXT("DuplicateBinding", "Duplicate"),
 				LOCTEXT("DuplicateBindingTooltip", "Duplicate bindings or events."),
 				FSlateIcon(FAppStyle::GetAppStyleSetName(), "Icons.Duplicate"),
-				DuplicateAction);
+				DuplicateAction,
+				NAME_None,
+				EUserInterfaceActionType::Button,
+				NAME_None,
+				LOCTEXT("CRTL-D", "CRTL-D"));
 
 			FUIAction CopyAction;
 			CopyAction.ExecuteAction = FExecuteAction::CreateStatic(&Private::HandleCopyEntries, WidgetBlueprint, View, CopiedEntries);
@@ -1047,15 +1051,25 @@ FMenuBuilder FRowHelper::CreateContextMenu(UWidgetBlueprint* WidgetBlueprint, UM
 			MenuBuilder.AddMenuEntry(LOCTEXT("CopyBinding", "Copy"),
 				LOCTEXT("CopyBindingTooltip", "Copy bindings or events."),
 				FSlateIcon(FAppStyle::GetAppStyleSetName(), "Icons.Copy"),
-				CopyAction);
+				CopyAction,
+				NAME_None,
+				EUserInterfaceActionType::Button,
+				NAME_None,
+				LOCTEXT("CRTL-C", "CRTL-C"));
 
 			FUIAction PasteAction;
 			PasteAction.ExecuteAction = FExecuteAction::CreateStatic(&Private::HandlePasteEntries, WidgetBlueprint, View, CopiedEntries);
 			PasteAction.CanExecuteAction = FCanExecuteAction::CreateLambda([bCanPasteEntry]() { return bCanPasteEntry; });
-			MenuBuilder.AddMenuEntry(LOCTEXT("PasteBinding", "Paste"),
+			MenuBuilder.AddMenuEntry(
+				LOCTEXT("PasteBinding", "Paste"),
 				LOCTEXT("PasteBindingTooltip", "Paste bindings or events."),
 				FSlateIcon(FAppStyle::GetAppStyleSetName(), "Icons.Paste"),
-				PasteAction);
+				PasteAction,
+				NAME_None,
+				EUserInterfaceActionType::Button,
+				NAME_None,
+				LOCTEXT("CRTL-V", "CRTL-V"));
+
 		}
 
 		{
diff --git a/Engine/Plugins/Runtime/ModelViewViewModel/Source/ModelViewViewModelEditor/Private/Widgets/SMVVMViewBindingListView.cpp b/Engine/Plugins/Runtime/ModelViewViewModel/Source/ModelViewViewModelEditor/Private/Widgets/SMVVMViewBindingListView.cpp
--- a/Engine/Plugins/Runtime/ModelViewViewModel/Source/ModelViewViewModelEditor/Private/Widgets/SMVVMViewBindingListView.cpp
+++ b/Engine/Plugins/Runtime/ModelViewViewModel/Source/ModelViewViewModelEditor/Private/Widgets/SMVVMViewBindingListView.cpp
@@ -824,24 +824,30 @@ FReply SBindingsList::OnKeyDown(const FGeometry& MyGeometry, const FKeyEvent& In
 
 	else if (InKeyEvent.GetModifierKeys().IsControlDown())
 	{
-			if (UMVVMWidgetBlueprintExtension_View* MVVMExtensionPtr = MVVMExtension.Get())
-			{
-				UMVVMBlueprintView* BlueprintView = MVVMExtensionPtr->GetBlueprintView();
-				UWidgetBlueprint* WidgetBlueprint = MVVMExtensionPtr->GetWidgetBlueprint();
-				TArray<TSharedPtr<FBindingEntry>> Selection = TreeView->GetSelectedItems();
+		if (UMVVMWidgetBlueprintExtension_View* MVVMExtensionPtr = MVVMExtension.Get())
+		{
+			UMVVMBlueprintView* BlueprintView = MVVMExtensionPtr->GetBlueprintView();
+			UWidgetBlueprint* WidgetBlueprint = MVVMExtensionPtr->GetWidgetBlueprint();
+			TArray<TSharedPtr<FBindingEntry>> Selection = TreeView->GetSelectedItems();
 
-				// Only support single selection for copy/paste behavior
-				if (InKeyEvent.GetKey() == EKeys::C)
-				{
-					BindingEntry::FRowHelper::CopyEntries(WidgetBlueprint, BlueprintView, Selection);
-				}
-				else if (InKeyEvent.GetKey() == EKeys::V)
-				{
-					BindingEntry::FRowHelper::PasteEntries(WidgetBlueprint, BlueprintView, Selection);
-				}
+			if (InKeyEvent.GetKey() == EKeys::C)
+			{
+				BindingEntry::FRowHelper::CopyEntries(WidgetBlueprint, BlueprintView, Selection);
+				return FReply::Handled();		
+			}
+			else if (InKeyEvent.GetKey() == EKeys::V)
+			{
+				BindingEntry::FRowHelper::PasteEntries(WidgetBlueprint, BlueprintView, Selection);
+				return FReply::Handled();		
 			}
-			return FReply::Handled();
-		
+			else if (InKeyEvent.GetKey() == EKeys::D)
+			{
+				TArray<const TSharedPtr<FBindingEntry>> NewSelection;
+				BindingEntry::FRowHelper::DuplicateEntries(WidgetBlueprint, BlueprintView, Selection, NewSelection);
+				SetSelection(NewSelection);
+				return FReply::Handled();		
+			}
+		}
 	}
 
 	return FReply::Unhandled();
