MVVM: Ensure opening Context Menu via row button will only retain selection for relevant binding row.

#rnx

diff --git a/Engine/Plugins/Runtime/ModelViewViewModel/Source/ModelViewViewModelEditor/Private/Framework/MVVMRowHelper.cpp b/Engine/Plugins/Runtime/ModelViewViewModel/Source/ModelViewViewModelEditor/Private/Framework/MVVMRowHelper.cpp
--- a/Engine/Plugins/Runtime/ModelViewViewModel/Source/ModelViewViewModelEditor/Private/Framework/MVVMRowHelper.cpp
+++ b/Engine/Plugins/Runtime/ModelViewViewModel/Source/ModelViewViewModelEditor/Private/Framework/MVVMRowHelper.cpp
@@ -697,6 +697,39 @@ void FRowHelper::PasteEntries(UWidgetBlueprint* WidgetBlueprint, UMVVMBlueprintV
 	}
 }
 
+bool FRowHelper::HasBlueprintGraph(UMVVMBlueprintView* BlueprintView, const TSharedPtr<FBindingEntry> Entry)
+{
+	if (!Entry.IsValid())
+	{
+		return false;
+	}
+
+	if (Entry->GetRowType() == FBindingEntry::ERowType::Binding)
+	{
+		if (FMVVMBlueprintViewBinding* ViewBinding = Entry->GetBinding(BlueprintView))
+		{
+			UMVVMBlueprintViewConversionFunction* ConversionFunctionA = ViewBinding->Conversion.GetConversionFunction(true);
+			UMVVMBlueprintViewConversionFunction* ConversionFunctionB = ViewBinding->Conversion.GetConversionFunction(false);
+			return ConversionFunctionA != nullptr || ConversionFunctionB != nullptr;
+		}
+	}
+	else if (Entry->GetRowType() == FBindingEntry::ERowType::Event)
+	{
+		if (UMVVMBlueprintViewEvent* Event = Entry->GetEvent())
+		{
+			return Event->GetWrapperGraph() != nullptr;
+		}
+	}
+	else if (Entry->GetRowType() == FBindingEntry::ERowType::Condition)
+	{
+		if (UMVVMBlueprintViewCondition* Condition = Entry->GetCondition())
+		{
+			return Condition->GetWrapperGraph() != nullptr;
+		}
+	}
+	return false;
+}
+
 void FRowHelper::ShowBlueprintGraph(FBlueprintEditor* BlueprintEditor, UWidgetBlueprint* WidgetBlueprint, UMVVMBlueprintView* BlueprintView, TArrayView<const TSharedPtr<FBindingEntry>> Entries)
 {
 	auto ShowGraph = [WidgetBlueprint, BlueprintEditor](UEdGraph* Graph)
@@ -1031,7 +1064,11 @@ FMenuBuilder FRowHelper::CreateContextMenu(UWidgetBlueprint* WidgetBlueprint, UM
 			MenuBuilder.AddMenuEntry(LOCTEXT("RemoveBinding", "Remove"),
 				LOCTEXT("RemoveBindingTooltip", "Remove bindings or events."),
 				FSlateIcon(FAppStyle::GetAppStyleSetName(), "Icons.Delete"),
-				RemoveAction);
+				RemoveAction,
+				NAME_None,
+				EUserInterfaceActionType::Button,
+				NAME_None,
+				LOCTEXT("DELETE", "DELETE"));
 
 			FUIAction DuplicateAction;
 			DuplicateAction.ExecuteAction = FExecuteAction::CreateStatic(&Private::HandleDuplicateEntries, WidgetBlueprint, View, CopiedEntries, OnSelectionChanged);
@@ -1043,7 +1080,7 @@ FMenuBuilder FRowHelper::CreateContextMenu(UWidgetBlueprint* WidgetBlueprint, UM
 				NAME_None,
 				EUserInterfaceActionType::Button,
 				NAME_None,
-				LOCTEXT("CRTL-D", "CRTL-D"));
+				LOCTEXT("CRTL+D", "CRTL+D"));
 
 			FUIAction CopyAction;
 			CopyAction.ExecuteAction = FExecuteAction::CreateStatic(&Private::HandleCopyEntries, WidgetBlueprint, View, CopiedEntries);
@@ -1055,7 +1092,7 @@ FMenuBuilder FRowHelper::CreateContextMenu(UWidgetBlueprint* WidgetBlueprint, UM
 				NAME_None,
 				EUserInterfaceActionType::Button,
 				NAME_None,
-				LOCTEXT("CRTL-C", "CRTL-C"));
+				LOCTEXT("CRTL+C", "CRTL+C"));
 
 			FUIAction PasteAction;
 			PasteAction.ExecuteAction = FExecuteAction::CreateStatic(&Private::HandlePasteEntries, WidgetBlueprint, View, CopiedEntries);
@@ -1068,7 +1105,7 @@ FMenuBuilder FRowHelper::CreateContextMenu(UWidgetBlueprint* WidgetBlueprint, UM
 				NAME_None,
 				EUserInterfaceActionType::Button,
 				NAME_None,
-				LOCTEXT("CRTL-V", "CRTL-V"));
+				LOCTEXT("CRTL+V", "CRTL+V"));
 
 		}
 
diff --git a/Engine/Plugins/Runtime/ModelViewViewModel/Source/ModelViewViewModelEditor/Private/Framework/MVVMRowHelper.h b/Engine/Plugins/Runtime/ModelViewViewModel/Source/ModelViewViewModelEditor/Private/Framework/MVVMRowHelper.h
--- a/Engine/Plugins/Runtime/ModelViewViewModel/Source/ModelViewViewModelEditor/Private/Framework/MVVMRowHelper.h
+++ b/Engine/Plugins/Runtime/ModelViewViewModel/Source/ModelViewViewModelEditor/Private/Framework/MVVMRowHelper.h
@@ -33,6 +33,7 @@ struct FRowHelper
 	static void CopyEntries(const UWidgetBlueprint* WidgetBlueprint, UMVVMBlueprintView* BlueprintView, const TArray<TSharedPtr<FBindingEntry>>& Entries);
 	static void PasteEntries(UWidgetBlueprint* WidgetBlueprint, UMVVMBlueprintView* BlueprintView, const TArray<TSharedPtr<FBindingEntry>>& Entries);
 
+	static bool HasBlueprintGraph(UMVVMBlueprintView* BlueprintView, const TSharedPtr<FBindingEntry> Entry);
 	static void ShowBlueprintGraph(FBlueprintEditor* Editor, UWidgetBlueprint* WidgetBlueprint, UMVVMBlueprintView* BlueprintView, TArrayView<const TSharedPtr<FBindingEntry>> Selection);
 
 	static TOptional<FMVVMBlueprintPropertyPath> DropFieldSelector(const UWidgetBlueprint* WidgetBlueprint, const FDragDropEvent& DragDropEvent, bool bIsSource);
diff --git a/Engine/Plugins/Runtime/ModelViewViewModel/Source/ModelViewViewModelEditor/Private/Widgets/BindingEntry/SMVVMBaseRow.cpp b/Engine/Plugins/Runtime/ModelViewViewModel/Source/ModelViewViewModelEditor/Private/Widgets/BindingEntry/SMVVMBaseRow.cpp
--- a/Engine/Plugins/Runtime/ModelViewViewModel/Source/ModelViewViewModelEditor/Private/Widgets/BindingEntry/SMVVMBaseRow.cpp
+++ b/Engine/Plugins/Runtime/ModelViewViewModel/Source/ModelViewViewModelEditor/Private/Widgets/BindingEntry/SMVVMBaseRow.cpp
@@ -1,19 +1,26 @@
 // Copyright Epic Games, Inc. All Rights Reserved.
 
 #include "Widgets/BindingEntry/SMVVMBaseRow.h"
+#include "Widgets/SMVVMViewBindingListView.h"
 
+#include "Framework/MVVMRowHelper.h"
+#include "MVVMBlueprintViewEvent.h"
+#include "MVVMDeveloperProjectSettings.h"
 #include "Styling/MVVMEditorStyle.h"
 
 #include "Widgets/SNullWidget.h"
 
+#define LOCTEXT_NAMESPACE "BindingListView_BaseRow"
+
 namespace UE::MVVM::BindingEntry
 {
 
-void SBaseRow::Construct(const FArguments& Args, const TSharedRef<STableViewBase>& OwnerTableView, const TSharedPtr<FWidgetBlueprintEditor>& InBlueprintEditor, UWidgetBlueprint* InBlueprint, const TSharedPtr<FBindingEntry>& InEntry)
+void SBaseRow::Construct(const FArguments& Args, SBindingsList* InBindingsList, const TSharedRef<STableViewBase>& OwnerTableView, const TSharedPtr<FWidgetBlueprintEditor>& InBlueprintEditor, UWidgetBlueprint* InBlueprint, const TSharedPtr<FBindingEntry>& InEntry)
 {
 	BindingEntry = InEntry;
 	WeakBlueprint = InBlueprint;
 	WeakBlueprintEditor = InBlueprintEditor;
+	BindingsList = InBindingsList;
 
 	STableRow<TSharedPtr<FBindingEntry>>::Construct(
 		STableRow<TSharedPtr<FBindingEntry>>::FArguments()
@@ -46,5 +53,59 @@ UMVVMEditorSubsystem* SBaseRow::GetEditorSubsystem() const
 	return GEditor->GetEditorSubsystem<UMVVMEditorSubsystem>();
 }
 
+TSharedRef<SWidget> SBaseRow::BuildContextMenuButton()
+{
+	return SNew(SComboButton)
+		.ContentPadding(FMargin(2.f, 0.0f))
+		.ComboButtonStyle(&FMVVMEditorStyle::Get().GetWidgetStyle<FComboButtonStyle>("NoStyleComboButton"))
+		.HasDownArrow(false)
+		.OnGetMenuContent(this, &SBaseRow::HandleContextMenu)
+		.ButtonContent()
+		[
+			SNew(SImage)
+			.Image(FMVVMEditorStyle::Get().GetBrush("Icon.Ellipsis"))
+			.DesiredSizeOverride(FVector2D(6.0, 24.0))
+		];
+}
+
+TSharedRef<SWidget> SBaseRow::HandleContextMenu() const
+{
+	TSharedPtr<FBindingEntry> Entry = GetEntry();
+
+	// Ensure only this entry is selected when opening context menu. When user has multiple rows selected, opening this menu will only be relevant to the selected entry
+	BindingsList->SetSelection(MakeArrayView(&Entry, 1));
+
+	FMenuBuilder MenuBuilder = BindingEntry::FRowHelper::CreateContextMenu(GetBlueprint(), GetBlueprintView(), MakeArrayView(&Entry, 1));
+	
+	{
+		MenuBuilder.BeginSection("Developer", LOCTEXT("Developer", "Developer"));
+
+		if (GetDefault<UMVVMDeveloperProjectSettings>()->bShowDeveloperGenerateGraphSettings)
+		{
+			bool bCanShowGraph = BindingEntry::FRowHelper::HasBlueprintGraph(GetBlueprintView(), Entry);
+
+			FUIAction ShowGraphAction;
+			ShowGraphAction.ExecuteAction = FExecuteAction::CreateSP(this, &SBaseRow::HandleShowBlueprintGraph);
+			ShowGraphAction.CanExecuteAction = FCanExecuteAction::CreateLambda([bCanShowGraph]() { return bCanShowGraph; });
+			MenuBuilder.AddMenuEntry(LOCTEXT("ShowGraph", "Show binding graph")
+				, LOCTEXT("ShowGraphTooltip", "Show the Blueprint graph that represent the binding."
+					" The graph is always generated but may not be visible to the user.")
+				, FSlateIcon(FAppStyle::GetAppStyleSetName(), "BlueprintEditor.FindInBlueprints.MenuIcon")
+				, ShowGraphAction);
+		}
+
+		MenuBuilder.EndSection();
+	}
+
+	return MenuBuilder.MakeWidget();
+}
+
+void SBaseRow::HandleShowBlueprintGraph() const
+{
+	TSharedPtr<FBindingEntry> Entry = GetEntry();
+	BindingEntry::FRowHelper::ShowBlueprintGraph(GetBlueprintEditor().Get(), GetBlueprint(), GetBlueprintView(), MakeArrayView(&Entry, 1));
+}
 
 } // namespace UE::MVVM::BindingEntry
+
+#undef LOCTEXT_NAMESPACE
diff --git a/Engine/Plugins/Runtime/ModelViewViewModel/Source/ModelViewViewModelEditor/Private/Widgets/BindingEntry/SMVVMBaseRow.h b/Engine/Plugins/Runtime/ModelViewViewModel/Source/ModelViewViewModelEditor/Private/Widgets/BindingEntry/SMVVMBaseRow.h
--- a/Engine/Plugins/Runtime/ModelViewViewModel/Source/ModelViewViewModelEditor/Private/Widgets/BindingEntry/SMVVMBaseRow.h
+++ b/Engine/Plugins/Runtime/ModelViewViewModel/Source/ModelViewViewModelEditor/Private/Widgets/BindingEntry/SMVVMBaseRow.h
@@ -11,6 +11,11 @@
 #include "WidgetBlueprint.h"
 #include "WidgetBlueprintEditor.h"
 
+namespace UE::MVVM
+{
+class SBindingsList;
+}
+
 namespace UE::MVVM::BindingEntry
 {
 
@@ -23,7 +28,7 @@ public:
 	SLATE_BEGIN_ARGS(SBaseRow) {}
 	SLATE_END_ARGS()
 
-	void Construct(const FArguments& Args, const TSharedRef<STableViewBase>& OwnerTableView, const TSharedPtr<FWidgetBlueprintEditor>& InBlueprintEditor, UWidgetBlueprint* InBlueprint, const TSharedPtr<FBindingEntry>& InEntry);
+	void Construct(const FArguments& Args, SBindingsList* BindingsList, const TSharedRef<STableViewBase>& OwnerTableView, const TSharedPtr<FWidgetBlueprintEditor>& InBlueprintEditor, UWidgetBlueprint* InBlueprint, const TSharedPtr<FBindingEntry>& InEntry);
 
 	TSharedPtr<FBindingEntry> GetEntry() const
 	{
@@ -42,11 +47,16 @@ public:
 protected:
 	virtual TSharedRef<SWidget> BuildRowWidget() = 0;
 	virtual const ANSICHAR* GetTableRowStyle() const = 0;
+	virtual TSharedRef<SWidget> BuildContextMenuButton();
 
 private:
+	TSharedRef<SWidget> HandleContextMenu() const;
+	void HandleShowBlueprintGraph() const;
+
 	TSharedPtr<FBindingEntry> BindingEntry;
 	TWeakObjectPtr<UWidgetBlueprint> WeakBlueprint;
 	TWeakPtr<FWidgetBlueprintEditor> WeakBlueprintEditor;
+	SBindingsList* BindingsList;
 };
 
 } // namespace UE::MVVM
diff --git a/Engine/Plugins/Runtime/ModelViewViewModel/Source/ModelViewViewModelEditor/Private/Widgets/BindingEntry/SMVVMBindingRow.cpp b/Engine/Plugins/Runtime/ModelViewViewModel/Source/ModelViewViewModelEditor/Private/Widgets/BindingEntry/SMVVMBindingRow.cpp
--- a/Engine/Plugins/Runtime/ModelViewViewModel/Source/ModelViewViewModelEditor/Private/Widgets/BindingEntry/SMVVMBindingRow.cpp
+++ b/Engine/Plugins/Runtime/ModelViewViewModel/Source/ModelViewViewModelEditor/Private/Widgets/BindingEntry/SMVVMBindingRow.cpp
@@ -57,13 +57,13 @@ namespace Private
 
 } // namespace
 
-void SBindingRow::Construct(const FArguments& Args, const TSharedRef<STableViewBase>& OwnerTableView, const TSharedPtr<FWidgetBlueprintEditor>& InBlueprintEditor, UWidgetBlueprint* InBlueprint, const TSharedPtr<FBindingEntry>& InEntry)
+void SBindingRow::Construct(const FArguments& Args, SBindingsList* InBindingsList, const TSharedRef<STableViewBase>& OwnerTableView, const TSharedPtr<FWidgetBlueprintEditor>& InBlueprintEditor, UWidgetBlueprint* InBlueprint, const TSharedPtr<FBindingEntry>& InEntry)
 {
 	static IConsoleVariable* CVarDefaultExecutionMode = IConsoleManager::Get().FindConsoleVariable(TEXT("MVVM.DefaultExecutionMode"));
 	ensure(CVarDefaultExecutionMode);
 	DefaultExecutionMode = CVarDefaultExecutionMode ? (EMVVMExecutionMode)CVarDefaultExecutionMode->GetInt() : EMVVMExecutionMode::DelayedWhenSharedElseImmediate;
 
-	SBaseRow::Construct(SBaseRow::FArguments(), OwnerTableView, InBlueprintEditor, InBlueprint, InEntry);
+	SBaseRow::Construct(SBaseRow::FArguments(), InBindingsList, OwnerTableView, InBlueprintEditor, InBlueprint, InEntry);
 
 	TSharedPtr<SWidget> ChildContent = ChildSlot.DetachWidget();
 	ChildSlot
@@ -229,17 +229,7 @@ TSharedRef<SWidget> SBindingRow::BuildRowWidget()
 			.HAlign(HAlign_Left)
 			.AutoWidth()
 			[
-				SNew(SComboButton)
-				.ContentPadding(FMargin(2.f, 0.0f))
-				.ComboButtonStyle(&FMVVMEditorStyle::Get().GetWidgetStyle<FComboButtonStyle>("NoStyleComboButton"))
-				.HasDownArrow(false)
-				.OnGetMenuContent(this, &SBindingRow::HandleContextMenu)
-				.ButtonContent()
-				[
-					SNew(SImage)
-					.Image(FMVVMEditorStyle::Get().GetBrush("Icon.Ellipsis"))
-					.DesiredSizeOverride(FVector2D(6.0, 24.0))
-				]
+				BuildContextMenuButton()
 			]
 		]
 	];
@@ -764,41 +754,6 @@ void SBindingRow::OnBindingModeSelectionChanged(FName ValueName, ESelectInfo::Ty
 	}
 }
 
-TSharedRef<SWidget> SBindingRow::HandleContextMenu() const
-{
-	TSharedPtr<FBindingEntry> Entry = GetEntry();
-	FMenuBuilder MenuBuilder = BindingEntry::FRowHelper::CreateContextMenu(GetBlueprint(), GetBlueprintView(), MakeArrayView(&Entry, 1));
-
-	{
-		MenuBuilder.BeginSection("Developer", LOCTEXT("Developer", "Developer"));
-
-		FMVVMBlueprintViewBinding* ViewBinding = GetThisViewBinding();
-		if (GetDefault<UMVVMDeveloperProjectSettings>()->bShowDeveloperGenerateGraphSettings)
-		{
-			bool bCanShowGraph = ViewBinding && (ViewBinding->Conversion.GetConversionFunction(true) != nullptr || ViewBinding->Conversion.GetConversionFunction(false) != nullptr);
-
-			FUIAction ShowGraphAction;
-			ShowGraphAction.ExecuteAction = FExecuteAction::CreateSP(this, &SBindingRow::HandleShowBlueprintGraph);
-			ShowGraphAction.CanExecuteAction = FCanExecuteAction::CreateLambda([bCanShowGraph]() { return bCanShowGraph; });
-			MenuBuilder.AddMenuEntry(LOCTEXT("ShowGraph", "Show binding graph")
-				, LOCTEXT("ShowGraphTooltip", "Show the Blueprint graph that represent the binding."
-					" The graph is always generated but may not be visible to the user.")
-				, FSlateIcon(FAppStyle::GetAppStyleSetName(), "BlueprintEditor.FindInBlueprints.MenuIcon")
-				, ShowGraphAction);
-		}
-
-		MenuBuilder.EndSection();
-	}
-
-	return MenuBuilder.MakeWidget();
-}
-
-void SBindingRow::HandleShowBlueprintGraph() const
-{
-	TSharedPtr<FBindingEntry> Entry = GetEntry();
-	BindingEntry::FRowHelper::ShowBlueprintGraph(GetBlueprintEditor().Get(), GetBlueprint(), GetBlueprintView(), MakeArrayView(&Entry, 1));
-}
-
 } // namespace
 
 #undef LOCTEXT_NAMESPACE
diff --git a/Engine/Plugins/Runtime/ModelViewViewModel/Source/ModelViewViewModelEditor/Private/Widgets/BindingEntry/SMVVMBindingRow.h b/Engine/Plugins/Runtime/ModelViewViewModel/Source/ModelViewViewModelEditor/Private/Widgets/BindingEntry/SMVVMBindingRow.h
--- a/Engine/Plugins/Runtime/ModelViewViewModel/Source/ModelViewViewModelEditor/Private/Widgets/BindingEntry/SMVVMBindingRow.h
+++ b/Engine/Plugins/Runtime/ModelViewViewModel/Source/ModelViewViewModelEditor/Private/Widgets/BindingEntry/SMVVMBindingRow.h
@@ -10,6 +10,10 @@
 
 struct FMVVMBlueprintViewBinding;
 
+namespace UE::MVVM
+{
+class SBindingsList;
+}
 
 namespace UE::MVVM::BindingEntry
 {
@@ -23,7 +27,7 @@ public:
 	SLATE_BEGIN_ARGS(SBindingRow) {}
 	SLATE_END_ARGS()
 
-	void Construct(const FArguments& Args, const TSharedRef<STableViewBase>& OwnerTableView, const TSharedPtr<FWidgetBlueprintEditor>& InBlueprintEditor, UWidgetBlueprint* InBlueprint, const TSharedPtr<FBindingEntry>& InEntry);
+	void Construct(const FArguments& Args, SBindingsList* BindingsList, const TSharedRef<STableViewBase>& OwnerTableView, const TSharedPtr<FWidgetBlueprintEditor>& InBlueprintEditor, UWidgetBlueprint* InBlueprint, const TSharedPtr<FBindingEntry>& InEntry);
 
 protected:
 	virtual TSharedRef<SWidget> BuildRowWidget() override;
@@ -91,9 +95,6 @@ private:
 
 	void OnBindingModeSelectionChanged(FName ValueName, ESelectInfo::Type);
 
-	TSharedRef<SWidget> HandleContextMenu() const;
-	void HandleShowBlueprintGraph() const;
-
 private:
 	TArray<TSharedPtr<FText>> ErrorItems;
 	EMVVMExecutionMode DefaultExecutionMode = EMVVMExecutionMode::DelayedWhenSharedElseImmediate;
diff --git a/Engine/Plugins/Runtime/ModelViewViewModel/Source/ModelViewViewModelEditor/Private/Widgets/BindingEntry/SMVVMConditionRow.cpp b/Engine/Plugins/Runtime/ModelViewViewModel/Source/ModelViewViewModelEditor/Private/Widgets/BindingEntry/SMVVMConditionRow.cpp
--- a/Engine/Plugins/Runtime/ModelViewViewModel/Source/ModelViewViewModelEditor/Private/Widgets/BindingEntry/SMVVMConditionRow.cpp
+++ b/Engine/Plugins/Runtime/ModelViewViewModel/Source/ModelViewViewModelEditor/Private/Widgets/BindingEntry/SMVVMConditionRow.cpp
@@ -50,9 +50,9 @@ namespace UE::MVVM::BindingEntry
 		}
 	}
 
-void SConditionRow::Construct(const FArguments& Args, const TSharedRef<STableViewBase>& OwnerTableView, const TSharedPtr<FWidgetBlueprintEditor>& InBlueprintEditor, UWidgetBlueprint* InBlueprint, const TSharedPtr<FBindingEntry>& InEntry)
+void SConditionRow::Construct(const FArguments& Args, SBindingsList* InBindingsList, const TSharedRef<STableViewBase>& OwnerTableView, const TSharedPtr<FWidgetBlueprintEditor>& InBlueprintEditor, UWidgetBlueprint* InBlueprint, const TSharedPtr<FBindingEntry>& InEntry)
 {
-	SBaseRow::Construct(SBaseRow::FArguments(), OwnerTableView, InBlueprintEditor, InBlueprint, InEntry);
+	SBaseRow::Construct(SBaseRow::FArguments(), InBindingsList, OwnerTableView, InBlueprintEditor, InBlueprint, InEntry);
 
 	TSharedPtr<SWidget> ChildContent = ChildSlot.DetachWidget();
 	ChildSlot
@@ -253,17 +253,7 @@ TSharedRef<SWidget> SConditionRow::BuildRowWidget()
 			.HAlign(HAlign_Left)
 			.AutoWidth()
 			[
-				SNew(SComboButton)
-				.ContentPadding(FMargin(2.f, 0.0f))
-				.ComboButtonStyle(&FMVVMEditorStyle::Get().GetWidgetStyle<FComboButtonStyle>("NoStyleComboButton"))
-				.HasDownArrow(false)
-				.OnGetMenuContent(this, &SConditionRow::HandleContextMenu)
-				.ButtonContent()
-				[
-					SNew(SImage)
-					.Image(FMVVMEditorStyle::Get().GetBrush("Icon.Ellipsis"))
-					.DesiredSizeOverride(FVector2D(6.0, 24.0))
-				]
+				BuildContextMenuButton()
 			]
 		]
 	];
@@ -639,42 +629,6 @@ FReply SConditionRow::HandleFieldSelectorOver(const FGeometry& MyGeometry, const
 	return BindingEntry::FRowHelper::DragOverFieldSelector(GetBlueprint(), DragDropEvent, bCondition);
 }
 
-TSharedRef<SWidget> SConditionRow::HandleContextMenu() const
-{
-	TSharedPtr<FBindingEntry> Entry = GetEntry();
-	FMenuBuilder MenuBuilder = BindingEntry::FRowHelper::CreateContextMenu(GetBlueprint(), GetBlueprintView(), MakeArrayView(&Entry, 1));
-	
-	{
-		MenuBuilder.BeginSection("Developer", LOCTEXT("Developer", "Developer"));
-
-		UMVVMBlueprintViewCondition* Condition = GetCondition();
-		if (GetDefault<UMVVMDeveloperProjectSettings>()->bShowDeveloperGenerateGraphSettings)
-		{
-			bool bCanShowGraph = Condition && Condition->GetWrapperGraph();
-
-			FUIAction ShowGraphAction;
-			ShowGraphAction.ExecuteAction = FExecuteAction::CreateSP(this, &SConditionRow::HandleShowBlueprintGraph);
-			ShowGraphAction.CanExecuteAction = FCanExecuteAction::CreateLambda([bCanShowGraph]() { return bCanShowGraph; });
-			MenuBuilder.AddMenuEntry(LOCTEXT("ShowGraph", "Show event graph")
-				, LOCTEXT("ShowGraphTooltip", "Show the Blueprint graph that represent the event."
-					" The graph is always generated but may not be visible to the user.")
-				, FSlateIcon(FAppStyle::GetAppStyleSetName(), "BlueprintEditor.FindInBlueprints.MenuIcon")
-				, ShowGraphAction);
-		}
-
-		MenuBuilder.EndSection();
-	}
-
-	return MenuBuilder.MakeWidget();
-}
-
-
-void SConditionRow::HandleShowBlueprintGraph() const
-{
-	TSharedPtr<FBindingEntry> Entry = GetEntry();
-	BindingEntry::FRowHelper::ShowBlueprintGraph(GetBlueprintEditor().Get(), GetBlueprint(), GetBlueprintView(), MakeArrayView(&Entry, 1));
-}
-
 } // namespace
 
 #undef LOCTEXT_NAMESPACE
diff --git a/Engine/Plugins/Runtime/ModelViewViewModel/Source/ModelViewViewModelEditor/Private/Widgets/BindingEntry/SMVVMConditionRow.h b/Engine/Plugins/Runtime/ModelViewViewModel/Source/ModelViewViewModelEditor/Private/Widgets/BindingEntry/SMVVMConditionRow.h
--- a/Engine/Plugins/Runtime/ModelViewViewModel/Source/ModelViewViewModelEditor/Private/Widgets/BindingEntry/SMVVMConditionRow.h
+++ b/Engine/Plugins/Runtime/ModelViewViewModel/Source/ModelViewViewModelEditor/Private/Widgets/BindingEntry/SMVVMConditionRow.h
@@ -8,6 +8,10 @@
 #include "Widgets/BindingEntry/SMVVMBaseRow.h"
 #include "Widgets/SMVVMFieldSelectorMenu.h"
 
+namespace UE::MVVM
+{
+class SBindingsList;
+}
 
 namespace UE::MVVM::BindingEntry
 {
@@ -21,7 +25,7 @@ public:
 	SLATE_BEGIN_ARGS(SConditionRow) {}
 	SLATE_END_ARGS()
 
-	void Construct(const FArguments& Args, const TSharedRef<STableViewBase>& OwnerTableView, const TSharedPtr<FWidgetBlueprintEditor>& InBlueprintEditor, UWidgetBlueprint* InBlueprint, const TSharedPtr<FBindingEntry>& InEntry);
+	void Construct(const FArguments& Args, SBindingsList* BindingsList, const TSharedRef<STableViewBase>& OwnerTableView, const TSharedPtr<FWidgetBlueprintEditor>& InBlueprintEditor, UWidgetBlueprint* InBlueprint, const TSharedPtr<FBindingEntry>& InEntry);
 
 protected:
 	virtual TSharedRef<SWidget> BuildRowWidget() override;
@@ -61,9 +65,6 @@ private:
 
 	FReply HandleFieldSelectorOver(const FGeometry& MyGeometry, const FDragDropEvent& DragDropEvent, bool bCondition);
 
-	TSharedRef<SWidget> HandleContextMenu() const;
-	void HandleShowBlueprintGraph() const;
-
 	void OnConditionOperationSelectionChanged(FName ValueName, ESelectInfo::Type);
 	FText GetCurrentConditionOperationLabel() const;
 	TSharedRef<SWidget> GenerateConditionOperationWidget(FName ValueName) const;
diff --git a/Engine/Plugins/Runtime/ModelViewViewModel/Source/ModelViewViewModelEditor/Private/Widgets/BindingEntry/SMVVMEventRow.cpp b/Engine/Plugins/Runtime/ModelViewViewModel/Source/ModelViewViewModelEditor/Private/Widgets/BindingEntry/SMVVMEventRow.cpp
--- a/Engine/Plugins/Runtime/ModelViewViewModel/Source/ModelViewViewModelEditor/Private/Widgets/BindingEntry/SMVVMEventRow.cpp
+++ b/Engine/Plugins/Runtime/ModelViewViewModel/Source/ModelViewViewModelEditor/Private/Widgets/BindingEntry/SMVVMEventRow.cpp
@@ -28,9 +28,9 @@
 namespace UE::MVVM::BindingEntry
 {
 
-void SEventRow::Construct(const FArguments& Args, const TSharedRef<STableViewBase>& OwnerTableView, const TSharedPtr<FWidgetBlueprintEditor>& InBlueprintEditor, UWidgetBlueprint* InBlueprint, const TSharedPtr<FBindingEntry>& InEntry)
+void SEventRow::Construct(const FArguments& Args, SBindingsList* InBindingsList, const TSharedRef<STableViewBase>& OwnerTableView, const TSharedPtr<FWidgetBlueprintEditor>& InBlueprintEditor, UWidgetBlueprint* InBlueprint, const TSharedPtr<FBindingEntry>& InEntry)
 {
-	SBaseRow::Construct(SBaseRow::FArguments(), OwnerTableView, InBlueprintEditor, InBlueprint, InEntry);
+	SBaseRow::Construct(SBaseRow::FArguments(), InBindingsList, OwnerTableView, InBlueprintEditor, InBlueprint, InEntry);
 
 	TSharedPtr<SWidget> ChildContent = ChildSlot.DetachWidget();
 	ChildSlot
@@ -153,17 +153,7 @@ TSharedRef<SWidget> SEventRow::BuildRowWidget()
 			.HAlign(HAlign_Left)
 			.AutoWidth()
 			[
-				SNew(SComboButton)
-				.ContentPadding(FMargin(2.f, 0.0f))
-				.ComboButtonStyle(&FMVVMEditorStyle::Get().GetWidgetStyle<FComboButtonStyle>("NoStyleComboButton"))
-				.HasDownArrow(false)
-				.OnGetMenuContent(this, &SEventRow::HandleContextMenu)
-				.ButtonContent()
-				[
-					SNew(SImage)
-					.Image(FMVVMEditorStyle::Get().GetBrush("Icon.Ellipsis"))
-					.DesiredSizeOverride(FVector2D(6.0, 24.0))
-				]
+				BuildContextMenuButton()
 			]
 		]
 	];
@@ -393,41 +383,6 @@ FReply SEventRow::HandleFieldSelectorDragOver(const FGeometry& MyGeometry, const
 	return BindingEntry::FRowHelper::DragOverFieldSelector(GetBlueprint(), DragDropEvent, bEvent);
 }
 
-TSharedRef<SWidget> SEventRow::HandleContextMenu() const
-{
-	TSharedPtr<FBindingEntry> Entry = GetEntry();
-	FMenuBuilder MenuBuilder = BindingEntry::FRowHelper::CreateContextMenu(GetBlueprint(), GetBlueprintView(), MakeArrayView(&Entry, 1));
-	
-	{
-		MenuBuilder.BeginSection("Developer", LOCTEXT("Developer", "Developer"));
-
-		UMVVMBlueprintViewEvent* Event = GetEvent();
-		if (GetDefault<UMVVMDeveloperProjectSettings>()->bShowDeveloperGenerateGraphSettings)
-		{
-			bool bCanShowGraph = Event && Event->GetWrapperGraph();
-
-			FUIAction ShowGraphAction;
-			ShowGraphAction.ExecuteAction = FExecuteAction::CreateSP(this, &SEventRow::HandleShowBlueprintGraph);
-			ShowGraphAction.CanExecuteAction = FCanExecuteAction::CreateLambda([bCanShowGraph]() { return bCanShowGraph; });
-			MenuBuilder.AddMenuEntry(LOCTEXT("ShowGraph", "Show event graph")
-				, LOCTEXT("ShowGraphTooltip", "Show the Blueprint graph that represent the event."
-					" The graph is always generated but may not be visible to the user.")
-				, FSlateIcon(FAppStyle::GetAppStyleSetName(), "BlueprintEditor.FindInBlueprints.MenuIcon")
-				, ShowGraphAction);
-		}
-
-		MenuBuilder.EndSection();
-	}
-
-	return MenuBuilder.MakeWidget();
-}
-
-void SEventRow::HandleShowBlueprintGraph() const
-{
-	TSharedPtr<FBindingEntry> Entry = GetEntry();
-	BindingEntry::FRowHelper::ShowBlueprintGraph(GetBlueprintEditor().Get(), GetBlueprint(), GetBlueprintView(), MakeArrayView(&Entry, 1));
-}
-
 } // namespace
 
 #undef LOCTEXT_NAMESPACE
diff --git a/Engine/Plugins/Runtime/ModelViewViewModel/Source/ModelViewViewModelEditor/Private/Widgets/BindingEntry/SMVVMEventRow.h b/Engine/Plugins/Runtime/ModelViewViewModel/Source/ModelViewViewModelEditor/Private/Widgets/BindingEntry/SMVVMEventRow.h
--- a/Engine/Plugins/Runtime/ModelViewViewModel/Source/ModelViewViewModelEditor/Private/Widgets/BindingEntry/SMVVMEventRow.h
+++ b/Engine/Plugins/Runtime/ModelViewViewModel/Source/ModelViewViewModelEditor/Private/Widgets/BindingEntry/SMVVMEventRow.h
@@ -8,6 +8,10 @@
 #include "Widgets/BindingEntry/SMVVMBaseRow.h"
 #include "Widgets/SMVVMFieldSelectorMenu.h"
 
+namespace UE::MVVM
+{
+class SBindingsList;
+}
 
 namespace UE::MVVM::BindingEntry
 {
@@ -21,7 +25,7 @@ public:
 	SLATE_BEGIN_ARGS(SEventRow) {}
 	SLATE_END_ARGS()
 
-	void Construct(const FArguments& Args, const TSharedRef<STableViewBase>& OwnerTableView, const TSharedPtr<FWidgetBlueprintEditor>& InBlueprintEditor, UWidgetBlueprint* InBlueprint, const TSharedPtr<FBindingEntry>& InEntry);
+	void Construct(const FArguments& Args, SBindingsList* BindingsList, const TSharedRef<STableViewBase>& OwnerTableView, const TSharedPtr<FWidgetBlueprintEditor>& InBlueprintEditor, UWidgetBlueprint* InBlueprint, const TSharedPtr<FBindingEntry>& InEntry);
 
 protected:
 	virtual TSharedRef<SWidget> BuildRowWidget() override;
@@ -61,9 +65,6 @@ private:
 
 	FReply HandleFieldSelectorDragOver(const FGeometry& MyGeometry, const FDragDropEvent& DragDropEvent, bool bEvent);
 
-	TSharedRef<SWidget> HandleContextMenu() const;
-	void HandleShowBlueprintGraph() const;
-
 private:
 	TArray<TSharedPtr<FText>> ErrorItems;
 };
diff --git a/Engine/Plugins/Runtime/ModelViewViewModel/Source/ModelViewViewModelEditor/Private/Widgets/BindingEntry/SMVVMGroupRow.cpp b/Engine/Plugins/Runtime/ModelViewViewModel/Source/ModelViewViewModelEditor/Private/Widgets/BindingEntry/SMVVMGroupRow.cpp
--- a/Engine/Plugins/Runtime/ModelViewViewModel/Source/ModelViewViewModelEditor/Private/Widgets/BindingEntry/SMVVMGroupRow.cpp
+++ b/Engine/Plugins/Runtime/ModelViewViewModel/Source/ModelViewViewModelEditor/Private/Widgets/BindingEntry/SMVVMGroupRow.cpp
@@ -16,9 +16,9 @@
 namespace UE::MVVM::BindingEntry
 {
 
-void SGroupRow::Construct(const FArguments& Args, const TSharedRef<STableViewBase>& OwnerTableView, const TSharedPtr<FWidgetBlueprintEditor>& InBlueprintEditor, UWidgetBlueprint* InBlueprint, const TSharedPtr<FBindingEntry>& InEntry)
+void SGroupRow::Construct(const FArguments& Args, SBindingsList* InBindingsList, const TSharedRef<STableViewBase>& OwnerTableView, const TSharedPtr<FWidgetBlueprintEditor>& InBlueprintEditor, UWidgetBlueprint* InBlueprint, const TSharedPtr<FBindingEntry>& InEntry)
 {
-	SBaseRow::Construct(SBaseRow::FArguments(), OwnerTableView, InBlueprintEditor, InBlueprint, InEntry);
+	SBaseRow::Construct(SBaseRow::FArguments(), InBindingsList, OwnerTableView, InBlueprintEditor, InBlueprint, InEntry);
 
 	TSharedPtr<SWidget> ChildContent = ChildSlot.DetachWidget();
 	ChildSlot
diff --git a/Engine/Plugins/Runtime/ModelViewViewModel/Source/ModelViewViewModelEditor/Private/Widgets/BindingEntry/SMVVMGroupRow.h b/Engine/Plugins/Runtime/ModelViewViewModel/Source/ModelViewViewModelEditor/Private/Widgets/BindingEntry/SMVVMGroupRow.h
--- a/Engine/Plugins/Runtime/ModelViewViewModel/Source/ModelViewViewModelEditor/Private/Widgets/BindingEntry/SMVVMGroupRow.h
+++ b/Engine/Plugins/Runtime/ModelViewViewModel/Source/ModelViewViewModelEditor/Private/Widgets/BindingEntry/SMVVMGroupRow.h
@@ -6,6 +6,12 @@
 #include "Types/MVVMBindingSource.h"
 #include "Widgets/BindingEntry/SMVVMBaseRow.h"
 
+
+namespace UE::MVVM
+{
+class SBindingsList;
+}
+
 namespace UE::MVVM::BindingEntry
 {
 
@@ -18,7 +24,7 @@ public:
 	SLATE_BEGIN_ARGS(SGroupRow) {}
 	SLATE_END_ARGS()
 
-	void Construct(const FArguments& Args, const TSharedRef<STableViewBase>& OwnerTableView, const TSharedPtr<FWidgetBlueprintEditor>& InBlueprintEditor, UWidgetBlueprint* InBlueprint, const TSharedPtr<FBindingEntry>& InEntry);
+	void Construct(const FArguments& Args, SBindingsList* BindingsList, const TSharedRef<STableViewBase>& OwnerTableView, const TSharedPtr<FWidgetBlueprintEditor>& InBlueprintEditor, UWidgetBlueprint* InBlueprint, const TSharedPtr<FBindingEntry>& InEntry);
 
 protected:
 	virtual TSharedRef<SWidget> BuildRowWidget() override;
diff --git a/Engine/Plugins/Runtime/ModelViewViewModel/Source/ModelViewViewModelEditor/Private/Widgets/SMVVMViewBindingListView.cpp b/Engine/Plugins/Runtime/ModelViewViewModel/Source/ModelViewViewModelEditor/Private/Widgets/SMVVMViewBindingListView.cpp
--- a/Engine/Plugins/Runtime/ModelViewViewModel/Source/ModelViewViewModelEditor/Private/Widgets/SMVVMViewBindingListView.cpp
+++ b/Engine/Plugins/Runtime/ModelViewViewModel/Source/ModelViewViewModelEditor/Private/Widgets/SMVVMViewBindingListView.cpp
@@ -674,12 +674,12 @@ TSharedRef<ITableRow> SBindingsList::GenerateEntryRow(TSharedPtr<FBindingEntry>
 		{
 			case FBindingEntry::ERowType::Group:
 			{
-				Row = SNew(UE::MVVM::BindingEntry::SGroupRow, OwnerTable, WeakBlueprintEditor.Pin(), MVVMExtensionPtr->GetWidgetBlueprint(), Entry);
+				Row = SNew(UE::MVVM::BindingEntry::SGroupRow, this, OwnerTable, WeakBlueprintEditor.Pin(), MVVMExtensionPtr->GetWidgetBlueprint(), Entry);
 				break;
 			}
 			case FBindingEntry::ERowType::Binding:
 			{
-				Row = SNew(UE::MVVM::BindingEntry::SBindingRow, OwnerTable, WeakBlueprintEditor.Pin(), MVVMExtensionPtr->GetWidgetBlueprint(), Entry);
+				Row = SNew(UE::MVVM::BindingEntry::SBindingRow, this, OwnerTable, WeakBlueprintEditor.Pin(), MVVMExtensionPtr->GetWidgetBlueprint(), Entry);
 				break;
 			}
 			case FBindingEntry::ERowType::BindingParameter:
@@ -691,17 +691,17 @@ TSharedRef<ITableRow> SBindingsList::GenerateEntryRow(TSharedPtr<FBindingEntry>
 					ensureMsgf(false, TEXT("Corrupted Binding Parameter."));
 					return SNew(STableRow<TSharedPtr<FBindingEntry>>, OwnerTable);
 				}
-				Row = SNew(UE::MVVM::BindingEntry::SFunctionParameterRow, OwnerTable, WeakBlueprintEditor.Pin(), MVVMExtensionPtr->GetWidgetBlueprint(), Entry);
+				Row = SNew(UE::MVVM::BindingEntry::SFunctionParameterRow, this, OwnerTable, WeakBlueprintEditor.Pin(), MVVMExtensionPtr->GetWidgetBlueprint(), Entry);
 				break;
 			}
 			case FBindingEntry::ERowType::Event:
 			{
-				Row = SNew(UE::MVVM::BindingEntry::SEventRow, OwnerTable, WeakBlueprintEditor.Pin(), MVVMExtensionPtr->GetWidgetBlueprint(), Entry);
+				Row = SNew(UE::MVVM::BindingEntry::SEventRow, this, OwnerTable, WeakBlueprintEditor.Pin(), MVVMExtensionPtr->GetWidgetBlueprint(), Entry);
 				break;
 			}
 			case FBindingEntry::ERowType::Condition:
 			{
-				Row = SNew(UE::MVVM::BindingEntry::SConditionRow, OwnerTable, WeakBlueprintEditor.Pin(), MVVMExtensionPtr->GetWidgetBlueprint(), Entry);
+				Row = SNew(UE::MVVM::BindingEntry::SConditionRow, this, OwnerTable, WeakBlueprintEditor.Pin(), MVVMExtensionPtr->GetWidgetBlueprint(), Entry);
 				break;
 			}
 		}
