Expose world context pin on MVVMViewModelBase. This allows blueprints on a Viewmodel to access world subsystems (assuming they have an object they can use to resolve the world context)

diff --git a/Engine/Plugins/Runtime/ModelViewViewModel/Source/ModelViewViewModel/Public/MVVMViewModelBase.h b/Engine/Plugins/Runtime/ModelViewViewModel/Source/ModelViewViewModel/Public/MVVMViewModelBase.h
--- a/Engine/Plugins/Runtime/ModelViewViewModel/Source/ModelViewViewModel/Public/MVVMViewModelBase.h
+++ b/Engine/Plugins/Runtime/ModelViewViewModel/Source/ModelViewViewModel/Public/MVVMViewModelBase.h
@@ -38,7 +38,7 @@ namespace UMMVVMViewModelBase
 
 
 /** Base class for MVVM viewmodel. */
-UCLASS(MinimalAPI, Blueprintable, Abstract, DisplayName="MVVM Base Viewmodel")
+UCLASS(MinimalAPI, Blueprintable, Abstract, DisplayName="MVVM Base Viewmodel", meta=(ShowWorldContextPin="true"))
 class UMVVMViewModelBase : public UObject, public INotifyFieldValueChanged
 {
 	GENERATED_BODY()
