Fix Wrapper Graph Duplication

diff --git a/Engine/Plugins/Runtime/ModelViewViewModel/Source/ModelViewViewModelBlueprint/Private/MVVMBlueprintView.cpp b/Engine/Plugins/Runtime/ModelViewViewModel/Source/ModelViewViewModelBlueprint/Private/MVVMBlueprintView.cpp
--- a/Engine/Plugins/Runtime/ModelViewViewModel/Source/ModelViewViewModelBlueprint/Private/MVVMBlueprintView.cpp
+++ b/Engine/Plugins/Runtime/ModelViewViewModel/Source/ModelViewViewModelBlueprint/Private/MVVMBlueprintView.cpp
@@ -312,6 +312,8 @@ void UMVVMBlueprintView::RemoveEvent(UMVVMBlueprintViewEvent* Event)
 UMVVMBlueprintViewEvent* UMVVMBlueprintView::DuplicateEvent(UMVVMBlueprintViewEvent* Event)
 {
 	UMVVMBlueprintViewEvent* NewEvent = DuplicateObject<UMVVMBlueprintViewEvent>(Event, this);
+	NewEvent->RecreateWrapperGraph();
+	
 	AddEvent(NewEvent);
 
 	return NewEvent;
@@ -349,6 +351,7 @@ void UMVVMBlueprintView::RemoveCondition(UMVVMBlueprintViewCondition* Condition)
 UMVVMBlueprintViewCondition* UMVVMBlueprintView::DuplicateCondition(UMVVMBlueprintViewCondition* Condition)
 {
 	UMVVMBlueprintViewCondition* NewCondition = DuplicateObject<UMVVMBlueprintViewCondition>(Condition, this);
+	NewCondition->RecreateWrapperGraph();
 
 	AddCondition(NewCondition);
 
diff --git a/Engine/Plugins/Runtime/ModelViewViewModel/Source/ModelViewViewModelBlueprint/Private/MVVMBlueprintViewCondition.cpp b/Engine/Plugins/Runtime/ModelViewViewModel/Source/ModelViewViewModelBlueprint/Private/MVVMBlueprintViewCondition.cpp
--- a/Engine/Plugins/Runtime/ModelViewViewModel/Source/ModelViewViewModelBlueprint/Private/MVVMBlueprintViewCondition.cpp
+++ b/Engine/Plugins/Runtime/ModelViewViewModel/Source/ModelViewViewModelBlueprint/Private/MVVMBlueprintViewCondition.cpp
@@ -160,6 +160,22 @@ UEdGraph* UMVVMBlueprintViewCondition::GetOrCreateWrapperGraph()
 	return CachedWrapperGraph;
 }
 
+void UMVVMBlueprintViewCondition::RecreateWrapperGraph()
+{
+	GraphName = FName();
+
+	RemoveWrapperGraph(LeaveConversionFunctionCurrentValues);
+
+	TStringBuilder<256> StringBuilder;
+	StringBuilder << TEXT("__");
+	StringBuilder << FGuid::NewGuid().ToString(EGuidFormats::DigitsWithHyphensLower);
+	GraphName = StringBuilder.ToString();
+
+	bNeedsToRegenerateChildren = true;
+
+	CreateWrapperGraphInternal();
+}
+
 void UMVVMBlueprintViewCondition::RemoveWrapperGraph(ERemoveWrapperGraphParam ActionForCurrentValues)
 {
 	if (CachedWrapperGraph)
diff --git a/Engine/Plugins/Runtime/ModelViewViewModel/Source/ModelViewViewModelBlueprint/Private/MVVMBlueprintViewEvent.cpp b/Engine/Plugins/Runtime/ModelViewViewModel/Source/ModelViewViewModelBlueprint/Private/MVVMBlueprintViewEvent.cpp
--- a/Engine/Plugins/Runtime/ModelViewViewModel/Source/ModelViewViewModelBlueprint/Private/MVVMBlueprintViewEvent.cpp
+++ b/Engine/Plugins/Runtime/ModelViewViewModel/Source/ModelViewViewModelBlueprint/Private/MVVMBlueprintViewEvent.cpp
@@ -109,6 +109,22 @@ UEdGraph* UMVVMBlueprintViewEvent::GetOrCreateWrapperGraph()
 	return CachedWrapperGraph;
 }
 
+void UMVVMBlueprintViewEvent::RecreateWrapperGraph()
+{
+	GraphName = FName();
+
+	RemoveWrapperGraph(LeaveConversionFunctionCurrentValues);
+
+	TStringBuilder<256> StringBuilder;
+	StringBuilder << TEXT("__");
+	StringBuilder << FGuid::NewGuid().ToString(EGuidFormats::DigitsWithHyphensLower);
+	GraphName = StringBuilder.ToString();
+
+	bNeedsToRegenerateChildren = true;
+
+	CreateWrapperGraphInternal();
+}
+
 void UMVVMBlueprintViewEvent::RemoveWrapperGraph(ERemoveWrapperGraphParam ActionForCurrentValues)
 {
 	if (CachedWrapperGraph)
diff --git a/Engine/Plugins/Runtime/ModelViewViewModel/Source/ModelViewViewModelBlueprint/Public/MVVMBlueprintViewCondition.h b/Engine/Plugins/Runtime/ModelViewViewModel/Source/ModelViewViewModelBlueprint/Public/MVVMBlueprintViewCondition.h
--- a/Engine/Plugins/Runtime/ModelViewViewModel/Source/ModelViewViewModelBlueprint/Public/MVVMBlueprintViewCondition.h
+++ b/Engine/Plugins/Runtime/ModelViewViewModel/Source/ModelViewViewModelBlueprint/Public/MVVMBlueprintViewCondition.h
@@ -106,6 +106,8 @@ public:
 
 	UEdGraph* GetOrCreateWrapperGraph();
 
+	void RecreateWrapperGraph();
+
 public:
 	TArrayView<const FMVVMBlueprintPin> GetPins() const
 	{
diff --git a/Engine/Plugins/Runtime/ModelViewViewModel/Source/ModelViewViewModelBlueprint/Public/MVVMBlueprintViewEvent.h b/Engine/Plugins/Runtime/ModelViewViewModel/Source/ModelViewViewModelBlueprint/Public/MVVMBlueprintViewEvent.h
--- a/Engine/Plugins/Runtime/ModelViewViewModel/Source/ModelViewViewModelBlueprint/Public/MVVMBlueprintViewEvent.h
+++ b/Engine/Plugins/Runtime/ModelViewViewModel/Source/ModelViewViewModelBlueprint/Public/MVVMBlueprintViewEvent.h
@@ -89,6 +89,8 @@ public:
 
 	UEdGraph* GetOrCreateWrapperGraph();
 
+	void RecreateWrapperGraph();
+
 public:
 	TArrayView<const FMVVMBlueprintPin> GetPins() const
 	{
