Protect against crash on conversion functions signature changing

#tests Tried finding repros for this one but couldn't. From the crash, it looks like the data might have been corrupted. Perhaps a function signature was changed and a parameter is now a type which leads to these crashes. Tested manually by changing the code to trigger the new if statement
#rnx

diff --git a/Engine/Plugins/Runtime/ModelViewViewModel/Source/ModelViewViewModelEditor/Private/Widgets/BindingEntry/SMVVMFunctionParameterRow.cpp b/Engine/Plugins/Runtime/ModelViewViewModel/Source/ModelViewViewModelEditor/Private/Widgets/BindingEntry/SMVVMFunctionParameterRow.cpp
--- a/Engine/Plugins/Runtime/ModelViewViewModel/Source/ModelViewViewModelEditor/Private/Widgets/BindingEntry/SMVVMFunctionParameterRow.cpp
+++ b/Engine/Plugins/Runtime/ModelViewViewModel/Source/ModelViewViewModelEditor/Private/Widgets/BindingEntry/SMVVMFunctionParameterRow.cpp
@@ -100,7 +100,14 @@ TSharedRef<SWidget> SFunctionParameterRow::BuildRowWidget()
 	{
 		PrimaryBrush = FBlueprintEditorUtils::GetIconFromPin(FEdGraphPinType());
 		PrimaryColor = FLinearColor::Red;
-		DisplayName = GraphPin ? GraphPin->GetDisplayName() : (Pin ? FText::FromName(Pin->GetId().GetNames().Last()) : FText::GetEmpty());
+		DisplayName = GraphPin ? GraphPin->GetDisplayName() : FText::GetEmpty();
+		if (GraphPin == nullptr && Pin != nullptr)
+		{
+			if (!Pin->GetId().GetNames().IsEmpty())
+			{
+				DisplayName = FText::FromName(Pin->GetId().GetNames().Last());
+			}
+		}
 		bTextColorIsRed = true;
 	}
 	else
