[AR] Remove unused Viewmodels tag from the AssetRegistry and cleared up places where it was previously used

#tests Local cook
#rnx

diff --git a/Engine/Plugins/Runtime/ModelViewViewModel/Source/ModelViewViewModelBlueprint/Private/MVVMBlueprintView.cpp b/Engine/Plugins/Runtime/ModelViewViewModel/Source/ModelViewViewModelBlueprint/Private/MVVMBlueprintView.cpp
--- a/Engine/Plugins/Runtime/ModelViewViewModel/Source/ModelViewViewModelBlueprint/Private/MVVMBlueprintView.cpp
+++ b/Engine/Plugins/Runtime/ModelViewViewModel/Source/ModelViewViewModelBlueprint/Private/MVVMBlueprintView.cpp
@@ -699,28 +699,7 @@ void UMVVMBlueprintView::AddAssetTags(TArray<FAssetRegistryTag>& OutTags) const
 
 void UMVVMBlueprintView::AddAssetTags(FAssetRegistryTagsContext Context) const
 {
-	if (AvailableViewModels.Num() > 0)
-	{
-		TStringBuilder<512> Builder;
-		for (const FMVVMBlueprintViewModelContext& ViewModelContext : AvailableViewModels)
-		{
-			if (ViewModelContext.IsValid())
-			{
-				if (Builder.Len() > 0)
-				{
-					Builder << TEXT(',');
-				}
-				Builder << ViewModelContext.GetViewModelName();
-				Builder << TEXT(';');
-				Builder << ViewModelContext.GetViewModelClass()->GetPathName();
-			}
-		}
 
-		if (Builder.Len() > 0)
-		{
-			Context.AddTag(FAssetRegistryTag(FName("Viewmodels"), Builder.ToString(), FAssetRegistryTag::TT_Hidden));
-		}
-	}
 }
 
 void UMVVMBlueprintView::OnFieldRenamed(UClass* FieldOwnerClass, FName OldObjectName, FName NewObjectName)
