MVVM: reset the registered binding count when uninitializing the bindings in View
fixes the false warning when uninitializing bindings on null optional VMs.

diff --git a/Engine/Plugins/Runtime/ModelViewViewModel/Source/ModelViewViewModel/Private/View/MVVMView.cpp b/Engine/Plugins/Runtime/ModelViewViewModel/Source/ModelViewViewModel/Private/View/MVVMView.cpp
--- a/Engine/Plugins/Runtime/ModelViewViewModel/Source/ModelViewViewModel/Private/View/MVVMView.cpp
+++ b/Engine/Plugins/Runtime/ModelViewViewModel/Source/ModelViewViewModel/Private/View/MVVMView.cpp
@@ -507,6 +507,7 @@ void UMVVMView::UninitializeSourceBindings(FMVVMViewClass_SourceKey SourceKey, c
 			TScriptInterface<INotifyFieldValueChanged> SourceAsInterface = ViewSource.Source;
 			checkf(SourceAsInterface.GetInterface(), TEXT("It was added as a INotifyFieldValueChanged. It should still be."));
 			SourceAsInterface->RemoveAllFieldValueChangedDelegates(this);
+			ViewSource.RegisteredCount = 0;
 		}
 		else
 		{
