MVVM: Fix crash that can occur when reparenting a widget blueprint's class

Issue stems from a missing "continue" statement after encountering an error.

#rnx

diff --git a/Engine/Plugins/Runtime/ModelViewViewModel/Source/ModelViewViewModelBlueprint/Private/MVVMViewBlueprintCompiler.cpp b/Engine/Plugins/Runtime/ModelViewViewModel/Source/ModelViewViewModelBlueprint/Private/MVVMViewBlueprintCompiler.cpp
--- a/Engine/Plugins/Runtime/ModelViewViewModel/Source/ModelViewViewModelBlueprint/Private/MVVMViewBlueprintCompiler.cpp
+++ b/Engine/Plugins/Runtime/ModelViewViewModel/Source/ModelViewViewModelBlueprint/Private/MVVMViewBlueprintCompiler.cpp
@@ -2319,6 +2319,7 @@ void FMVVMViewBlueprintCompiler::CreateReadFieldContexts(UWidgetBlueprintGenerat
 						, Pin.GetId()
 					);
 					bIsPreCompileStepValid = false;
+					continue;
 				}
 
 				TSharedPtr<FGeneratedReadFieldPathContext> Found = AlreadyExist(CreateSourceResult.GetValue().SkeletalGeneratedFields);
