MVVM: Fix typo in error message.
#rnx

diff --git a/Engine/Plugins/Runtime/ModelViewViewModel/Source/ModelViewViewModel/Private/Bindings/MVVMBindingHelper.cpp b/Engine/Plugins/Runtime/ModelViewViewModel/Source/ModelViewViewModel/Private/Bindings/MVVMBindingHelper.cpp
--- a/Engine/Plugins/Runtime/ModelViewViewModel/Source/ModelViewViewModel/Private/Bindings/MVVMBindingHelper.cpp
+++ b/Engine/Plugins/Runtime/ModelViewViewModel/Source/ModelViewViewModel/Private/Bindings/MVVMBindingHelper.cpp
@@ -237,7 +237,7 @@ namespace UE::MVVM::BindingHelper
 
 			if (InProperty->HasAnyPropertyFlags(CPF_Deprecated))
 			{
-				return FText::Format(LOCTEXT("PropertyDeprecated", "The property '{0}' is depreated."), FText::FromString(InProperty->GetName()));
+				return FText::Format(LOCTEXT("PropertyDeprecated", "The property '{0}' is deprecated."), FText::FromString(InProperty->GetName()));
 			}
 
 			if (InProperty->HasAnyPropertyFlags(CPF_EditorOnly))
