MVVM: Allow pasting into empty path bindings
#rnx

diff --git a/Engine/Plugins/Runtime/ModelViewViewModel/Source/ModelViewViewModelEditor/Private/Framework/MVVMRowHelper.cpp b/Engine/Plugins/Runtime/ModelViewViewModel/Source/ModelViewViewModelEditor/Private/Framework/MVVMRowHelper.cpp
--- a/Engine/Plugins/Runtime/ModelViewViewModel/Source/ModelViewViewModelEditor/Private/Framework/MVVMRowHelper.cpp
+++ b/Engine/Plugins/Runtime/ModelViewViewModel/Source/ModelViewViewModelEditor/Private/Framework/MVVMRowHelper.cpp
@@ -288,7 +288,7 @@ FMVVMBlueprintPropertyPath TryGetCommonPropertyPath(UWidgetBlueprint* WidgetBlue
 	}
 	default:
 	{
-		check(false);
+		// If target path is empty, do nothing. Will support all source property paths by default. 
 		break;
 	}
 	}
@@ -297,7 +297,7 @@ FMVVMBlueprintPropertyPath TryGetCommonPropertyPath(UWidgetBlueprint* WidgetBlue
 }
 }
 
-void FRowHelper::GatherAllChildBindings(UMVVMBlueprintView* BlueprintView, const TConstArrayView<TSharedPtr<FBindingEntry>> Entries, TArray<FGuid>& OutBindingIds, TArray<UMVVMBlueprintViewEvent*>& OutEvents, TArray<UMVVMBlueprintViewCondition*>& OutConditions)
+void FRowHelper::GatherAllChildBindings(UMVVMBlueprintView* BlueprintView, const TConstArrayView<TSharedPtr<FBindingEntry>> Entries, TArray<FGuid>& OutBindingIds, TArray<UMVVMBlueprintViewEvent*>& OutEvents, TArray<UMVVMBlueprintViewCondition*>& OutConditions, bool bRecurse)
 {
 	for (const TSharedPtr<FBindingEntry>& Entry : Entries)
 	{
@@ -327,7 +327,11 @@ void FRowHelper::GatherAllChildBindings(UMVVMBlueprintView* BlueprintView, const
 				OutConditions.AddUnique(Condition);
 			}
 		}
-		GatherAllChildBindings(BlueprintView, Entry->GetAllChildren(), OutBindingIds, OutEvents, OutConditions);
+
+		if (bRecurse)
+		{
+			GatherAllChildBindings(BlueprintView, Entry->GetAllChildren(), OutBindingIds, OutEvents, OutConditions, bRecurse);
+		}
 	}
 }
 
@@ -400,7 +404,7 @@ void FRowHelper::DuplicateEntries(const UWidgetBlueprint* WidgetBlueprint, UMVVM
 		TArray<FGuid> BindingIdsToDuplicate;
 		TArray<UMVVMBlueprintViewEvent*> EventsToDuplicate;
 		TArray<UMVVMBlueprintViewCondition*> ConditionsToDuplicate;
-		GatherAllChildBindings(BlueprintView, Selection, BindingIdsToDuplicate, EventsToDuplicate, ConditionsToDuplicate);
+		GatherAllChildBindings(BlueprintView, Selection, BindingIdsToDuplicate, EventsToDuplicate, ConditionsToDuplicate, false);
 
 		if (BindingIdsToDuplicate.Num() == 0 && EventsToDuplicate.Num() == 0 && ConditionsToDuplicate.Num() == 0)
 		{
@@ -646,6 +650,8 @@ void FRowHelper::PasteEntries(UWidgetBlueprint* WidgetBlueprint, UMVVMBlueprintV
 		const FMVVMBlueprintPropertyPath NewPath = Private::TryGetCommonPropertyPath(WidgetBlueprint, BlueprintView, TargetPath, DesiredPath);
 
 		EditorSubsystem->SetDestinationPathForBinding(WidgetBlueprint, *BindingSelected, NewPath, false);
+
+		BlueprintView->OnBindingsUpdated.Broadcast();
 	}
 	else if (Item.IsType<UMVVMBlueprintViewCondition*>())
 	{
@@ -1058,6 +1064,33 @@ FMenuBuilder FRowHelper::CreateContextMenu(UWidgetBlueprint* WidgetBlueprint, UM
 				}
 			}
 
+			if (bCanPasteEntry && Entries.Num() == 1) // Check paste type
+			{
+				bCanPasteEntry = false;
+
+				// Only allow pasting a single entry into entry of same type
+				Private::FBindingClipboardData ClipboardData;
+				Private::ImportClipboardData(View, ClipboardData);
+				if (ClipboardData.Items.Num() == 1)
+				{
+					const TSharedPtr<FBindingEntry> Entry = Entries.Last();
+					const Private::FBindingClipboardData::Type& Item = ClipboardData.Items.Last();
+
+					if (Item.IsType<FMVVMBlueprintViewBinding>())
+					{
+						bCanPasteEntry = Entry->GetRowType() == FBindingEntry::ERowType::Binding;
+					}
+					else if (Item.IsType<UMVVMBlueprintViewCondition*>())
+					{
+						bCanPasteEntry = Entry->GetRowType() == FBindingEntry::ERowType::Condition;
+					}
+					else if (Item.IsType<UMVVMBlueprintViewEvent*>())
+					{
+						bCanPasteEntry = Entry->GetRowType() == FBindingEntry::ERowType::Event;
+					}
+				}
+			}
+
 			FUIAction RemoveAction;
 			RemoveAction.ExecuteAction = FExecuteAction::CreateStatic(&Private::HandleDeleteEntries, WidgetBlueprint, View, CopiedEntries);
 			RemoveAction.CanExecuteAction = FCanExecuteAction::CreateLambda([bCanRemoveEntry]() { return bCanRemoveEntry; });
diff --git a/Engine/Plugins/Runtime/ModelViewViewModel/Source/ModelViewViewModelEditor/Private/Framework/MVVMRowHelper.h b/Engine/Plugins/Runtime/ModelViewViewModel/Source/ModelViewViewModelEditor/Private/Framework/MVVMRowHelper.h
--- a/Engine/Plugins/Runtime/ModelViewViewModel/Source/ModelViewViewModelEditor/Private/Framework/MVVMRowHelper.h
+++ b/Engine/Plugins/Runtime/ModelViewViewModel/Source/ModelViewViewModelEditor/Private/Framework/MVVMRowHelper.h
@@ -25,7 +25,7 @@ DECLARE_DELEGATE_OneParam(FOnContextMenuEntryCallback, TConstArrayView<const TSh
 
 struct FRowHelper
 {
-	static void GatherAllChildBindings(UMVVMBlueprintView* BlueprintView, const TConstArrayView<TSharedPtr<FBindingEntry>> Entries, TArray<FGuid>& OutBindingIds, TArray<UMVVMBlueprintViewEvent*>& OutEvents, TArray<UMVVMBlueprintViewCondition*>& OutConditions);
+	static void GatherAllChildBindings(UMVVMBlueprintView* BlueprintView, const TConstArrayView<TSharedPtr<FBindingEntry>> Entries, TArray<FGuid>& OutBindingIds, TArray<UMVVMBlueprintViewEvent*>& OutEvents, TArray<UMVVMBlueprintViewCondition*>& OutConditions, bool bRecurse = true);
 
 	static void DeleteEntries(const UWidgetBlueprint* WidgetBlueprint, UMVVMBlueprintView* BlueprintView, TArrayView<const TSharedPtr<FBindingEntry>> Selection);
 	static void DuplicateEntries(const UWidgetBlueprint* WidgetBlueprint, UMVVMBlueprintView* BlueprintView, TArrayView<const TSharedPtr<FBindingEntry>> Selection, TArray<const TSharedPtr<FBindingEntry>>& OutSelection);
