Replaced bool arguments with EFindObjectFlags.

diff --git a/Engine/Plugins/Animation/ControlRig/Source/ControlRig/Private/ControlRig.cpp b/Engine/Plugins/Animation/ControlRig/Source/ControlRig/Private/ControlRig.cpp
--- a/Engine/Plugins/Animation/ControlRig/Source/ControlRig/Private/ControlRig.cpp
+++ b/Engine/Plugins/Animation/ControlRig/Source/ControlRig/Private/ControlRig.cpp
@@ -3702,7 +3702,7 @@ void UControlRig::PostInitInstance(URigVMHost* InCDO)
 		if (!HasAnyFlags(RF_ClassDefaultObject))
 		{
 			FGCScopeGuard Guard;
-			UObject* ObjectFound = StaticFindObjectFast(URigHierarchy::StaticClass(), this, TEXT("DynamicHierarchy"), true, RF_DefaultSubObject);
+			UObject* ObjectFound = StaticFindObjectFast(URigHierarchy::StaticClass(), this, TEXT("DynamicHierarchy"), EFindObjectFlags::ExactClass, RF_DefaultSubObject);
 			if (ObjectFound)
 			{
 				FName NewName = MakeUniqueObjectName(GetTransientPackage(), URigHierarchy::StaticClass(), TEXT("DynamicHierarchy_Deleted"));
diff --git a/Engine/Plugins/Animation/ControlRig/Source/ControlRigDeveloper/Private/ControlRigVisualGraphUtils.cpp b/Engine/Plugins/Animation/ControlRig/Source/ControlRigDeveloper/Private/ControlRigVisualGraphUtils.cpp
--- a/Engine/Plugins/Animation/ControlRig/Source/ControlRigDeveloper/Private/ControlRigVisualGraphUtils.cpp
+++ b/Engine/Plugins/Animation/ControlRig/Source/ControlRigDeveloper/Private/ControlRigVisualGraphUtils.cpp
@@ -64,7 +64,7 @@ FAutoConsoleCommandWithWorldAndArgs FCmdControlRigVisualGraphUtilsDumpHierarchy
 		TArray<URigHierarchy*> Hierarchies;
 		for(const FString& ObjectPathName : ObjectPathNames)
 		{
-			if(UObject* Object = FindObject<UObject>(nullptr, *ObjectPathName, false))
+			if(UObject* Object = FindObject<UObject>(nullptr, *ObjectPathName, EFindObjectFlags::None))
 			{
 				if(UControlRig* CR = Cast<UControlRig>(Object))
 				{
diff --git a/Engine/Plugins/Animation/DeformerGraph/Source/OptimusCore/Private/Actions/OptimusComponentBindingActions.cpp b/Engine/Plugins/Animation/DeformerGraph/Source/OptimusCore/Private/Actions/OptimusComponentBindingActions.cpp
--- a/Engine/Plugins/Animation/DeformerGraph/Source/OptimusCore/Private/Actions/OptimusComponentBindingActions.cpp
+++ b/Engine/Plugins/Animation/DeformerGraph/Source/OptimusCore/Private/Actions/OptimusComponentBindingActions.cpp
@@ -40,7 +40,7 @@ bool FOptimusComponentBindingAction_AddBinding::Do(IOptimusPathResolver* InRoot)
 {
 	UOptimusDeformer* Deformer = Cast<UOptimusDeformer>(InRoot);
 
-	const UClass* ComponentSourceClass = Cast<UClass>(StaticFindObject(UClass::StaticClass(), ComponentSourceClassPath, false));
+	const UClass* ComponentSourceClass = Cast<UClass>(StaticFindObject(UClass::StaticClass(), ComponentSourceClassPath, EFindObjectFlags::None));
 	if (!ComponentSourceClass)
 	{
 		return false;
@@ -113,7 +113,7 @@ bool FOptimusComponentBindingAction_RemoveBinding::Undo(IOptimusPathResolver* In
 {
 	UOptimusDeformer* Deformer = Cast<UOptimusDeformer>(InRoot);
 
-	const UClass* ComponentSourceClass = Cast<UClass>(StaticFindObject(UClass::StaticClass(), ComponentSourceClassPath, false));
+	const UClass* ComponentSourceClass = Cast<UClass>(StaticFindObject(UClass::StaticClass(), ComponentSourceClassPath, EFindObjectFlags::None));
 	if (!ComponentSourceClass)
 	{
 		return false;
@@ -209,7 +209,7 @@ bool FOptimusComponentBindingAction_SetComponentSource::SetComponentSource(
 	FTopLevelAssetPath InComponentSourceClassPath
 	) const
 {
-	const UClass* ComponentSourceClass = Cast<UClass>(StaticFindObject(UClass::StaticClass(), InComponentSourceClassPath, false));
+	const UClass* ComponentSourceClass = Cast<UClass>(StaticFindObject(UClass::StaticClass(), InComponentSourceClassPath, EFindObjectFlags::None));
 	if (!ComponentSourceClass)
 	{
 		return false;
diff --git a/Engine/Plugins/Animation/DeformerGraph/Source/OptimusCore/Private/Nodes/OptimusNode_ComponentSource.cpp b/Engine/Plugins/Animation/DeformerGraph/Source/OptimusCore/Private/Nodes/OptimusNode_ComponentSource.cpp
--- a/Engine/Plugins/Animation/DeformerGraph/Source/OptimusCore/Private/Nodes/OptimusNode_ComponentSource.cpp
+++ b/Engine/Plugins/Animation/DeformerGraph/Source/OptimusCore/Private/Nodes/OptimusNode_ComponentSource.cpp
@@ -88,7 +88,7 @@ void UOptimusNode_ComponentSource::ImportCustomProperties(const TCHAR* SourceTex
 			return;
 		}
 
-		UClass* ComponentSourceClass = Cast<UClass>(StaticFindObject(UClass::StaticClass(), FTopLevelAssetPath(ComponentTypeClassPath), false));
+		UClass* ComponentSourceClass = Cast<UClass>(StaticFindObject(UClass::StaticClass(), FTopLevelAssetPath(ComponentTypeClassPath), EFindObjectFlags::None));
 		if (!ComponentSourceClass)
 		{
 			return;
diff --git a/Engine/Plugins/Animation/DeformerGraph/Source/OptimusCore/Private/OptimusComputeGraph.cpp b/Engine/Plugins/Animation/DeformerGraph/Source/OptimusCore/Private/OptimusComputeGraph.cpp
--- a/Engine/Plugins/Animation/DeformerGraph/Source/OptimusCore/Private/OptimusComputeGraph.cpp
+++ b/Engine/Plugins/Animation/DeformerGraph/Source/OptimusCore/Private/OptimusComputeGraph.cpp
@@ -67,7 +67,7 @@ static EOptimusDiagnosticLevel ProcessCompilationMessage(UOptimusDeformer* InOwn
 	FString Path = InMessage.VirtualFilePath;
 	if (Optimus::ConvertShaderFilePathToObjectPath(Path))
 	{
-		Diagnostic.Object = StaticFindObject(nullptr, nullptr, *Path, true);
+		Diagnostic.Object = StaticFindObject(nullptr, nullptr, *Path, EFindObjectFlags::ExactClass);
 		if (UObject const* ObjectPtr = Diagnostic.Object.Get())
 		{
 			// Special case path display if object is the current kernel.
diff --git a/Engine/Plugins/Experimental/PythonScriptPlugin/Source/PythonScriptPlugin/Private/PyGenUtil.cpp b/Engine/Plugins/Experimental/PythonScriptPlugin/Source/PythonScriptPlugin/Private/PyGenUtil.cpp
--- a/Engine/Plugins/Experimental/PythonScriptPlugin/Source/PythonScriptPlugin/Private/PyGenUtil.cpp
+++ b/Engine/Plugins/Experimental/PythonScriptPlugin/Source/PythonScriptPlugin/Private/PyGenUtil.cpp
@@ -2503,7 +2503,7 @@ void PythonizeStructValueImpl(const UScriptStruct* InStruct, const void* InStruc
 		const FString MakeBreakName = InStruct->GetMetaData(InKey);
 		if (!MakeBreakName.IsEmpty())
 		{
-			return FindObject<UFunction>(nullptr, *MakeBreakName, true);
+			return FindObject<UFunction>(nullptr, *MakeBreakName, EFindObjectFlags::ExactClass);
 		}
 		return nullptr;
 	};
diff --git a/Engine/Plugins/Experimental/PythonScriptPlugin/Source/PythonScriptPlugin/Private/PyWrapperTypeRegistry.cpp b/Engine/Plugins/Experimental/PythonScriptPlugin/Source/PythonScriptPlugin/Private/PyWrapperTypeRegistry.cpp
--- a/Engine/Plugins/Experimental/PythonScriptPlugin/Source/PythonScriptPlugin/Private/PyWrapperTypeRegistry.cpp
+++ b/Engine/Plugins/Experimental/PythonScriptPlugin/Source/PythonScriptPlugin/Private/PyWrapperTypeRegistry.cpp
@@ -102,7 +102,7 @@ UFunction* FindMakeBreakFunction(const UScriptStruct* InStruct, const FName& InN
 		if (!MakeBreakFunctionName.IsEmpty())
 		{
 			// Find the function.
-			MakeBreakFunc = FindObject<UFunction>(/*Outer*/nullptr, *MakeBreakFunctionName, /*ExactClass*/true);
+			MakeBreakFunc = FindObject<UFunction>(/*Outer*/nullptr, *MakeBreakFunctionName, EFindObjectFlags::ExactClass);
 			if (!MakeBreakFunc)
 			{
 				NotFoundFn(MakeBreakFunctionName);
@@ -1093,7 +1093,7 @@ void FPyWrapperTypeRegistry::GenerateWrappedTypesForModule(const FName ModuleNam
 	// that might not be loaded yet when Python glue was generated for the struct.
 	UnresolvedMakeFuncs.RemoveIf([](TPair<FString, TSharedPtr<FPyWrapperStructMetaData>>& Pair)
 	{
-		if (UFunction* MakeFunc = FindObject<UFunction>(nullptr, *Pair.Key, true))
+		if (UFunction* MakeFunc = FindObject<UFunction>(nullptr, *Pair.Key, EFindObjectFlags::ExactClass))
 		{
 			UE::Python::SetMakeFunction(*Pair.Value, MakeFunc);
 			return true; // Remove
@@ -1104,7 +1104,7 @@ void FPyWrapperTypeRegistry::GenerateWrappedTypesForModule(const FName ModuleNam
 	
 	UnresolvedBreakFuncs.RemoveIf([](TPair<FString, TSharedPtr<FPyWrapperStructMetaData>>& Pair)
 	{
-		if (UFunction* BreakFunc = FindObject<UFunction>(nullptr, *Pair.Key, true))
+		if (UFunction* BreakFunc = FindObject<UFunction>(nullptr, *Pair.Key, EFindObjectFlags::ExactClass))
 		{
 			UE::Python::SetBreakFunction(*Pair.Value, BreakFunc);
 			return true; // Remove
diff --git a/Engine/Plugins/Interchange/Runtime/Source/Pipelines/Private/InterchangeGenericMaterialPipeline.cpp b/Engine/Plugins/Interchange/Runtime/Source/Pipelines/Private/InterchangeGenericMaterialPipeline.cpp
--- a/Engine/Plugins/Interchange/Runtime/Source/Pipelines/Private/InterchangeGenericMaterialPipeline.cpp
+++ b/Engine/Plugins/Interchange/Runtime/Source/Pipelines/Private/InterchangeGenericMaterialPipeline.cpp
@@ -270,12 +270,11 @@ namespace UE::Interchange::InterchangeGenericMaterialPipeline::Private
 		}
 
 		//Search in memory
-		constexpr bool bExactClass = false;
 		UMaterialInterface* FoundMaterial = nullptr;
 		//We search only in memory for search in local folder.
 		if (SearchLocation == EInterchangeMaterialSearchLocation::Local)
 		{
-			FoundMaterial = FindObject<UMaterialInterface>(nullptr, *MaterialFullName, bExactClass);
+			FoundMaterial = FindObject<UMaterialInterface>(nullptr, *MaterialFullName, EFindObjectFlags::None);
 			if(FoundMaterial)
 			{
 				//Make sure the path of the material in memory is local
diff --git a/Engine/Plugins/Runtime/MeshModelingToolset/Source/ModelingComponentsEditorOnly/Private/AssetUtils/CreateTexture2DUtil.cpp b/Engine/Plugins/Runtime/MeshModelingToolset/Source/ModelingComponentsEditorOnly/Private/AssetUtils/CreateTexture2DUtil.cpp
--- a/Engine/Plugins/Runtime/MeshModelingToolset/Source/ModelingComponentsEditorOnly/Private/AssetUtils/CreateTexture2DUtil.cpp
+++ b/Engine/Plugins/Runtime/MeshModelingToolset/Source/ModelingComponentsEditorOnly/Private/AssetUtils/CreateTexture2DUtil.cpp
@@ -42,7 +42,7 @@ UE::AssetUtils::ECreateTexture2DResult UE::AssetUtils::SaveGeneratedTexture2DAss
 	}
 
 	// check for existing asset of the same type.
-	if (UObject* ExistingAsset = StaticFindObject(nullptr, UsePackage, *NewObjectName, false))
+	if (UObject* ExistingAsset = StaticFindObject(nullptr, UsePackage, *NewObjectName, EFindObjectFlags::None))
 	{
 		UTexture2D* ExistingTexture2D = Cast<UTexture2D>(ExistingAsset);
 		if (ExistingTexture2D && Options.bOverwriteIfExists)
diff --git a/Engine/Plugins/Runtime/ModelViewViewModel/Source/ModelViewViewModelBlueprint/Private/Bindings/MVVMConversionFunctionHelper.cpp b/Engine/Plugins/Runtime/ModelViewViewModel/Source/ModelViewViewModelBlueprint/Private/Bindings/MVVMConversionFunctionHelper.cpp
--- a/Engine/Plugins/Runtime/ModelViewViewModel/Source/ModelViewViewModelBlueprint/Private/Bindings/MVVMConversionFunctionHelper.cpp
+++ b/Engine/Plugins/Runtime/ModelViewViewModel/Source/ModelViewViewModelBlueprint/Private/Bindings/MVVMConversionFunctionHelper.cpp
@@ -121,7 +121,7 @@ namespace NamedNodes
 
 	FCreateGraphResult CreateGraph(UBlueprint* Blueprint, FName GraphName, const UFunction* FunctionEntryDefinition, FCreateGraphParams InParams)
 	{
-		if (UObject* ExistingObject = StaticFindObject(nullptr, Blueprint, *GraphName.ToString(), true))
+		if (UObject* ExistingObject = StaticFindObject(nullptr, Blueprint, *GraphName.ToString(), EFindObjectFlags::ExactClass))
 		{
 			FunctionGraphHelper::RenameObjectToTransientPackage(ExistingObject);
 		}
@@ -354,7 +354,7 @@ namespace NamedNodes
 						const FString& MetaData = ArgumentStructProperty->Struct->GetMetaData(FBlueprintMetadata::MD_NativeBreakFunction);
 						if (MetaData.Len() > 0)
 						{
-							if (const UFunction* NativeBreakFunction = FindObject<UFunction>(nullptr, *MetaData, true))
+							if (const UFunction* NativeBreakFunction = FindObject<UFunction>(nullptr, *MetaData, EFindObjectFlags::ExactClass))
 							{
 								AddBreakNode(ArgumentStructProperty->Struct, NodePair.Get<UEdGraphPin*>()->GetFName());
 								bAddPropertyPath = false;
@@ -533,7 +533,7 @@ namespace NamedNodes
 					const FString& MetaData = PreviousStructProperty->Struct->GetMetaData(FBlueprintMetadata::MD_NativeBreakFunction);
 					if (MetaData.Len() > 0)
 					{
-						if (const UFunction* Function = FindObject<UFunction>(nullptr, *MetaData, true))
+						if (const UFunction* Function = FindObject<UFunction>(nullptr, *MetaData, EFindObjectFlags::ExactClass))
 						{
 							FGraphNodeCreator<UK2Node_CallFunction> MakeStructCreator(*FunctionGraph);
 							UK2Node_CallFunction* FunctionNode = MakeStructCreator.CreateNode(false);
diff --git a/Engine/Plugins/Runtime/ModelViewViewModel/Source/ModelViewViewModelBlueprint/Private/MVVMBlueprintInstancedViewModel.cpp b/Engine/Plugins/Runtime/ModelViewViewModel/Source/ModelViewViewModelBlueprint/Private/MVVMBlueprintInstancedViewModel.cpp
--- a/Engine/Plugins/Runtime/ModelViewViewModel/Source/ModelViewViewModelBlueprint/Private/MVVMBlueprintInstancedViewModel.cpp
+++ b/Engine/Plugins/Runtime/ModelViewViewModel/Source/ModelViewViewModelBlueprint/Private/MVVMBlueprintInstancedViewModel.cpp
@@ -45,7 +45,7 @@ void UMVVMBlueprintInstancedViewModelBase::GenerateClass(bool bForceGeneration)
 	{
 		UPackage* Outermost = GetOutermost();
 		FName NewClassName = *FString::Printf(TEXT("InstanceViewmodel%d_IC"), GetFName().GetNumber());
-		PreviousClass = StaticFindObjectFastInternal(nullptr, Outermost, NewClassName, true);
+		PreviousClass = StaticFindObjectFastInternal(nullptr, Outermost, NewClassName, EFindObjectFlags::ExactClass);
 		if (PreviousClass)
 		{
 			SafeRename(PreviousClass);
@@ -269,7 +269,7 @@ void UMVVMBlueprintInstancedViewModelBase::AddOnRepFunction(FProperty* NewProper
 {
 	FString OnRepCallFunctionName = FString::Printf(TEXT("__OnRep_%s"), *NewProperty->GetName());
 	FName Name_OnRepCallFunctionName = *OnRepCallFunctionName;
-	UObject* PreviousObj = StaticFindObjectFastInternal(nullptr, GeneratedClass, Name_OnRepCallFunctionName, true);
+	UObject* PreviousObj = StaticFindObjectFastInternal(nullptr, GeneratedClass, Name_OnRepCallFunctionName, EFindObjectFlags::ExactClass);
 	ensure(PreviousObj == nullptr);
 	if (PreviousObj)
 	{
diff --git a/Engine/Plugins/Runtime/ModelViewViewModel/Source/ModelViewViewModelBlueprint/Private/MVVMBlueprintViewConversionFunction.cpp b/Engine/Plugins/Runtime/ModelViewViewModel/Source/ModelViewViewModelBlueprint/Private/MVVMBlueprintViewConversionFunction.cpp
--- a/Engine/Plugins/Runtime/ModelViewViewModel/Source/ModelViewViewModelBlueprint/Private/MVVMBlueprintViewConversionFunction.cpp
+++ b/Engine/Plugins/Runtime/ModelViewViewModel/Source/ModelViewViewModelBlueprint/Private/MVVMBlueprintViewConversionFunction.cpp
@@ -550,7 +550,7 @@ void UMVVMBlueprintViewConversionFunction::RemoveWrapperGraph(UBlueprint* Bluepr
 	}
 	else
 	{
-		if (UObject* ExistingObject = StaticFindObject(nullptr, Blueprint, *GraphName.ToString(), true))
+		if (UObject* ExistingObject = StaticFindObject(nullptr, Blueprint, *GraphName.ToString(), EFindObjectFlags::ExactClass))
 		{
 			UE::MVVM::FunctionGraphHelper::RenameObjectToTransientPackage(ExistingObject);
 		}
diff --git a/Engine/Plugins/Runtime/ModelViewViewModel/Source/ModelViewViewModelBlueprint/Private/MVVMViewBlueprintCompiler.cpp b/Engine/Plugins/Runtime/ModelViewViewModel/Source/ModelViewViewModelBlueprint/Private/MVVMViewBlueprintCompiler.cpp
--- a/Engine/Plugins/Runtime/ModelViewViewModel/Source/ModelViewViewModelBlueprint/Private/MVVMViewBlueprintCompiler.cpp
+++ b/Engine/Plugins/Runtime/ModelViewViewModel/Source/ModelViewViewModelBlueprint/Private/MVVMViewBlueprintCompiler.cpp
@@ -4767,7 +4767,7 @@ void FMVVMViewBlueprintCompiler::TestGenerateSetter(const UBlueprint* Context, F
 #if UE_BUILD_DEVELOPMENT || UE_BUILD_DEBUG
 	UWidgetBlueprint* WidgetBlueprint = nullptr;
 	{
-		UObject* FoundObject = FindObject<UObject>(nullptr, ObjectName.GetData(), false);
+		UObject* FoundObject = FindObject<UObject>(nullptr, ObjectName.GetData(), EFindObjectFlags::None);
 		WidgetBlueprint = Cast<UWidgetBlueprint>(FoundObject);
 	}
 
diff --git a/Engine/Plugins/Runtime/ModelViewViewModel/Source/ModelViewViewModelBlueprint/Private/MVVMWidgetBlueprintExtension_View.cpp b/Engine/Plugins/Runtime/ModelViewViewModel/Source/ModelViewViewModelBlueprint/Private/MVVMWidgetBlueprintExtension_View.cpp
--- a/Engine/Plugins/Runtime/ModelViewViewModel/Source/ModelViewViewModelBlueprint/Private/MVVMWidgetBlueprintExtension_View.cpp
+++ b/Engine/Plugins/Runtime/ModelViewViewModel/Source/ModelViewViewModelBlueprint/Private/MVVMWidgetBlueprintExtension_View.cpp
@@ -344,7 +344,7 @@ void UMVVMWidgetBlueprintExtension_View::HandleFinishCompilingClass(UWidgetBluep
 		if (CurrentCompilerContext->PreCompile(Class))
 		{
 			FName ClassName = "ViewClass";
-			if (UObject* PreviousObj = StaticFindObjectFastInternal(nullptr, Class, ClassName, true))
+			if (UObject* PreviousObj = StaticFindObjectFastInternal(nullptr, Class, ClassName, EFindObjectFlags::ExactClass))
 			{
 				// Remove previous object.
 				ERenameFlags RenameFlags = REN_NonTransactional | REN_DoNotDirty | REN_DontCreateRedirectors;
diff --git a/Engine/Plugins/Runtime/RigVM/Source/RigVMDeveloper/Private/EdGraph/RigVMEdGraph.cpp b/Engine/Plugins/Runtime/RigVM/Source/RigVMDeveloper/Private/EdGraph/RigVMEdGraph.cpp
--- a/Engine/Plugins/Runtime/RigVM/Source/RigVMDeveloper/Private/EdGraph/RigVMEdGraph.cpp
+++ b/Engine/Plugins/Runtime/RigVM/Source/RigVMDeveloper/Private/EdGraph/RigVMEdGraph.cpp
@@ -1110,7 +1110,7 @@ void URigVMEdGraph::RemoveNode(UEdGraphNode* InNode)
 		do
 		{
 			DeletedName = FString::Printf(TEXT("EdGraph_%s_Deleted_%d"), *InNode->GetName(), DeletedIndex++); 
-			ExistingObject = StaticFindObject(/*Class=*/ NULL, this, *DeletedName, true);						
+			ExistingObject = StaticFindObject(/*Class=*/ nullptr, this, *DeletedName, EFindObjectFlags::ExactClass);
 		}
 		while (ExistingObject);
 	}
diff --git a/Engine/Plugins/Runtime/USDCore/Source/USDClasses/Private/USDAssetCache2.cpp b/Engine/Plugins/Runtime/USDCore/Source/USDClasses/Private/USDAssetCache2.cpp
--- a/Engine/Plugins/Runtime/USDCore/Source/USDClasses/Private/USDAssetCache2.cpp
+++ b/Engine/Plugins/Runtime/USDCore/Source/USDClasses/Private/USDAssetCache2.cpp
@@ -163,7 +163,7 @@ namespace UE::AssetCache::Private
 				else
 				{
 					// look up the object by fully qualified pathname
-					Obj = FindObject<UObject>(nullptr, *LoadedString, false);
+					Obj = FindObject<UObject>(nullptr, *LoadedString, EFindObjectFlags::None);
 					// If we couldn't find it, and we want to load it, do that
 					if (!Obj && bLoadIfFindFails)
 					{
diff --git a/Engine/Source/Editor/Persona/Private/PersonaMeshDetails.cpp b/Engine/Source/Editor/Persona/Private/PersonaMeshDetails.cpp
--- a/Engine/Source/Editor/Persona/Private/PersonaMeshDetails.cpp
+++ b/Engine/Source/Editor/Persona/Private/PersonaMeshDetails.cpp
@@ -584,15 +584,15 @@ void FSkeletalMeshReductionSettingsLayout::GenerateChildContent(IDetailChildrenB
 	//Get the Enums
 	if (EnumReductionMethod == nullptr)
 	{
-		EnumReductionMethod = FindObject<UEnum>(nullptr, TEXT("/Script/Engine.SkeletalMeshOptimizationType"), true);
+		EnumReductionMethod = FindObject<UEnum>(nullptr, TEXT("/Script/Engine.SkeletalMeshOptimizationType"), EFindObjectFlags::ExactClass);
 	}
 	if (EnumImportance == nullptr)
 	{
-		EnumImportance = FindObject<UEnum>(nullptr, TEXT("/Script/Engine.SkeletalMeshOptimizationImportance"), true);
+		EnumImportance = FindObject<UEnum>(nullptr, TEXT("/Script/Engine.SkeletalMeshOptimizationImportance"), EFindObjectFlags::ExactClass);
 	}
 	if (EnumTerminationCriterion == nullptr)
 	{
-		EnumTerminationCriterion = FindObject<UEnum>(nullptr, TEXT("/Script/Engine.SkeletalMeshTerminationCriterion"), true);
+		EnumTerminationCriterion = FindObject<UEnum>(nullptr, TEXT("/Script/Engine.SkeletalMeshTerminationCriterion"), EFindObjectFlags::ExactClass);
 	}
 
 	bool bUseThirdPartyUI = !UseNativeReductionTool();
diff --git a/Engine/Source/Editor/UMGEditor/Private/TabFactory/AnimationTabSummoner.cpp b/Engine/Source/Editor/UMGEditor/Private/TabFactory/AnimationTabSummoner.cpp
--- a/Engine/Source/Editor/UMGEditor/Private/TabFactory/AnimationTabSummoner.cpp
+++ b/Engine/Source/Editor/UMGEditor/Private/TabFactory/AnimationTabSummoner.cpp
@@ -86,7 +86,7 @@ bool VerifyAnimationRename( FWidgetBlueprintEditor& BlueprintEditor, UWidgetAnim
 	}
 
 	UWidgetBlueprint* Blueprint = BlueprintEditor.GetWidgetBlueprintObj();
-	if (Blueprint && FindObject<UWidgetAnimation>( Blueprint, *NewAnimationName, true ) )
+	if (Blueprint && FindObject<UWidgetAnimation>( Blueprint, *NewAnimationName, EFindObjectFlags::ExactClass ) )
 	{
 		OutErrorMessage = LOCTEXT( "NameInUseByAnimation", "An animation with this name already exists" );
 		return false;
