MVVM: Paste only the destination path for Events and Conditions. Retain event/condition path

#rnx

diff --git a/Engine/Plugins/Runtime/ModelViewViewModel/Source/ModelViewViewModelBlueprint/Private/MVVMBlueprintView.cpp b/Engine/Plugins/Runtime/ModelViewViewModel/Source/ModelViewViewModelBlueprint/Private/MVVMBlueprintView.cpp
--- a/Engine/Plugins/Runtime/ModelViewViewModel/Source/ModelViewViewModelBlueprint/Private/MVVMBlueprintView.cpp
+++ b/Engine/Plugins/Runtime/ModelViewViewModel/Source/ModelViewViewModelBlueprint/Private/MVVMBlueprintView.cpp
@@ -309,8 +309,28 @@ void UMVVMBlueprintView::RemoveEvent(UMVVMBlueprintViewEvent* Event)
 	}
 }
 
+void UMVVMBlueprintView::ReplaceEvent(UMVVMBlueprintViewEvent* OldEvent, UMVVMBlueprintViewEvent* NewEvent)
+{
+	for (TObjectPtr<UMVVMBlueprintViewEvent>& Event : Events)
+	{
+		if (Event == OldEvent)
+		{
+			Event = NewEvent;
+			OldEvent->RemoveWrapperGraph();
+			OnBindingsUpdated.Broadcast();
+
+			FBlueprintEditorUtils::MarkBlueprintAsModified(GetOuterUMVVMWidgetBlueprintExtension_View()->GetWidgetBlueprint());
+		}
+	}
+}
+
 UMVVMBlueprintViewEvent* UMVVMBlueprintView::DuplicateEvent(UMVVMBlueprintViewEvent* Event)
 {
+	if (Event != nullptr)
+	{
+		Event->SavePinValues();
+	}
+
 	UMVVMBlueprintViewEvent* NewEvent = DuplicateObject<UMVVMBlueprintViewEvent>(Event, this);
 	NewEvent->RecreateWrapperGraph();
 	
@@ -348,8 +368,28 @@ void UMVVMBlueprintView::RemoveCondition(UMVVMBlueprintViewCondition* Condition)
 	}
 }
 
+void UMVVMBlueprintView::ReplaceCondition(UMVVMBlueprintViewCondition* OldCondition, UMVVMBlueprintViewCondition* NewCondition)
+{
+	for (TObjectPtr<UMVVMBlueprintViewCondition>& Condition : Conditions)
+	{
+		if (Condition == OldCondition)
+		{
+			Condition = NewCondition;
+			OldCondition->RemoveWrapperGraph();
+			OnBindingsUpdated.Broadcast();
+
+			FBlueprintEditorUtils::MarkBlueprintAsModified(GetOuterUMVVMWidgetBlueprintExtension_View()->GetWidgetBlueprint());
+		}
+	}
+}
+
 UMVVMBlueprintViewCondition* UMVVMBlueprintView::DuplicateCondition(UMVVMBlueprintViewCondition* Condition)
 {
+	if (Condition != nullptr)
+	{
+		Condition->SavePinValues();
+	}
+
 	UMVVMBlueprintViewCondition* NewCondition = DuplicateObject<UMVVMBlueprintViewCondition>(Condition, this);
 	NewCondition->RecreateWrapperGraph();
 
diff --git a/Engine/Plugins/Runtime/ModelViewViewModel/Source/ModelViewViewModelBlueprint/Public/MVVMBlueprintView.h b/Engine/Plugins/Runtime/ModelViewViewModel/Source/ModelViewViewModelBlueprint/Public/MVVMBlueprintView.h
--- a/Engine/Plugins/Runtime/ModelViewViewModel/Source/ModelViewViewModelBlueprint/Public/MVVMBlueprintView.h
+++ b/Engine/Plugins/Runtime/ModelViewViewModel/Source/ModelViewViewModelBlueprint/Public/MVVMBlueprintView.h
@@ -138,6 +138,7 @@ public:
 	UMVVMBlueprintViewEvent* AddDefaultEvent();
 	void AddEvent(UMVVMBlueprintViewEvent* Event);
 	void RemoveEvent(UMVVMBlueprintViewEvent* Event);
+	void ReplaceEvent(UMVVMBlueprintViewEvent* OldEvent, UMVVMBlueprintViewEvent* NewEvent);
 	UMVVMBlueprintViewEvent* DuplicateEvent(UMVVMBlueprintViewEvent* Event);
 
 	TArrayView<TObjectPtr<UMVVMBlueprintViewEvent>> GetEvents()
@@ -153,6 +154,7 @@ public:
 	UMVVMBlueprintViewCondition* AddDefaultCondition();
 	void AddCondition(UMVVMBlueprintViewCondition* Condition);
 	void RemoveCondition(UMVVMBlueprintViewCondition* Condition);
+	void ReplaceCondition(UMVVMBlueprintViewCondition* OldCondition, UMVVMBlueprintViewCondition* NewCondition);
 	UMVVMBlueprintViewCondition* DuplicateCondition(UMVVMBlueprintViewCondition* Condition);
 
 	TArrayView<TObjectPtr<UMVVMBlueprintViewCondition>> GetConditions()
diff --git a/Engine/Plugins/Runtime/ModelViewViewModel/Source/ModelViewViewModelEditor/Private/Framework/MVVMRowHelper.cpp b/Engine/Plugins/Runtime/ModelViewViewModel/Source/ModelViewViewModelEditor/Private/Framework/MVVMRowHelper.cpp
--- a/Engine/Plugins/Runtime/ModelViewViewModel/Source/ModelViewViewModelEditor/Private/Framework/MVVMRowHelper.cpp
+++ b/Engine/Plugins/Runtime/ModelViewViewModel/Source/ModelViewViewModelEditor/Private/Framework/MVVMRowHelper.cpp
@@ -451,7 +451,7 @@ void FRowHelper::DuplicateEntries(const UWidgetBlueprint* WidgetBlueprint, UMVVM
 	}
 }
 
-void FRowHelper::CopyEntries(const UWidgetBlueprint* WidgetBlueprint, UMVVMBlueprintView* BlueprintView, const TArray<TSharedPtr<FBindingEntry>>& Entries)
+void FRowHelper::CopyEntries(UWidgetBlueprint* WidgetBlueprint, UMVVMBlueprintView* BlueprintView, const TArray<TSharedPtr<FBindingEntry>>& Entries)
 {
 	if (WidgetBlueprint == nullptr || BlueprintView == nullptr || Entries.IsEmpty())
 	{
@@ -466,7 +466,7 @@ void FRowHelper::CopyEntries(const UWidgetBlueprint* WidgetBlueprint, UMVVMBluep
 	
 	FString CopyText;
 
-	for (const Private::FBindingClipboardData::Type& Item : ClipboardData.Items)
+	for (Private::FBindingClipboardData::Type Item : ClipboardData.Items)
 	{
 		bool bFailed = false;
 		UScriptStruct* Struct = nullptr;
@@ -475,16 +475,25 @@ void FRowHelper::CopyEntries(const UWidgetBlueprint* WidgetBlueprint, UMVVMBluep
 
 		if (Item.IsType<FMVVMBlueprintViewBinding>())
 		{
-			StructData = &Item.Get<FMVVMBlueprintViewBinding>();
+			FMVVMBlueprintViewBinding* Binding = &Item.Get<FMVVMBlueprintViewBinding>();
+			Binding->Conversion.SavePinValues(WidgetBlueprint);
+
+			StructData = Binding;
 			Struct = FMVVMBlueprintViewBinding::StaticStruct();
 		}
 		else if (Item.IsType<UMVVMBlueprintViewCondition*>())
 		{
-			Object = Item.Get<UMVVMBlueprintViewCondition*>();
+			UMVVMBlueprintViewCondition* Condition = Item.Get<UMVVMBlueprintViewCondition*>();
+			Condition->SavePinValues();
+
+			Object = Condition;
 		}
 		else if (Item.IsType<UMVVMBlueprintViewEvent*>())
 		{
-			Object = Item.Get<UMVVMBlueprintViewEvent*>();
+			UMVVMBlueprintViewEvent* Event = Item.Get<UMVVMBlueprintViewEvent*>();
+			Event->SavePinValues();
+
+			Object = Event;
 		}
 
 		if (Struct != nullptr && StructData != nullptr)
@@ -664,18 +673,16 @@ void FRowHelper::PasteEntries(UWidgetBlueprint* WidgetBlueprint, UMVVMBlueprintV
 		FScopedTransaction Transaction(LOCTEXT("PasteConditionTransaction", "Paste Condition"));
 		BlueprintView->Modify();
 
-		// Update the pasted condition's path and replace selected condition with the condition being pasted.
+		// Retain the condition path when pasting. Copy over all other properties
 		UMVVMBlueprintViewCondition* ConditionSelected = Entry->GetCondition();
 		UMVVMBlueprintViewCondition* ConditionToPaste = Item.Get<UMVVMBlueprintViewCondition*>();
+		UMVVMBlueprintViewCondition* NewCondition = DuplicateObject<UMVVMBlueprintViewCondition>(ConditionToPaste, BlueprintView);
 
-		const FMVVMBlueprintPropertyPath& DesiredPath = ConditionToPaste->GetConditionPath();
-		const FMVVMBlueprintPropertyPath& TargetPath = ConditionSelected->GetConditionPath();
-		const FMVVMBlueprintPropertyPath NewPath = Private::TryGetCommonPropertyPath(WidgetBlueprint, BlueprintView, TargetPath, DesiredPath);
-
-		EditorSubsystem->SetConditionPath(ConditionToPaste, NewPath, false);
-
-		BlueprintView->RemoveCondition(ConditionSelected);
-		BlueprintView->AddCondition(ConditionToPaste);
+		NewCondition->SetConditionPath(ConditionSelected->GetConditionPath());
+		NewCondition->SetOperation(ConditionSelected->GetOperation());
+		NewCondition->SetOperationValue(ConditionSelected->GetOperationValue());
+		NewCondition->SetOperationMaxValue(ConditionSelected->GetOperationMaxValue());
+		BlueprintView->ReplaceCondition(ConditionSelected, NewCondition);
 	}
 	else if (Item.IsType<UMVVMBlueprintViewEvent*>())
 	{
@@ -688,18 +695,14 @@ void FRowHelper::PasteEntries(UWidgetBlueprint* WidgetBlueprint, UMVVMBlueprintV
 		FScopedTransaction Transaction(LOCTEXT("PasteEventTransaction", "Paste Event"));
 		BlueprintView->Modify();
 
-		// Update the pasted event's path and replace selected event with the event being pasted.
+		// Retain the event path when pasting. Copy over all other properties
 		UMVVMBlueprintViewEvent* EventSelected = Entry->GetEvent();
 		UMVVMBlueprintViewEvent* EventToPaste = Item.Get<UMVVMBlueprintViewEvent*>();
+		UMVVMBlueprintViewEvent* NewEvent = DuplicateObject<UMVVMBlueprintViewEvent>(EventToPaste, BlueprintView);
 
-		const FMVVMBlueprintPropertyPath& DesiredPath = EventToPaste->GetEventPath();
-		const FMVVMBlueprintPropertyPath& TargetPath = EventSelected->GetEventPath();
-		const FMVVMBlueprintPropertyPath NewPath = Private::TryGetCommonPropertyPath(WidgetBlueprint, BlueprintView, TargetPath, DesiredPath);
-
-		EditorSubsystem->SetEventPath(EventToPaste, NewPath, false);
-
-		BlueprintView->RemoveEvent(EventSelected);
-		BlueprintView->AddEvent(EventToPaste);
+		const FMVVMBlueprintPropertyPath EventPath = EventSelected->GetEventPath();
+		NewEvent->SetEventPath(EventPath);
+		BlueprintView->ReplaceEvent(EventSelected, NewEvent);
 	}
 }
 
diff --git a/Engine/Plugins/Runtime/ModelViewViewModel/Source/ModelViewViewModelEditor/Private/Framework/MVVMRowHelper.h b/Engine/Plugins/Runtime/ModelViewViewModel/Source/ModelViewViewModelEditor/Private/Framework/MVVMRowHelper.h
--- a/Engine/Plugins/Runtime/ModelViewViewModel/Source/ModelViewViewModelEditor/Private/Framework/MVVMRowHelper.h
+++ b/Engine/Plugins/Runtime/ModelViewViewModel/Source/ModelViewViewModelEditor/Private/Framework/MVVMRowHelper.h
@@ -30,7 +30,7 @@ struct FRowHelper
 	static void DeleteEntries(const UWidgetBlueprint* WidgetBlueprint, UMVVMBlueprintView* BlueprintView, TArrayView<const TSharedPtr<FBindingEntry>> Selection);
 	static void DuplicateEntries(const UWidgetBlueprint* WidgetBlueprint, UMVVMBlueprintView* BlueprintView, TArrayView<const TSharedPtr<FBindingEntry>> Selection, TArray<const TSharedPtr<FBindingEntry>>& OutSelection);
 
-	static void CopyEntries(const UWidgetBlueprint* WidgetBlueprint, UMVVMBlueprintView* BlueprintView, const TArray<TSharedPtr<FBindingEntry>>& Entries);
+	static void CopyEntries(UWidgetBlueprint* WidgetBlueprint, UMVVMBlueprintView* BlueprintView, const TArray<TSharedPtr<FBindingEntry>>& Entries);
 	static void PasteEntries(UWidgetBlueprint* WidgetBlueprint, UMVVMBlueprintView* BlueprintView, const TArray<TSharedPtr<FBindingEntry>>& Entries);
 
 	static bool HasBlueprintGraph(UMVVMBlueprintView* BlueprintView, const TSharedPtr<FBindingEntry> Entry);
