MVVM Conversion Functions: don't include the functions from the skeleton class when including blueprint function libraries, otherwise you end up with duplicate entries.
And some extra debug info to make it easier to track down problems.

#rnx

diff --git a/Engine/Plugins/Runtime/ModelViewViewModel/Source/ModelViewViewModelBlueprint/Private/Bindings/MVVMConversionFunctionHelper.cpp b/Engine/Plugins/Runtime/ModelViewViewModel/Source/ModelViewViewModelBlueprint/Private/Bindings/MVVMConversionFunctionHelper.cpp
--- a/Engine/Plugins/Runtime/ModelViewViewModel/Source/ModelViewViewModelBlueprint/Private/Bindings/MVVMConversionFunctionHelper.cpp
+++ b/Engine/Plugins/Runtime/ModelViewViewModel/Source/ModelViewViewModelBlueprint/Private/Bindings/MVVMConversionFunctionHelper.cpp
@@ -453,6 +453,7 @@ namespace NamedNodes
 					UMVVMBlueprintView* View = Private::GetView(Blueprint);
 					const FMVVMBlueprintViewModelContext* Context = View ? View->FindViewModel(PropertyPath.GetViewModelId()) : nullptr;
 					RootProperty = Context ? Blueprint->SkeletonGeneratedClass->FindPropertyByName(Context->GetViewModelName()) : nullptr;
+					ensureAlwaysMsgf(RootProperty, TEXT("Could not find root property %s in blueprint %s"), *Context->GetViewModelName().ToString(), *GetNameSafe(Blueprint));
 					break;
 				}
 				case EMVVMBlueprintFieldPathSource::Widget:
@@ -467,6 +468,7 @@ namespace NamedNodes
 					}
 
 					RootProperty = Blueprint->SkeletonGeneratedClass->FindPropertyByName(PropertyPath.GetWidgetName());
+					ensureAlwaysMsgf(RootProperty, TEXT("Could not find root property %s in blueprint %s"), *PropertyPath.GetWidgetName().ToString(), *GetNameSafe(Blueprint));
 					break;
 				default:
 					check(false);
diff --git a/Engine/Plugins/Runtime/ModelViewViewModel/Source/ModelViewViewModelEditor/Private/Framework/MVVMConversionFunctionLibrary.cpp b/Engine/Plugins/Runtime/ModelViewViewModel/Source/ModelViewViewModelEditor/Private/Framework/MVVMConversionFunctionLibrary.cpp
--- a/Engine/Plugins/Runtime/ModelViewViewModel/Source/ModelViewViewModelEditor/Private/Framework/MVVMConversionFunctionLibrary.cpp
+++ b/Engine/Plugins/Runtime/ModelViewViewModel/Source/ModelViewViewModelEditor/Private/Framework/MVVMConversionFunctionLibrary.cpp
@@ -241,6 +241,12 @@ bool FCollection::IsClassSupported(const TArray<const UClass*>& AllowClasses, co
 		return false;
 	}
 
+	// Ignore skeleton classes
+	if (FKismetEditorUtilities::IsClassABlueprintSkeleton(Class))
+	{
+		return false;
+	}
+
 	if (!Class->HasAnyClassFlags(CLASS_Deprecated | CLASS_Hidden | CLASS_NewerVersionExists) && Class->GetPackage() != GetTransientPackage())
 	{
 		if (Class->IsChildOf(UUserWidget::StaticClass()))
