MVVM: Custom widget for asset diff bindings.

[RN] minor creator/ui

diff --git a/Engine/Plugins/Runtime/ModelViewViewModel/Source/ModelViewViewModelBlueprint/Private/MVVMWidgetBlueprintDiff.cpp b/Engine/Plugins/Runtime/ModelViewViewModel/Source/ModelViewViewModelBlueprint/Private/MVVMWidgetBlueprintDiff.cpp
new file mode 100644
--- /dev/null
+++ b/Engine/Plugins/Runtime/ModelViewViewModel/Source/ModelViewViewModelBlueprint/Private/MVVMWidgetBlueprintDiff.cpp
@@ -0,0 +1,695 @@
+// Copyright Epic Games, Inc. All Rights Reserved.
+
+#include "MVVMWidgetBlueprintDiff.h"
+
+#if WITH_EDITOR
+
+#include "Blueprint/WidgetTree.h"
+#include "Bindings/MVVMConversionFunctionHelper.h"
+#include "Extensions/MVVMBlueprintViewExtension.h"
+#include "MVVMBlueprintInstancedViewModel.h"
+#include "MVVMBlueprintView.h"
+#include "MVVMBlueprintViewModel.h"
+#include "MVVMBlueprintViewConversionFunction.h"
+#include "MVVMBlueprintViewBinding.h"
+#include "MVVMBlueprintViewCondition.h"
+#include "MVVMBlueprintViewEvent.h"
+#include "MVVMPropertyPath.h"
+#include "MVVMWidgetBlueprintExtension_View.h"
+#include "WidgetBlueprint.h"
+
+#include "DiffResults.h"
+
+#define LOCTEXT_NAMESPACE "MVVMWidgetBlueprintDiff"
+
+namespace UE::MVVM::Private
+{
+	TSharedPtr<UE::MVVM::FMVVMDiffCustomObjectProvider> DiffCustomObjectProvider;
+
+	struct FMVVMBlueprintViewDiffDetails
+	{
+		const UWidgetBlueprint* OldBlueprint;
+		const UWidgetBlueprint* NewBlueprint;
+		UMVVMBlueprintView* OldBlueprintView;
+		UMVVMBlueprintView* NewBlueprintView;
+		FDiffResults* ResultsPtr;
+	};
+
+	enum class EMVVMBindingPropertyPathType
+	{
+		Source,
+		Destination,
+		Event,
+		Condition
+	};
+
+	FText GetPathSourceText(const UWidgetBlueprint* WidgetBlueprint, const FMVVMBlueprintPropertyPath& Path)
+	{
+		switch (Path.GetSource(WidgetBlueprint))
+		{
+		case EMVVMBlueprintFieldPathSource::SelfContext:
+		{
+			return FText::FromString(WidgetBlueprint->GetFriendlyName());
+		}
+		case EMVVMBlueprintFieldPathSource::ViewModel:
+		{
+			UMVVMWidgetBlueprintExtension_View* ExtensionView = UMVVMWidgetBlueprintExtension_View::GetExtension<UMVVMWidgetBlueprintExtension_View>(WidgetBlueprint);
+			UMVVMBlueprintView* BlueprintView = ExtensionView ? ExtensionView->GetBlueprintView() : nullptr;
+			const FMVVMBlueprintViewModelContext* ViewModel = BlueprintView ? BlueprintView->FindViewModel(Path.GetViewModelId()) : nullptr;
+			return ViewModel ? ViewModel->GetDisplayName() : LOCTEXT("None", "<None>");
+		}
+		case EMVVMBlueprintFieldPathSource::Widget:
+		{
+			UWidgetTree* WidgetTree = WidgetBlueprint ? WidgetBlueprint->WidgetTree : nullptr;
+			UWidget* FoundWidget = WidgetTree ? WidgetTree->FindWidget(Path.GetWidgetName()) : nullptr;
+			return FoundWidget ? FoundWidget->GetLabelText() : FText::FromName(Path.GetWidgetName());
+		}
+		}
+		return FText::GetEmpty();
+	}
+
+	template <class TViewProperty>
+	FText GetContextText(const FMVVMBlueprintViewDiffDetails& Details, TViewProperty NewProperty, TViewProperty OldProperty)
+	{
+		FText NewContext;
+		FText OldContext;
+
+		if constexpr (std::is_same_v<TViewProperty, FMVVMBlueprintViewBinding*>)
+		{
+			NewContext = NewProperty != nullptr ? GetPathSourceText(Details.NewBlueprint, NewProperty->DestinationPath) : FText::GetEmpty();
+			OldContext = OldProperty != nullptr ? GetPathSourceText(Details.OldBlueprint, OldProperty->DestinationPath) : FText::GetEmpty();
+		}
+		else if constexpr (std::is_same_v<TViewProperty, const UMVVMBlueprintViewCondition*>)
+		{
+			NewContext = NewProperty != nullptr ? GetPathSourceText(Details.NewBlueprint, NewProperty->GetConditionPath()) : FText::GetEmpty();
+			OldContext = OldProperty != nullptr ? GetPathSourceText(Details.OldBlueprint, OldProperty->GetConditionPath()) : FText::GetEmpty();
+		}
+		else if constexpr (std::is_same_v<TViewProperty, const UMVVMBlueprintViewEvent*>)
+		{
+			NewContext = NewProperty != nullptr ? GetPathSourceText(Details.NewBlueprint, NewProperty->GetEventPath()) : FText::GetEmpty();
+			OldContext = OldProperty != nullptr ? GetPathSourceText(Details.OldBlueprint, OldProperty->GetEventPath()) : FText::GetEmpty();
+		}
+
+		// Get the context of the diff. If both properties are valid, the context may differ i.e. destination binding changes between widgets. 
+		// In this case the diff details will denote the change, and context is not necessarily relevant. 
+		if (NewContext.IsEmpty())
+		{
+			return OldContext;
+		}
+		else if (OldContext.IsEmpty() || NewContext.EqualTo(OldContext))
+		{
+			return NewContext;
+
+		}
+		return FText::GetEmpty();
+	}
+
+	template <class TViewProperty>
+	void ApplyDiff(const FMVVMBlueprintViewDiffDetails& Details, TViewProperty NewProperty, TViewProperty OldProperty, EDiffType::Category Category, const FText& DisplayString)
+	{
+		FDiffSingleResult Diff;
+		Diff.Category = Category;
+		Diff.Object1 = Details.NewBlueprintView;
+		Diff.Object2 = Details.OldBlueprintView;
+		Diff.DisplayString = DisplayString;
+
+		if (!DiffCustomObjectProvider.IsValid())
+		{
+			Diff.Diff = EDiffType::OBJECT_PROPERTY;
+		}
+		else
+		{
+			Diff.Diff = EDiffType::CUSTOM_OBJECT;
+
+			if constexpr (std::is_same_v<TViewProperty, FMVVMBlueprintViewBinding*>)
+			{
+				const FGuid NewId = NewProperty ? NewProperty->BindingId : FGuid();
+				const FGuid OldId = OldProperty ? OldProperty->BindingId : FGuid();
+				Diff.CustomProperty = DiffCustomObjectProvider->CreatePropertyBindingDiff(NewId, OldId);
+			}
+			else if constexpr (std::is_same_v<TViewProperty, FMVVMBlueprintPin>)
+			{
+				Diff.CustomProperty = DiffCustomObjectProvider->CreatePropertyParameterDiff(NewProperty.GetId(), OldProperty.GetId());
+			}
+			else if constexpr (std::is_same_v<TViewProperty, const UMVVMBlueprintViewCondition*>)
+			{
+				Diff.CustomProperty = DiffCustomObjectProvider->CreatePropertyConditionDiff(NewProperty, OldProperty);
+			}
+			else if constexpr (std::is_same_v<TViewProperty, const UMVVMBlueprintViewEvent*>)
+			{
+				Diff.CustomProperty = DiffCustomObjectProvider->CreatePropertyEventDiff(NewProperty, OldProperty);
+			}
+		}
+		
+		Details.ResultsPtr->Add(Diff);
+	}
+
+	template <class TViewProperty>
+	bool DiffPropertyPath(FMVVMBlueprintViewDiffDetails& Details, TViewProperty NewProperty, TViewProperty OldProperty, const FMVVMBlueprintPropertyPath& OldPath, const FMVVMBlueprintPropertyPath& NewPath, EMVVMBindingPropertyPathType Type)
+	{
+		if (OldPath == NewPath)
+		{
+			return false;
+		}
+		FText TypeText;
+		switch (Type)
+		{		
+		case EMVVMBindingPropertyPathType::Source:	 
+			TypeText = FText::FromString(TEXT("Source"));  
+			break;
+		case EMVVMBindingPropertyPathType::Destination:
+			TypeText = FText::FromString(TEXT("Destination"));  
+			break;
+		case EMVVMBindingPropertyPathType::Event:
+			TypeText = FText::FromString(TEXT("Event"));  
+			break;
+		case EMVVMBindingPropertyPathType::Condition:
+			TypeText = FText::FromString(TEXT("Condition"));  
+			break;
+		}
+
+		FFormatNamedArguments Args;
+		Args.Add(TEXT("Context"), GetContextText(Details, NewProperty, OldProperty));
+		Args.Add(TEXT("Type"), TypeText);
+		const FText DisplayString = FText::Format(LOCTEXT("MVVM_ChangedPathDisplay", "[{Context}] {Type} changed"), Args);
+
+		ApplyDiff(Details, NewProperty, OldProperty, EDiffType::MODIFICATION, DisplayString);
+
+		return true;
+	}
+
+	using FPinPair = TTuple<FMVVMBlueprintPin, UEdGraphPin*>;
+	void DiffPins(FMVVMBlueprintViewDiffDetails& Details, const TMap<FString, FPinPair>& OldPins, const TMap<FString, FPinPair>& NewPins, const FText& ContextText)
+	{
+		// Check if pin values changed. Do not need to check for added/removed pins. Pins are associated with blueprint conversion functions. 
+		// If the function is changed, this will be noted in the property path diff.
+		for (TPair<FString, FPinPair> Pair : OldPins)
+		{
+			const FPinPair* NewPinPairPtr = NewPins.Find(Pair.Key);
+
+			// Pins should not differ as source and destination paths are not changed.
+			if (NewPinPairPtr == nullptr)
+			{
+				continue;
+			}
+
+			const FPinPair& OldPinPair = Pair.Value;
+			const FPinPair& NewPinPair = *NewPinPairPtr;
+
+			const FMVVMBlueprintPin& OldPin = OldPinPair.Key;
+			const FMVVMBlueprintPin& NewPin = NewPinPair.Key;
+			
+			UEdGraphPin* OldGraphPin = OldPinPair.Value;
+			UEdGraphPin* NewGraphPin = NewPinPair.Value;
+
+			FString OldPinValue;
+			if (OldPin.UsedPathAsValue())
+			{
+				OldPinValue = OldPin.GetPath().ToString(Details.OldBlueprint, true, false);
+			}
+			else
+			{
+				OldPinValue = OldPin.GetValueAsString(Details.OldBlueprint ? Details.OldBlueprint->SkeletonGeneratedClass : nullptr);
+			}
+
+			FString NewPinValue;
+			if (NewPin.UsedPathAsValue())
+			{
+				NewPinValue = NewPin.GetPath().ToString(Details.NewBlueprint, true, false);
+			}
+			else
+			{
+				NewPinValue = NewPin.GetValueAsString(Details.NewBlueprint ? Details.NewBlueprint->SkeletonGeneratedClass : nullptr);
+			}
+
+			if(OldPinValue != NewPinValue)
+			{				
+				FText PinDisplayName;
+				if (OldGraphPin != nullptr)
+				{
+					PinDisplayName = OldGraphPin->GetDisplayName();
+				}
+				else
+				{
+					PinDisplayName = FText::FromString(OldPin.GetId().ToString());
+				}
+
+				FFormatNamedArguments Args;
+				Args.Add(TEXT("Context"), ContextText);
+				Args.Add(TEXT("PinName"), PinDisplayName);
+				const FText DisplayString = FText::Format(LOCTEXT("MVVM_ChangedPinDisplay", "[{Context}] Pin \"{PinName}\" changed"), Args);
+
+				ApplyDiff(Details, NewPin, OldPin, EDiffType::MODIFICATION, DisplayString);
+			}
+		}
+	}
+	
+	void DiffBindings(FMVVMBlueprintViewDiffDetails& Details)
+	{
+		TSet<FGuid> OldBindings;
+		TSet<FGuid> NewBindings;
+		TMap<FGuid, FMVVMBlueprintViewBinding> OldBindingMap;
+		TMap<FGuid, FMVVMBlueprintViewBinding> NewBindingMap;
+
+		if (Details.OldBlueprintView != nullptr)
+		{
+			for (const FMVVMBlueprintViewBinding& OldBinding : Details.OldBlueprintView->GetBindings())
+			{
+				OldBindings.Add(OldBinding.BindingId);
+				OldBindingMap.Add(OldBinding.BindingId, OldBinding);
+			}
+		}
+
+		if (Details.NewBlueprintView != nullptr)
+		{
+			for (const FMVVMBlueprintViewBinding& NewBinding : Details.NewBlueprintView->GetBindings())
+			{
+				NewBindings.Add(NewBinding.BindingId);
+				NewBindingMap.Add(NewBinding.BindingId, NewBinding);
+			}
+		}
+
+		const TSet<FGuid> SharedBindings = OldBindings.Intersect(NewBindings);
+		const TSet<FGuid> AddedBindings = NewBindings.Difference(OldBindings);
+		const TSet<FGuid> RemovedBindings = OldBindings.Difference(NewBindings);
+
+		for (const FGuid& Id: AddedBindings)
+		{
+			FMVVMBlueprintViewBinding* NewBinding = NewBindingMap.Find(Id);
+			FMVVMBlueprintViewBinding* OldBinding = nullptr;
+			if (!ensure(NewBinding != nullptr))
+			{
+				continue;
+			}
+			
+			FFormatNamedArguments Args;
+			Args.Add(TEXT("Context"), GetContextText(Details, NewBinding, OldBinding));
+			Args.Add(TEXT("BindingName"), Details.NewBlueprint ? FText::FromString(NewBinding->GetDisplayNameString(Details.NewBlueprint, true)) : FText::GetEmpty());
+			const FText DisplayString = FText::Format(LOCTEXT("MVVM_AddedBindingDisplay", "[{Context}] Added Binding \"{BindingName}\""), Args);
+			
+			ApplyDiff(Details, NewBinding, OldBinding, EDiffType::ADDITION, DisplayString);
+		}
+
+		for (const FGuid& Id: RemovedBindings)
+		{
+			FMVVMBlueprintViewBinding* OldBinding = OldBindingMap.Find(Id);
+			FMVVMBlueprintViewBinding* NewBinding = nullptr;
+			if (!ensure(OldBinding != nullptr))
+			{
+				continue;
+			}
+
+			FFormatNamedArguments Args;
+			Args.Add(TEXT("Context"), GetContextText(Details, NewBinding, OldBinding));
+			Args.Add(TEXT("BindingName"), Details.OldBlueprint ? FText::FromString(OldBinding->GetDisplayNameString(Details.OldBlueprint, true)) : FText::GetEmpty());
+			const FText DisplayString = FText::Format(LOCTEXT("MVVM_RemovedBindingDisplay", "[{Context}] Removed Binding \"{BindingName}\""), Args);
+
+			ApplyDiff(Details, NewBinding, OldBinding, EDiffType::SUBTRACTION, DisplayString);
+		}
+
+		for (const FGuid& Id: SharedBindings)
+		{
+			FMVVMBlueprintViewBinding* OldBinding = OldBindingMap.Find(Id);
+			FMVVMBlueprintViewBinding* NewBinding = NewBindingMap.Find(Id);
+			if (!ensure(OldBinding != nullptr && NewBinding != nullptr))
+			{
+				continue;
+			}
+
+			// 1. Diff property paths
+			const bool bSourcePathChanged = DiffPropertyPath(Details, NewBinding, OldBinding, OldBinding->SourcePath, NewBinding->SourcePath, EMVVMBindingPropertyPathType::Source);
+			const bool bDestinationPathChanged = DiffPropertyPath(Details, NewBinding, OldBinding, OldBinding->DestinationPath, NewBinding->DestinationPath, EMVVMBindingPropertyPathType::Destination);
+
+			// 2. Diff the binding modes
+			if (OldBinding->BindingType != NewBinding->BindingType)
+			{
+				UEnum* BindingModeEnum = StaticEnum<EMVVMBindingMode>();
+
+				FFormatNamedArguments Args;
+				Args.Add(TEXT("Context"), GetContextText(Details, NewBinding, OldBinding));
+				Args.Add(TEXT("NewMode"), BindingModeEnum->GetDisplayValueAsText(NewBinding->BindingType));
+
+				const FText DisplayString = FText::Format(LOCTEXT("MVVM_ChangedModeDisplay", "[{Context}] Binding Mode changed to \"{NewMode}\""), Args);
+
+				ApplyDiff(Details, NewBinding, OldBinding, EDiffType::MODIFICATION, DisplayString);
+			}
+
+			// 3. Diff the source and destination conversion function Pins
+			TMap<FString, FPinPair> OldPins;
+			TMap<FString, FPinPair> NewPins;
+
+			// If the source path did not changed. Check if any pins have changed
+			if (!bSourcePathChanged)
+			{
+				if (UMVVMBlueprintViewConversionFunction* OldSourceFunction = OldBinding->Conversion.GetConversionFunction(true))
+				{
+					for (const FMVVMBlueprintPin& Pin : OldSourceFunction->GetPins())
+					{
+						UEdGraphPin* GraphPin = UE::MVVM::ConversionFunctionHelper::FindPin(OldSourceFunction->GetWrapperGraph(), Pin.GetId().GetNames());
+						OldPins.Add(Pin.GetId().ToString(), MakeTuple(Pin, GraphPin));
+					}
+				}
+
+				if (UMVVMBlueprintViewConversionFunction* NewSourceFunction = NewBinding->Conversion.GetConversionFunction(true))
+				{
+					for (const FMVVMBlueprintPin& Pin : NewSourceFunction->GetPins())
+					{
+						UEdGraphPin* GraphPin = UE::MVVM::ConversionFunctionHelper::FindPin(NewSourceFunction->GetWrapperGraph(), Pin.GetId().GetNames());
+						NewPins.Add(Pin.GetId().ToString(), MakeTuple(Pin, GraphPin));
+					}
+				}
+			}
+
+			// If the destination path did not changed. Check if any pins have changed
+			if (!bDestinationPathChanged)
+			{
+				if (UMVVMBlueprintViewConversionFunction* OldDestinationFunction = OldBinding->Conversion.GetConversionFunction(false))
+				{
+					for (const FMVVMBlueprintPin& Pin : OldDestinationFunction->GetPins())
+					{
+						UEdGraphPin* GraphPin = UE::MVVM::ConversionFunctionHelper::FindPin(OldDestinationFunction->GetWrapperGraph(), Pin.GetId().GetNames());
+						OldPins.Add(Pin.GetId().ToString(), MakeTuple(Pin, GraphPin));
+					}
+				}
+
+				if (UMVVMBlueprintViewConversionFunction* NewDestinationFunction = NewBinding->Conversion.GetConversionFunction(false))
+				{
+					for (const FMVVMBlueprintPin& Pin : NewDestinationFunction->GetPins())
+					{
+						UEdGraphPin* GraphPin = UE::MVVM::ConversionFunctionHelper::FindPin(NewDestinationFunction->GetWrapperGraph(), Pin.GetId().GetNames());
+						NewPins.Add(Pin.GetId().ToString(), MakeTuple(Pin, GraphPin));
+					}
+				}
+			}
+
+			const FText ContextText = GetContextText(Details, NewBinding, OldBinding);
+			DiffPins(Details, OldPins, NewPins, ContextText);
+		}
+	}
+
+	void DiffEvents(FMVVMBlueprintViewDiffDetails& Details)
+	{
+		TSet<FName> OldEvents;
+		TSet<FName> NewEvents;
+		TMap<FName, const UMVVMBlueprintViewEvent*> OldEventMap;
+		TMap<FName, const UMVVMBlueprintViewEvent*> NewEventMap;
+
+		if (Details.OldBlueprintView != nullptr )
+		{
+			for (const UMVVMBlueprintViewEvent* OldEvent : Details.OldBlueprintView->GetEvents())
+			{
+				OldEvents.Add(OldEvent->GetFName());
+				OldEventMap.Add(OldEvent->GetFName(), OldEvent);
+			}
+		}
+
+		if (Details.NewBlueprintView != nullptr )
+		{
+			for (const UMVVMBlueprintViewEvent* NewEvent : Details.NewBlueprintView->GetEvents())
+			{
+				NewEvents.Add(NewEvent->GetFName());
+				NewEventMap.Add(NewEvent->GetFName(), NewEvent);
+			}
+		}
+
+		const TSet<FName> SharedEvents = OldEvents.Intersect(NewEvents);
+		const TSet<FName> AddedEvents = NewEvents.Difference(OldEvents);
+		const TSet<FName> RemovedEvents = OldEvents.Difference(NewEvents);
+
+		for (const FName& Id: AddedEvents)
+		{
+			const UMVVMBlueprintViewEvent** EventPtr = NewEventMap.Find(Id);
+			if (!ensure(EventPtr != nullptr && *EventPtr != nullptr))
+			{
+				continue;
+			}
+
+			const UMVVMBlueprintViewEvent* NewEvent = *EventPtr;
+			const UMVVMBlueprintViewEvent* OldEvent = nullptr;
+						
+			FFormatNamedArguments Args;
+			Args.Add(TEXT("Context"), GetContextText(Details, NewEvent, OldEvent));
+			Args.Add(TEXT("EventName"), NewEvent->GetDisplayName(true));			
+			const FText DisplayString = FText::Format(LOCTEXT("MVVM_AddedEventDisplay", "[{Context}] Added Event {EventName}"), Args);
+
+			ApplyDiff(Details, NewEvent, OldEvent, EDiffType::ADDITION, DisplayString);
+		}
+
+		for (const FName& Id: RemovedEvents)
+		{
+			const UMVVMBlueprintViewEvent** EventPtr = OldEventMap.Find(Id);
+			if (!ensure(EventPtr != nullptr && *EventPtr != nullptr))
+			{
+				continue;
+			}
+
+			const UMVVMBlueprintViewEvent* NewEvent = nullptr;
+			const UMVVMBlueprintViewEvent* OldEvent = *EventPtr;
+
+			FFormatNamedArguments Args;
+			Args.Add(TEXT("Context"), GetContextText(Details, NewEvent, OldEvent));
+			Args.Add(TEXT("EventName"), OldEvent->GetDisplayName(true));			
+			const FText DisplayString = FText::Format(LOCTEXT("MVVM_RemovedEventDisplay", "[{Context}] Removed Event {EventName}"), Args);
+
+			ApplyDiff(Details, NewEvent, OldEvent, EDiffType::SUBTRACTION, DisplayString);
+		}
+
+		for (const FName& Id: SharedEvents)
+		{
+			const UMVVMBlueprintViewEvent** OldEventPtr = OldEventMap.Find(Id);
+			const UMVVMBlueprintViewEvent** NewEventPtr = NewEventMap.Find(Id);
+			if (!ensure(OldEventPtr != nullptr && NewEventPtr != nullptr))
+			{
+				continue;
+			}
+
+			const UMVVMBlueprintViewEvent* OldEvent = *OldEventPtr;
+			const UMVVMBlueprintViewEvent* NewEvent = *NewEventPtr;
+			if (!ensure(OldEvent != nullptr && NewEvent != nullptr))
+			{
+				continue;
+			}
+
+			// Diff property paths
+			const bool bEventPathChanged = DiffPropertyPath(Details, NewEvent, OldEvent, OldEvent->GetEventPath(), NewEvent->GetEventPath(), EMVVMBindingPropertyPathType::Event);
+			const bool bDestinationPathChanged = DiffPropertyPath(Details, NewEvent, OldEvent, OldEvent->GetDestinationPath(), NewEvent->GetDestinationPath(), EMVVMBindingPropertyPathType::Destination);
+
+			TMap<FString, FPinPair> OldPins;
+			TMap<FString, FPinPair> NewPins;
+
+			// If the destination path did not changed. Check if any pins have changed
+			if (!bDestinationPathChanged)
+			{
+				for (const FMVVMBlueprintPin& Pin : OldEvent->GetPins())
+				{
+					UEdGraphPin* GraphPin = UE::MVVM::ConversionFunctionHelper::FindPin(OldEvent->GetWrapperGraph(), Pin.GetId().GetNames());
+					OldPins.Add(Pin.GetId().ToString(), MakeTuple(Pin, GraphPin));
+				}
+
+				for (const FMVVMBlueprintPin& Pin : NewEvent->GetPins())
+				{
+					UEdGraphPin* GraphPin = UE::MVVM::ConversionFunctionHelper::FindPin(OldEvent->GetWrapperGraph(), Pin.GetId().GetNames());
+					NewPins.Add(Pin.GetId().ToString(), MakeTuple(Pin, GraphPin));
+				}
+			}
+
+			const FText ContextText = GetContextText(Details, NewEvent, OldEvent);
+			DiffPins(Details, OldPins, NewPins, ContextText);
+		}
+	}
+
+	void DiffConditions(FMVVMBlueprintViewDiffDetails& Details)
+	{
+		TSet<FName> OldConditions;
+		TSet<FName> NewConditions;
+		TMap<FName, const UMVVMBlueprintViewCondition*> OldConditionMap;
+		TMap<FName, const UMVVMBlueprintViewCondition*> NewConditionMap;
+
+		
+		if (Details.OldBlueprintView != nullptr)
+		{
+			for (const UMVVMBlueprintViewCondition* OldCondition : Details.OldBlueprintView->GetConditions())
+			{
+				OldConditions.Add(OldCondition->GetFName());
+				OldConditionMap.Add(OldCondition->GetFName(), OldCondition);
+			}
+		}
+		
+		if (Details.NewBlueprintView != nullptr)
+		{
+			for (const UMVVMBlueprintViewCondition* NewCondition : Details.NewBlueprintView->GetConditions())
+			{
+				NewConditions.Add(NewCondition->GetFName());
+				NewConditionMap.Add(NewCondition->GetFName(), NewCondition);
+			}
+		}
+
+		const TSet<FName> SharedConditions = OldConditions.Intersect(NewConditions);
+		const TSet<FName> AddedConditions = NewConditions.Difference(OldConditions);
+		const TSet<FName> RemovedConditions = OldConditions.Difference(NewConditions);
+
+		for (const FName& Id: AddedConditions)
+		{
+			const UMVVMBlueprintViewCondition** ConditionPtr = NewConditionMap.Find(Id);
+			if (!ensure(ConditionPtr != nullptr && *ConditionPtr != nullptr))
+			{
+				continue;
+			}
+
+			const UMVVMBlueprintViewCondition* NewCondition = *ConditionPtr;
+			const UMVVMBlueprintViewCondition* OldCondition = nullptr;
+	
+			FFormatNamedArguments Args;
+			Args.Add(TEXT("Context"), GetContextText(Details, NewCondition, OldCondition));
+			Args.Add(TEXT("ConditionName"), NewCondition->GetDisplayName(true));
+			const FText DisplayString = FText::Format(LOCTEXT("MVVM_AddedConditionDisplay", "[{Context}] Added Condition {ConditionName}"), Args);
+
+			ApplyDiff(Details, NewCondition, OldCondition, EDiffType::ADDITION, DisplayString);
+		}
+
+		for (const FName& Id: RemovedConditions)
+		{
+			const UMVVMBlueprintViewCondition** ConditionPtr = OldConditionMap.Find(Id);
+			if (!ensure(ConditionPtr != nullptr && *ConditionPtr != nullptr))
+			{
+				continue;
+			}
+		
+			const UMVVMBlueprintViewCondition* OldCondition = *ConditionPtr;
+			const UMVVMBlueprintViewCondition* NewCondition = nullptr;
+
+			FFormatNamedArguments Args;
+			Args.Add(TEXT("Context"), GetContextText(Details, NewCondition, OldCondition));
+			Args.Add(TEXT("ConditionName"), OldCondition->GetDisplayName(true));
+			const FText DisplayString = FText::Format(LOCTEXT("MVVM_RemovedConditionDisplay", "[{Context}] Removed Condition {ConditionName}"), Args);
+
+			ApplyDiff(Details, NewCondition, OldCondition, EDiffType::SUBTRACTION, DisplayString);
+		}
+
+		for (const FName& Id: SharedConditions)
+		{
+			const UMVVMBlueprintViewCondition** OldConditionPtr = OldConditionMap.Find(Id);
+			const UMVVMBlueprintViewCondition** NewConditionPtr = NewConditionMap.Find(Id);
+			if (!ensure(OldConditionPtr != nullptr && NewConditionPtr != nullptr))
+			{
+				continue;
+			}
+
+			const UMVVMBlueprintViewCondition* OldCondition = * OldConditionPtr;
+			const UMVVMBlueprintViewCondition* NewCondition = *NewConditionPtr;
+			if (!ensure(OldCondition != nullptr && NewCondition != nullptr))
+			{
+				continue;
+			}
+
+			// Diff property paths
+			const bool bConditionPathChanged = DiffPropertyPath(Details, NewCondition, OldCondition, OldCondition->GetConditionPath(), NewCondition->GetConditionPath(), EMVVMBindingPropertyPathType::Condition);
+			const bool bDestinationPathChanged = DiffPropertyPath(Details, NewCondition, OldCondition, OldCondition->GetDestinationPath(), NewCondition->GetDestinationPath(), EMVVMBindingPropertyPathType::Destination);
+
+			if (OldCondition->GetOperation() != NewCondition->GetOperation())
+			{
+				UEnum* OperationEnum = StaticEnum<EMVVMConditionOperation>();
+
+				FFormatNamedArguments Args;
+				Args.Add(TEXT("Context"), GetContextText(Details, NewCondition, OldCondition));
+				Args.Add(TEXT("NewOperation"), OperationEnum->GetDisplayValueAsText(NewCondition->GetOperation()));
+				const FText DisplayString = FText::Format(LOCTEXT("MVVM_ChangedOperationDisplay", "[{Context}] Condition operation changed to \"{NewOperation}\""), Args);
+
+				ApplyDiff(Details, NewCondition, OldCondition, EDiffType::MODIFICATION, DisplayString);
+			}
+
+			if (OldCondition->GetOperationValue() != NewCondition->GetOperationValue())
+			{
+				FFormatNamedArguments Args;
+				Args.Add(TEXT("Context"), GetContextText(Details, NewCondition, OldCondition));
+				Args.Add(TEXT("NewValue"), FText::AsNumber(NewCondition->GetOperationValue()));
+				const FText DisplayString = FText::Format(LOCTEXT("MVVM_ChangedOperationValueDisplay", "[{Context}] Condition value changed to \"{NewValue}\""), Args);
+
+				ApplyDiff(Details, NewCondition, OldCondition, EDiffType::MODIFICATION, DisplayString);
+			}
+
+			if (OldCondition->GetOperationMaxValue() != NewCondition->GetOperationMaxValue())
+			{
+				FFormatNamedArguments Args;
+				Args.Add(TEXT("Context"), GetContextText(Details, NewCondition, OldCondition));
+				Args.Add(TEXT("NewValue"), FText::AsNumber(NewCondition->GetOperationMaxValue()));
+				const FText DisplayString = FText::Format(LOCTEXT("MVVM_ChangedOperationMaxValueDisplay", "[{Context}] Condition max value changed to \"{NewValue}\""), Args);
+
+				ApplyDiff(Details, NewCondition, OldCondition, EDiffType::MODIFICATION, DisplayString);
+			}
+
+			TMap<FString, FPinPair> OldPins;
+			TMap<FString, FPinPair> NewPins;
+
+			// If the destination path did not changed. Check if any pins have changed
+			if (!bDestinationPathChanged)
+			{
+				for (const FMVVMBlueprintPin& Pin : OldCondition->GetPins())
+				{
+					UEdGraphPin* GraphPin = UE::MVVM::ConversionFunctionHelper::FindPin(OldCondition->GetWrapperGraph(), Pin.GetId().GetNames());
+					OldPins.Add(Pin.GetId().ToString(), MakeTuple(Pin, GraphPin));
+				}
+
+				for (const FMVVMBlueprintPin& Pin : NewCondition->GetPins())
+				{
+					UEdGraphPin* GraphPin = UE::MVVM::ConversionFunctionHelper::FindPin(OldCondition->GetWrapperGraph(), Pin.GetId().GetNames());
+					NewPins.Add(Pin.GetId().ToString(), MakeTuple(Pin, GraphPin));
+				}
+			}
+
+			const FText ContextText = GetContextText(Details, NewCondition, OldCondition);
+			DiffPins(Details, OldPins, NewPins, ContextText);
+		}
+	}
+}
+
+#endif //WITH_EDITOR
+
+
+namespace UE::MVVM::FMVVMWidgetBlueprintDiff
+{
+void FindDiffs(const UWidgetBlueprint* NewBlueprint, const UWidgetBlueprint* OldBlueprint, FDiffResults& OutResults)
+{
+
+#if WITH_EDITOR
+	UMVVMWidgetBlueprintExtension_View* NewExtension = NewBlueprint != nullptr ? UWidgetBlueprintExtension::GetExtension<UMVVMWidgetBlueprintExtension_View>(NewBlueprint) : nullptr;
+	UMVVMWidgetBlueprintExtension_View* OldExtension = OldBlueprint != nullptr ? UWidgetBlueprintExtension::GetExtension<UMVVMWidgetBlueprintExtension_View>(OldBlueprint) : nullptr;
+
+	UE::MVVM::Private::FMVVMBlueprintViewDiffDetails Details;
+	Details.NewBlueprint = NewBlueprint;
+	Details.NewBlueprintView = NewExtension != nullptr ? NewExtension->GetBlueprintView() : nullptr;
+	Details.OldBlueprint = OldBlueprint;
+	Details.OldBlueprintView = OldExtension != nullptr ? OldExtension->GetBlueprintView() : nullptr;
+	Details.ResultsPtr = &OutResults;
+
+	if (UE::MVVM::Private::DiffCustomObjectProvider.IsValid())
+	{
+		FDiffSingleResult Diff;
+		Diff.Diff = EDiffType::CUSTOM_OBJECT;
+		Diff.Category = EDiffType::MODIFICATION;
+		Diff.Object1 = Details.NewBlueprintView;
+		Diff.Object2 = Details.OldBlueprintView;
+		Diff.DisplayString = LOCTEXT("MVVM_Bindings", "ViewModel Bindings");
+		Diff.CustomObject = UE::MVVM::Private::DiffCustomObjectProvider->CreateObjectDiff(Details.NewBlueprint, Details.OldBlueprint);
+		Details.ResultsPtr->Add(Diff);
+	}
+
+	UE::MVVM::Private::DiffBindings(Details);
+	UE::MVVM::Private::DiffConditions(Details);
+	UE::MVVM::Private::DiffEvents(Details);
+
+#endif //WITH_EDITOR
+}
+
+void RegisterCustomDiff(TSharedPtr<FMVVMDiffCustomObjectProvider> DiffCustomObjectProvider)
+{	
+	UE::MVVM::Private::DiffCustomObjectProvider = DiffCustomObjectProvider;
+}
+
+void UnregisterCustomDiff()
+{
+	UE::MVVM::Private::DiffCustomObjectProvider.Reset();
+}
+}
+
+#undef LOCTEXT_NAMESPACE
diff --git a/Engine/Plugins/Runtime/ModelViewViewModel/Source/ModelViewViewModelBlueprint/Private/MVVMWidgetBlueprintExtension_View.cpp b/Engine/Plugins/Runtime/ModelViewViewModel/Source/ModelViewViewModelBlueprint/Private/MVVMWidgetBlueprintExtension_View.cpp
--- a/Engine/Plugins/Runtime/ModelViewViewModel/Source/ModelViewViewModelBlueprint/Private/MVVMWidgetBlueprintExtension_View.cpp
+++ b/Engine/Plugins/Runtime/ModelViewViewModel/Source/ModelViewViewModelBlueprint/Private/MVVMWidgetBlueprintExtension_View.cpp
@@ -3,6 +3,7 @@
 #include "MVVMWidgetBlueprintExtension_View.h"
 
 #include "Blueprint/WidgetTree.h"
+#include "Bindings/MVVMConversionFunctionHelper.h"
 #include "Extensions/MVVMBlueprintViewExtension.h"
 #include "FindInBlueprintManager.h"
 #include "Kismet2/BlueprintEditorUtils.h"
@@ -12,10 +13,10 @@
 #include "MVVMBlueprintViewModel.h"
 #include "MVVMBlueprintViewConversionFunction.h"
 #include "MVVMViewBlueprintCompiler.h"
+#include "MVVMWidgetBlueprintDiff.h"
 #include "ScopedTransaction.h"
 #include "View/MVVMViewClass.h"
 
-
 #include UE_INLINE_GENERATED_CPP_BY_NAME(MVVMWidgetBlueprintExtension_View)
 
 #define LOCTEXT_NAMESPACE "MVVMBlueprintExtensionView"
@@ -424,6 +425,14 @@ UMVVMWidgetBlueprintExtension_View::FSearchData UMVVMWidgetBlueprintExtension_Vi
 	return SearchData;
 }
 
+#if WITH_EDITOR
+
+void UMVVMWidgetBlueprintExtension_View::FindDiffs(const UWidgetBlueprint* OwningBlueprint, const UWidgetBlueprint* OtherBlueprint, FDiffResults& Results) const
+{
+	UE::MVVM::FMVVMWidgetBlueprintDiff::FindDiffs(OwningBlueprint, OtherBlueprint, Results);
+}
+
+#endif //WITH_EDITOR
 
 void UMVVMWidgetBlueprintExtension_View::SetFilterSettings(FMVVMViewBindingFilterSettings InFilterSettings)
 {
diff --git a/Engine/Plugins/Runtime/ModelViewViewModel/Source/ModelViewViewModelBlueprint/Public/MVVMBlueprintViewCondition.h b/Engine/Plugins/Runtime/ModelViewViewModel/Source/ModelViewViewModelBlueprint/Public/MVVMBlueprintViewCondition.h
--- a/Engine/Plugins/Runtime/ModelViewViewModel/Source/ModelViewViewModelBlueprint/Public/MVVMBlueprintViewCondition.h
+++ b/Engine/Plugins/Runtime/ModelViewViewModel/Source/ModelViewViewModelBlueprint/Public/MVVMBlueprintViewCondition.h
@@ -69,13 +69,13 @@ public:
 
 	void SetOperation(EMVVMConditionOperation Operation);
 
-	float GetOperationValue()
+	float GetOperationValue() const
 	{
 		return Value;
 	}
 	void SetOperationValue(float NewValue);
 
-	float GetOperationMaxValue()
+	float GetOperationMaxValue() const
 	{
 		return MaxValue;
 	}
diff --git a/Engine/Plugins/Runtime/ModelViewViewModel/Source/ModelViewViewModelBlueprint/Public/MVVMWidgetBlueprintDiff.h b/Engine/Plugins/Runtime/ModelViewViewModel/Source/ModelViewViewModelBlueprint/Public/MVVMWidgetBlueprintDiff.h
new file mode 100644
--- /dev/null
+++ b/Engine/Plugins/Runtime/ModelViewViewModel/Source/ModelViewViewModelBlueprint/Public/MVVMWidgetBlueprintDiff.h
@@ -0,0 +1,33 @@
+// Copyright Epic Games, Inc. All Rights Reserved.
+#pragma once
+
+#include "CoreMinimal.h"
+
+struct FDiffResults;
+struct FMVVMBlueprintPinId;
+
+class IDiffCustomObject;
+class IDiffCustomProperty;
+class UWidgetBlueprint;
+class UMVVMBlueprintViewEvent;
+class UMVVMBlueprintViewCondition;
+
+namespace UE::MVVM
+{
+struct FMVVMDiffCustomObjectProvider : TSharedFromThis<FMVVMDiffCustomObjectProvider>
+{
+	virtual ~FMVVMDiffCustomObjectProvider() = default;
+	virtual TSharedPtr<IDiffCustomObject> CreateObjectDiff(const UWidgetBlueprint* NewBlueprint, const UWidgetBlueprint* OldBlueprint) { return nullptr; }
+	virtual TSharedPtr<IDiffCustomProperty> CreatePropertyBindingDiff(FGuid NewBinding, FGuid OldBinding) {return nullptr;}
+	virtual TSharedPtr<IDiffCustomProperty> CreatePropertyConditionDiff(const UMVVMBlueprintViewCondition* NewCondition, const UMVVMBlueprintViewCondition* OldCondition) {return nullptr;}
+	virtual TSharedPtr<IDiffCustomProperty> CreatePropertyEventDiff(const UMVVMBlueprintViewEvent* NewEvent, const UMVVMBlueprintViewEvent* OldEvent) {return nullptr;}
+	virtual TSharedPtr<IDiffCustomProperty> CreatePropertyParameterDiff(const FMVVMBlueprintPinId& NewParameter, const FMVVMBlueprintPinId& OldParameter) {return nullptr;}
+};
+
+namespace FMVVMWidgetBlueprintDiff
+{
+	void MODELVIEWVIEWMODELBLUEPRINT_API FindDiffs(const UWidgetBlueprint* NewBlueprint, const UWidgetBlueprint* OldBlueprint, FDiffResults& Results);
+	void MODELVIEWVIEWMODELBLUEPRINT_API RegisterCustomDiff(TSharedPtr<FMVVMDiffCustomObjectProvider> DiffCustomObject);
+	void MODELVIEWVIEWMODELBLUEPRINT_API UnregisterCustomDiff();
+}
+}
\ No newline at end of file
diff --git a/Engine/Plugins/Runtime/ModelViewViewModel/Source/ModelViewViewModelBlueprint/Public/MVVMWidgetBlueprintExtension_View.h b/Engine/Plugins/Runtime/ModelViewViewModel/Source/ModelViewViewModelBlueprint/Public/MVVMWidgetBlueprintExtension_View.h
--- a/Engine/Plugins/Runtime/ModelViewViewModel/Source/ModelViewViewModelBlueprint/Public/MVVMWidgetBlueprintExtension_View.h
+++ b/Engine/Plugins/Runtime/ModelViewViewModel/Source/ModelViewViewModelBlueprint/Public/MVVMWidgetBlueprintExtension_View.h
@@ -94,6 +94,10 @@ public:
 	virtual void HandleFinishCompilingClass(UWidgetBlueprintGeneratedClass* Class) override;
 	virtual void HandleEndCompilation() override;
 	virtual FSearchData HandleGatherSearchData(const UBlueprint* OwningBlueprint) const override;
+
+#if WITH_EDITOR
+	virtual void FindDiffs(const UWidgetBlueprint* OwningBlueprint, const UWidgetBlueprint* OtherBlueprint, FDiffResults& Results) const override;
+#endif //WITH_EDITOR
 	//~ End UWidgetBlueprintExtension interface
 
 	void VerifyWidgetExtensions();
diff --git a/Engine/Plugins/Runtime/ModelViewViewModel/Source/ModelViewViewModelEditor/Private/Customizations/MVVMWidgetBlueprintDiffProvider.h b/Engine/Plugins/Runtime/ModelViewViewModel/Source/ModelViewViewModelEditor/Private/Customizations/MVVMWidgetBlueprintDiffProvider.h
new file mode 100644
--- /dev/null
+++ b/Engine/Plugins/Runtime/ModelViewViewModel/Source/ModelViewViewModelEditor/Private/Customizations/MVVMWidgetBlueprintDiffProvider.h
@@ -0,0 +1,174 @@
+// Copyright Epic Games, Inc. All Rights Reserved.
+#pragma once 
+
+#include "DiffResults.h"
+#include "MVVMBlueprintView.h"
+#include "MVVMBlueprintViewBinding.h"
+#include "MVVMBlueprintViewCondition.h"
+#include "MVVMBlueprintViewEvent.h"
+#include "MVVMWidgetBlueprintDiff.h"
+#include "MVVMWidgetBlueprintExtension_View.h"
+#include "Widgets\SMVVMViewBindingListView.h"
+
+class FMVVMDiffCustomObject_View : public IDiffCustomObject
+{
+public:
+	TSharedPtr<UE::MVVM::SBindingsList> OldList;
+	TSharedPtr<UE::MVVM::SBindingsList> NewList;
+
+	FMVVMDiffCustomObject_View(const UWidgetBlueprint* NewBlueprint, const UWidgetBlueprint* OldBlueprint)
+	{
+		if (OldBlueprint != nullptr)
+		{
+			if (UMVVMWidgetBlueprintExtension_View* OldExtension = UWidgetBlueprintExtension::GetExtension<UMVVMWidgetBlueprintExtension_View>(OldBlueprint))
+			{
+				OldList = SNew(UE::MVVM::SBindingsList, nullptr, nullptr, OldExtension);
+			}
+		}
+
+		if (NewBlueprint != nullptr)
+		{
+			if (UMVVMWidgetBlueprintExtension_View* NewExtension = UWidgetBlueprintExtension::GetExtension<UMVVMWidgetBlueprintExtension_View>(NewBlueprint))
+			{
+				NewList = SNew(UE::MVVM::SBindingsList, nullptr, nullptr, NewExtension);
+			}
+		}
+	}
+
+	TSharedPtr<SWidget> GetOldWidget() override { return OldList; }
+	TSharedPtr<SWidget> GetNewWidget() override { return NewList; }
+};
+
+class FMVVMDiffCustomProperty_Binding : public IDiffCustomProperty
+{
+public:
+	FGuid NewBinding;
+	FGuid OldBinding;
+
+	FMVVMDiffCustomProperty_Binding(const FGuid NewBinding, const FGuid OldBinding)
+		: NewBinding(NewBinding)
+		, OldBinding(OldBinding)
+	{
+	}
+
+	void HighlightProperty(const TSharedPtr<IDiffCustomObject> Widget) override
+	{
+		if (FMVVMDiffCustomObject_View* DiffWidget = StaticCast<FMVVMDiffCustomObject_View*>(Widget.Get()))
+		{
+			if (DiffWidget->NewList.IsValid())
+			{
+				DiffWidget->NewList->SelectBinding(NewBinding, true);
+			}
+			if (DiffWidget->OldList.IsValid())
+			{
+				DiffWidget->OldList->SelectBinding(OldBinding, true);
+			}
+		}
+	}
+};
+
+class FMVVMDiffCustomProperty_Event : public IDiffCustomProperty
+{
+public:
+	const UMVVMBlueprintViewEvent* NewEvent;
+	const UMVVMBlueprintViewEvent* OldEvent;
+
+	FMVVMDiffCustomProperty_Event(const UMVVMBlueprintViewEvent* NewEvent, const UMVVMBlueprintViewEvent* OldEvent)
+		: NewEvent(NewEvent)
+		, OldEvent(OldEvent)
+	{
+	}
+
+	void HighlightProperty(const TSharedPtr<IDiffCustomObject> Widget) override
+	{
+		if (FMVVMDiffCustomObject_View* DiffWidget = StaticCast<FMVVMDiffCustomObject_View*>(Widget.Get()))
+		{
+			if (DiffWidget->NewList.IsValid())
+			{
+				DiffWidget->NewList->SelectEvent(NewEvent, true);
+			}
+			if (DiffWidget->OldList.IsValid())
+			{
+				DiffWidget->OldList->SelectEvent(OldEvent, true);
+			}
+		}
+	}
+};
+
+class FMVVMDiffCustomProperty_Condition : public IDiffCustomProperty
+{
+public:
+	const UMVVMBlueprintViewCondition* NewCondition;
+	const UMVVMBlueprintViewCondition* OldCondition;
+
+	FMVVMDiffCustomProperty_Condition(const UMVVMBlueprintViewCondition* NewCondition, const UMVVMBlueprintViewCondition* OldCondition)
+		: NewCondition(NewCondition)
+		, OldCondition(OldCondition)
+	{
+	}
+
+	void HighlightProperty(const TSharedPtr<IDiffCustomObject> Widget) override
+	{
+		if (FMVVMDiffCustomObject_View* DiffWidget = StaticCast<FMVVMDiffCustomObject_View*>(Widget.Get()))
+		{
+			if (DiffWidget->NewList.IsValid())
+			{
+				DiffWidget->NewList->SelectCondition(NewCondition, true);
+			}
+			if (DiffWidget->OldList.IsValid())
+			{
+				DiffWidget->OldList->SelectCondition(OldCondition, true);
+			}
+		}
+	}
+};
+
+class FMVVMDiffCustomProperty_Parameter : public IDiffCustomProperty
+{
+public:
+	FMVVMBlueprintPinId NewParameter;
+	FMVVMBlueprintPinId OldParameter;
+
+	FMVVMDiffCustomProperty_Parameter(FMVVMBlueprintPinId NewParameter, FMVVMBlueprintPinId OldParameter)
+		: NewParameter(NewParameter)
+		, OldParameter(OldParameter)
+	{
+	}
+
+	void HighlightProperty(const TSharedPtr<IDiffCustomObject> Widget) override
+	{
+		if (FMVVMDiffCustomObject_View* DiffWidget = StaticCast<FMVVMDiffCustomObject_View*>(Widget.Get()))
+		{
+			DiffWidget->NewList->SelectParameter(NewParameter, true);
+			DiffWidget->OldList->SelectParameter(OldParameter, true);
+		}
+	}
+};
+
+struct FMVVMWidgetBlueprintDiffProvider : public UE::MVVM::FMVVMDiffCustomObjectProvider
+{
+	TSharedPtr<IDiffCustomObject> CreateObjectDiff(const UWidgetBlueprint* NewBlueprint, const UWidgetBlueprint* OldBlueprint)
+	{
+		return MakeShared<FMVVMDiffCustomObject_View>(NewBlueprint, OldBlueprint);
+	}
+	
+	TSharedPtr<IDiffCustomProperty> CreatePropertyBindingDiff(FGuid NewBinding, FGuid OldBinding)
+	{
+		return MakeShared<FMVVMDiffCustomProperty_Binding>(NewBinding, OldBinding);
+	}
+	
+	TSharedPtr<IDiffCustomProperty> CreatePropertyConditionDiff(const UMVVMBlueprintViewCondition* NewCondition, const UMVVMBlueprintViewCondition* OldCondition)
+	{
+		return MakeShared<FMVVMDiffCustomProperty_Condition>(NewCondition, OldCondition);
+	}
+	
+	TSharedPtr<IDiffCustomProperty> CreatePropertyConditionDiff(const UMVVMBlueprintViewEvent* NewEvent, const UMVVMBlueprintViewEvent* OldEvent)
+	{
+		return MakeShared<FMVVMDiffCustomProperty_Event>(NewEvent, OldEvent);
+	}
+
+	TSharedPtr<IDiffCustomProperty> CreatePropertyParameterDiff(const FMVVMBlueprintPinId& NewParameter, const FMVVMBlueprintPinId& OldParameter)
+	{
+		return MakeShared<FMVVMDiffCustomProperty_Parameter>(NewParameter, OldParameter);
+	}
+};
diff --git a/Engine/Plugins/Runtime/ModelViewViewModel/Source/ModelViewViewModelEditor/Private/ModelViewViewModelEditorModule.cpp b/Engine/Plugins/Runtime/ModelViewViewModel/Source/ModelViewViewModelEditor/Private/ModelViewViewModelEditorModule.cpp
--- a/Engine/Plugins/Runtime/ModelViewViewModel/Source/ModelViewViewModelEditor/Private/ModelViewViewModelEditorModule.cpp
+++ b/Engine/Plugins/Runtime/ModelViewViewModel/Source/ModelViewViewModelEditor/Private/ModelViewViewModelEditorModule.cpp
@@ -12,6 +12,7 @@
 #include "Customizations/MVVMListViewBaseExtensionCustomizationExtender.h"
 #include "Customizations/MVVMPanelWidgetExtensionCustomizationExtender.h"
 #include "Customizations/MVVMPropertyBindingExtension.h"
+#include "Customizations/MVVMWidgetBlueprintDiffProvider.h"
 #include "Customizations/MVVMWidgetContextMenuExtension.h"
 #include "Framework/Application/SlateApplication.h"
 #include "Framework/Docking/LayoutExtender.h"
@@ -90,11 +91,14 @@ void FModelViewViewModelEditorModule::StartupModule()
 	FWidgetBlueprintGeneratedClassDelegates::CollectSaveOverrides.AddRaw(this, &FModelViewViewModelEditorModule::HandleCollectSaveOverrides);
 
 	UToolMenus::RegisterStartupCallback(FSimpleMulticastDelegate::FDelegate::CreateRaw(this, &FModelViewViewModelEditorModule::HandleRegisterMenus));
-}
 
+	UE::MVVM::FMVVMWidgetBlueprintDiff::RegisterCustomDiff(MakeShared<FMVVMWidgetBlueprintDiffProvider>());
+}
 
 void FModelViewViewModelEditorModule::ShutdownModule()
 {
+	UE::MVVM::FMVVMWidgetBlueprintDiff::UnregisterCustomDiff();
+
 	UnregisterMenus();
 
 	FWidgetBlueprintGeneratedClassDelegates::CollectSaveOverrides.RemoveAll(this);
@@ -318,6 +322,7 @@ void FModelViewViewModelEditorModule::UnregisterMenus()
 	UToolMenus::UnregisterOwner(this);
 }
 
+
 IMPLEMENT_MODULE(FModelViewViewModelEditorModule, ModelViewViewModelEditor);
 
 #undef LOCTEXT_NAMESPACE
diff --git a/Engine/Plugins/Runtime/ModelViewViewModel/Source/ModelViewViewModelEditor/Private/Widgets/SMVVMViewBindingListView.cpp b/Engine/Plugins/Runtime/ModelViewViewModel/Source/ModelViewViewModelEditor/Private/Widgets/SMVVMViewBindingListView.cpp
--- a/Engine/Plugins/Runtime/ModelViewViewModel/Source/ModelViewViewModelEditor/Private/Widgets/SMVVMViewBindingListView.cpp
+++ b/Engine/Plugins/Runtime/ModelViewViewModel/Source/ModelViewViewModelEditor/Private/Widgets/SMVVMViewBindingListView.cpp
@@ -64,7 +64,7 @@ namespace Private
 		return TSharedPtr<FBindingEntry>();
 	}
 	
-	TSharedPtr<FBindingEntry> FindEvent(UMVVMBlueprintViewEvent* Event, TConstArrayView<TSharedPtr<FBindingEntry>> Entries)
+	TSharedPtr<FBindingEntry> FindEvent(const UMVVMBlueprintViewEvent* Event, TConstArrayView<TSharedPtr<FBindingEntry>> Entries)
 	{
 		for (const TSharedPtr<FBindingEntry>& Entry : Entries)
 		{
@@ -81,7 +81,7 @@ namespace Private
 		return TSharedPtr<FBindingEntry>();
 	}
 
-	TSharedPtr<FBindingEntry> FindCondition(UMVVMBlueprintViewCondition* Condition, TConstArrayView<TSharedPtr<FBindingEntry>> Entries)
+	TSharedPtr<FBindingEntry> FindCondition(const UMVVMBlueprintViewCondition* Condition, TConstArrayView<TSharedPtr<FBindingEntry>> Entries)
 	{
 		for (const TSharedPtr<FBindingEntry>& Entry : Entries)
 		{
@@ -98,6 +98,34 @@ namespace Private
 		return TSharedPtr<FBindingEntry>();
 	}
 
+	TSharedPtr<FBindingEntry> FindParameter(FMVVMBlueprintPinId PinId, TConstArrayView<TSharedPtr<FBindingEntry>> Entries)
+	{
+		for (const TSharedPtr<FBindingEntry>& Entry : Entries)
+		{
+			if (Entry->GetRowType() == FBindingEntry::ERowType::ConditionParameter)
+			{
+				if (Entry->GetConditionParameterId() == PinId)
+				{
+					return Entry;
+				}
+			}
+			else if (Entry->GetRowType() == FBindingEntry::ERowType::EventParameter)
+			{
+				if (Entry->GetEventParameterId() == PinId)
+				{
+					return Entry;
+				}
+			}
+
+			TSharedPtr<FBindingEntry> Result = FindParameter(PinId, Entry->GetAllChildren());
+			if (Result)
+			{
+				return Result;
+			}
+		}
+		return TSharedPtr<FBindingEntry>();
+	}
+
 	void FilterEntryList(FString FilterString, const TArray<TSharedPtr<FBindingEntry>>& RootGroups, TArray<TSharedPtr<FBindingEntry>>& FilteredRootGroups, UMVVMBlueprintView* BlueprintView, UMVVMWidgetBlueprintExtension_View* MVVMExtensionPtr)
 	{
 		if (!FilterString.TrimStartAndEnd().IsEmpty())
@@ -705,6 +733,9 @@ TSharedRef<ITableRow> SBindingsList::GenerateEntryRow(TSharedPtr<FBindingEntry>
 				break;
 			}
 		}
+	
+		//If this row was not spawned from the blueprint editor. Disable modifying, visual only
+		Row->GetContent()->SetEnabled(WeakBlueprintEditor.IsValid());
 
 		return Row.ToSharedRef();
 	}
@@ -761,6 +792,70 @@ void SBindingsList::RequestNavigateToCondition(UMVVMBlueprintViewCondition* Cond
 	}
 }
 
+void SBindingsList::SelectBinding(FGuid BindingId, bool bShouldScrollIntoView)
+{	
+	if (TSharedPtr<FBindingEntry> Entry = Private::FindBinding(BindingId, FilteredRootGroups))
+	{
+		SetSelection(MakeConstArrayView({Entry}));
+		if (bShouldScrollIntoView && TreeView)
+		{
+			TreeView->RequestNavigateToItem(Entry);
+		}
+	}
+	else
+	{
+		SetSelection(MakeConstArrayView<TSharedPtr<FBindingEntry>>({}));
+	}
+}
+
+void SBindingsList::SelectEvent(const UMVVMBlueprintViewEvent* Event, bool bShouldScrollIntoView)
+{	
+	if (TSharedPtr<FBindingEntry> Entry = Private::FindEvent(Event, FilteredRootGroups))
+	{
+		SetSelection(MakeConstArrayView({Entry}));
+		if (bShouldScrollIntoView && TreeView)
+		{
+			TreeView->RequestNavigateToItem(Entry);
+		}
+	}
+	else
+	{
+		SetSelection(MakeConstArrayView<TSharedPtr<FBindingEntry>>({}));
+	}
+}
+
+void SBindingsList::SelectCondition(const UMVVMBlueprintViewCondition* Condition, bool bShouldScrollIntoView)
+{	
+	if (TSharedPtr<FBindingEntry> Entry = Private::FindCondition(Condition, FilteredRootGroups))
+	{
+		SetSelection(MakeConstArrayView({Entry}));
+		if (bShouldScrollIntoView && TreeView)
+		{
+			TreeView->RequestNavigateToItem(Entry);
+		}
+	}
+	else
+	{
+		SetSelection(MakeConstArrayView<TSharedPtr<FBindingEntry>>({}));
+	}
+}
+
+void SBindingsList::SelectParameter(const FMVVMBlueprintPinId& PinId, bool bShouldScrollIntoView)
+{
+	if (TSharedPtr<FBindingEntry> Entry = Private::FindParameter(PinId, FilteredRootGroups))
+	{
+		SetSelection(MakeConstArrayView({Entry}));
+		if (bShouldScrollIntoView && TreeView)
+		{
+			TreeView->RequestNavigateToItem(Entry);
+		}
+	}
+	else
+	{
+		SetSelection(MakeConstArrayView<TSharedPtr<FBindingEntry>>({}));
+	}
+}
+
 void SBindingsList::SetRootGroupsExpansion(bool InIsExpanded)
 {
 	if (TreeView.IsValid())
diff --git a/Engine/Plugins/Runtime/ModelViewViewModel/Source/ModelViewViewModelEditor/Private/Widgets/SMVVMViewBindingListView.h b/Engine/Plugins/Runtime/ModelViewViewModel/Source/ModelViewViewModelEditor/Private/Widgets/SMVVMViewBindingListView.h
--- a/Engine/Plugins/Runtime/ModelViewViewModel/Source/ModelViewViewModelEditor/Private/Widgets/SMVVMViewBindingListView.h
+++ b/Engine/Plugins/Runtime/ModelViewViewModel/Source/ModelViewViewModelEditor/Private/Widgets/SMVVMViewBindingListView.h
@@ -15,6 +15,8 @@ class UMVVMBlueprintViewEvent;
 class UMVVMBlueprintViewCondition;
 class FWidgetBlueprintEditor;
 
+struct FMVVMBlueprintPinId;
+
 namespace UE::MVVM
 {
 class SBindingsPanel;
@@ -50,6 +52,11 @@ public:
 
 	void SetSelection(TConstArrayView<const TSharedPtr<FBindingEntry>> InEntries);
 
+	void SelectBinding(FGuid BindingId, bool bShouldScrollIntoView);
+	void SelectEvent(const UMVVMBlueprintViewEvent* Event, bool bShouldScrollIntoView);
+	void SelectCondition(const UMVVMBlueprintViewCondition* Condition, bool bShouldScrollIntoView);
+	void SelectParameter(const FMVVMBlueprintPinId& PinId, bool bShouldScrollIntoView);
+
 private:
 	TSharedRef<ITableRow> GenerateEntryRow(TSharedPtr<FBindingEntry> Entry, const TSharedRef<STableViewBase>& OwnerTable);
 	void OnSourceListSelectionChanged(TSharedPtr<FBindingEntry> Entry, ESelectInfo::Type SelectionType) const;
diff --git a/Engine/Source/Editor/Kismet/Private/DiffControl.cpp b/Engine/Source/Editor/Kismet/Private/DiffControl.cpp
--- a/Engine/Source/Editor/Kismet/Private/DiffControl.cpp
+++ b/Engine/Source/Editor/Kismet/Private/DiffControl.cpp
@@ -687,6 +687,23 @@ void FBlueprintTypeDiffControl::GenerateTreeEntries(TArray< TSharedPtr<FBlueprin
 		// add comments as children to this category
 		GenerateCategoryCommentTreeEntries(OutTreeEntries.Last()->Children, OutRealDifferences, SubObjectDiff->SourceResult.OwningObjectPath);
 	}
+
+	// Now add customobject diffs, one category per object
+	for (const TSharedPtr<FDiffCustomObject>& DiffCustomObject : DiffCustomObjects)
+	{
+		Children.Reset();
+
+		Children.Append(DiffCustomObject->Diffs);
+		OutRealDifferences.Append(DiffCustomObject->Diffs);
+
+		TSharedPtr<FBlueprintDifferenceTreeEntry> CustomObjectEntry = FBlueprintDifferenceTreeEntry::CreateCategoryEntry(
+			DiffCustomObject->SourceResult.DisplayString,
+			DiffCustomObject->SourceResult.ToolTip,
+			FOnDiffEntryFocused::CreateSP(AsShared(), &FBlueprintTypeDiffControl::OnSelectDiffCustomObject, TSharedPtr<IDiffCustomProperty>(), DiffCustomObject),
+			Children,
+			true);
+		OutTreeEntries.Push(CustomObjectEntry);
+	}
 }
 
 void FBlueprintTypeDiffControl::EnableComments(TWeakPtr<STreeView<TSharedPtr<FBlueprintDifferenceTreeEntry>>> TreeView)
@@ -755,6 +772,47 @@ void FBlueprintTypeDiffControl::BuildDiffSourceArray()
 					}
 				}
 			}
+			else if (CurrentDiff.Diff == EDiffType::CUSTOM_OBJECT)
+			{
+				// If no custom object, do not add to DiffListSource
+				if (!CurrentDiff.CustomObject.IsValid())
+				{
+					continue;
+				}
+				
+				// Invert order, we want old then new
+				TSharedPtr<FDiffCustomObject> DiffCustomObject = MakeShared<FDiffCustomObject>(CurrentDiff, CurrentDiff.CustomObject);
+
+				// Gather associated custom property diffs for the object
+				for (const FDiffSingleResult& OtherDiff : BlueprintDiffResults)
+				{
+					// If the property diff is not associated with the custom object diff, ignore
+					if (OtherDiff.Object1 != CurrentDiff.Object1 || OtherDiff.Object2 != CurrentDiff.Object2)
+					{
+						continue;
+					}
+
+					if (OtherDiff.CustomProperty.IsValid())
+					{
+						FGenerateDiffEntryWidget CreateEntryWidget = FGenerateDiffEntryWidget::CreateSPLambda(this, [this](const FDiffSingleResult& Diff) {
+							return SNew(STextBlock)
+							.Text(Diff.DisplayString)
+							.ToolTipText(Diff.ToolTip)
+							.ColorAndOpacity(Diff.GetDisplayColor());
+						}, OtherDiff);
+
+						TSharedPtr<FBlueprintDifferenceTreeEntry> Entry = MakeShared<FBlueprintDifferenceTreeEntry>(
+							FOnDiffEntryFocused::CreateSP(AsShared(), &FBlueprintTypeDiffControl::OnSelectDiffCustomObject, OtherDiff.CustomProperty, DiffCustomObject),
+							CreateEntryWidget);
+						DiffCustomObject->Diffs.Push(Entry);
+					}
+				}
+
+				if(!DiffCustomObject->Diffs.IsEmpty())
+				{
+					DiffCustomObjects.Add(DiffCustomObject);
+				}
+			}
 			else
 			{
 				DiffListSource.Add(MakeShared<FDiffResultItem>(CurrentDiff));
@@ -780,6 +838,36 @@ void FBlueprintTypeDiffControl::OnSelectSubobjectDiff(FPropertySoftPath Identifi
 	}
 }
 
+void FBlueprintTypeDiffControl::OnSelectDiffCustomObject(TSharedPtr<IDiffCustomProperty> CustomProperty, TSharedPtr<FDiffCustomObject> DiffCustomObject)
+{
+	// This allows the owning control to focus the correct tab (or do whatever else it likes):
+	SelectionCallback.ExecuteIfBound();
+
+	if (DiffCustomObject.IsValid())
+	{
+		if (TSharedPtr<IDiffCustomObject> DiffWidget = DiffCustomObject->DiffWidget)
+		{
+			if (CustomProperty.IsValid())
+			{
+				CustomProperty->HighlightProperty(DiffWidget);
+			}
+
+			if (DiffWidget->GetOldWidget().IsValid())
+			{
+				OldDetailsBox->SetContent(DiffWidget->GetOldWidget().ToSharedRef());
+			}
+			if (DiffWidget->GetNewWidget().IsValid())
+			{
+				NewDetailsBox->SetContent(DiffWidget->GetNewWidget().ToSharedRef());
+			}
+		}
+		else
+		{
+			OldDetailsBox->SetContent(SNullWidget::NullWidget);
+			NewDetailsBox->SetContent(SNullWidget::NullWidget);
+		}
+	}
+}
 
 /////////////////////////////////////////////////////////////////////////////
 /// FGraphToDiff
diff --git a/Engine/Source/Editor/Kismet/Private/DiffControl.h b/Engine/Source/Editor/Kismet/Private/DiffControl.h
--- a/Engine/Source/Editor/Kismet/Private/DiffControl.h
+++ b/Engine/Source/Editor/Kismet/Private/DiffControl.h
@@ -186,6 +186,17 @@ struct FBlueprintTypeDiffControl : public TSharedFromThis<FBlueprintTypeDiffCont
 		{}
 	};
 
+	struct FDiffCustomObject
+	{
+		FDiffSingleResult SourceResult;
+		TSharedPtr<IDiffCustomObject> DiffWidget;
+		TArray<TSharedPtr<FBlueprintDifferenceTreeEntry>> Diffs;
+		FDiffCustomObject(const FDiffSingleResult& InSourceResult, TSharedPtr<IDiffCustomObject> DiffWidget)
+			: SourceResult(InSourceResult)
+			, DiffWidget(DiffWidget)
+		{}
+	};
+
 	FBlueprintTypeDiffControl(const UBlueprint* InBlueprintOld, const UBlueprint* InBlueprintNew, FOnDiffEntryFocused InSelectionCallback);
 
 	/** Generate difference tree widgets */
@@ -214,8 +225,14 @@ private:
 	/** Handle selecting a diff */
 	void OnSelectSubobjectDiff(FPropertySoftPath Identifier, TSharedPtr<FSubObjectDiff> SubObjectDiff);
 
+	/** Handle selecting a custom object diff */
+	void OnSelectDiffCustomObject(TSharedPtr<IDiffCustomProperty> CustomProperty, TSharedPtr<FDiffCustomObject> DiffCustomObject);
+
 	/** List of objects with differences */
 	TArray<TSharedPtr<FSubObjectDiff>> SubObjectDiffs;
+	
+	/** List of custom object differences */
+	TArray<TSharedPtr<FDiffCustomObject>> DiffCustomObjects;
 
 	/** Source for list view */
 	TArray<TSharedPtr<FDiffResultItem>> DiffListSource;
diff --git a/Engine/Source/Editor/UMGEditor/Private/WidgetBlueprint.cpp b/Engine/Source/Editor/UMGEditor/Private/WidgetBlueprint.cpp
--- a/Engine/Source/Editor/UMGEditor/Private/WidgetBlueprint.cpp
+++ b/Engine/Source/Editor/UMGEditor/Private/WidgetBlueprint.cpp
@@ -970,6 +970,17 @@ bool UWidgetBlueprint::FindDiffs(const UBlueprint* OtherBlueprint, FDiffResults&
 		}
 	}
 
+	for (const TObjectPtr<UBlueprintExtension>& Extension : GetExtensions())
+	{
+		if (Extension->IsA<UWidgetBlueprintExtension>())
+		{
+			if (UWidgetBlueprintExtension* WidgetExtension = StaticCast<UWidgetBlueprintExtension*>(Extension.Get()))
+			{
+				WidgetExtension->FindDiffs(this, OtherWidgetBP, Results);
+			}
+		}
+	}
+
 	// Add info warning
 	if (Results.CanStoreResults())
 	{
diff --git a/Engine/Source/Editor/UMGEditor/Public/WidgetBlueprintExtension.h b/Engine/Source/Editor/UMGEditor/Public/WidgetBlueprintExtension.h
--- a/Engine/Source/Editor/UMGEditor/Public/WidgetBlueprintExtension.h
+++ b/Engine/Source/Editor/UMGEditor/Public/WidgetBlueprintExtension.h
@@ -65,6 +65,10 @@ public:
 	/** Get the WidgetBlueprint that hosts this extension */
 	UWidgetBlueprint* GetWidgetBlueprint() const;
 
+#if WITH_EDITOR
+	virtual void FindDiffs(const UWidgetBlueprint* OwningBlueprint, const UWidgetBlueprint* OtherBlueprint, FDiffResults& Results) const {}
+#endif //WITH_EDITOR
+
 protected:
 	/** 
 	 * Override point called when a compiler context is created for the WidgetBlueprint
diff --git a/Engine/Source/Runtime/Engine/Public/DiffResults.h b/Engine/Source/Runtime/Engine/Public/DiffResults.h
--- a/Engine/Source/Runtime/Engine/Public/DiffResults.h
+++ b/Engine/Source/Runtime/Engine/Public/DiffResults.h
@@ -6,6 +6,8 @@
 #include "Styling/StyleColors.h"
 #include "Styling/AppStyle.h"
 
+class SWidget;
+
 /** Differences found within a graph or object */
 namespace EDiffType
 {
@@ -51,7 +53,10 @@ namespace EDiffType
 		NODE_PROPERTY,
 
 		// Informational message, does't count as a real diff
-		INFO_MESSAGE
+		INFO_MESSAGE,
+
+		// User defined custom diff. Allows user to override property diffing and panel widgets
+		CUSTOM_OBJECT
 	};
 	
 	enum Category
@@ -111,6 +116,10 @@ struct FDiffSingleResult
 	/** Optional tooltip containing more information */
 	FText ToolTip; 
 
+	/** Custom Object Diff. Used with CUSTOM_OBJECT diff types */
+	TSharedPtr<class IDiffCustomObject> CustomObject;
+	TSharedPtr<class IDiffCustomProperty> CustomProperty;
+
 	/** Get the color that is associated with this diff category */
 	FLinearColor GetDisplayColor() const
 	{
@@ -193,3 +202,21 @@ private:
 	TArray<FDiffSingleResult>* ResultArray;
 	bool bHasFoundDiffs;
 };
+
+/** 
+ * Interface to support custom widget for object diffs
+ */
+class IDiffCustomObject : public TSharedFromThis<IDiffCustomObject>
+{
+public:
+	virtual ~IDiffCustomObject() = default;
+	virtual TSharedPtr<SWidget> GetOldWidget() = 0;
+	virtual TSharedPtr<SWidget> GetNewWidget() = 0;
+};
+
+class IDiffCustomProperty : public TSharedFromThis<IDiffCustomProperty>
+{
+public:
+	virtual ~IDiffCustomProperty() = default;
+	virtual void HighlightProperty(const TSharedPtr<IDiffCustomObject> Widget) = 0;
+};
