MVVM: Fix include path on Mac

diff --git a/Engine/Plugins/Runtime/ModelViewViewModel/Source/ModelViewViewModelEditor/Private/Customizations/MVVMWidgetBlueprintDiffProvider.h b/Engine/Plugins/Runtime/ModelViewViewModel/Source/ModelViewViewModelEditor/Private/Customizations/MVVMWidgetBlueprintDiffProvider.h
--- a/Engine/Plugins/Runtime/ModelViewViewModel/Source/ModelViewViewModelEditor/Private/Customizations/MVVMWidgetBlueprintDiffProvider.h
+++ b/Engine/Plugins/Runtime/ModelViewViewModel/Source/ModelViewViewModelEditor/Private/Customizations/MVVMWidgetBlueprintDiffProvider.h
@@ -8,7 +8,7 @@
 #include "MVVMBlueprintViewEvent.h"
 #include "MVVMWidgetBlueprintDiff.h"
 #include "MVVMWidgetBlueprintExtension_View.h"
-#include "Widgets\SMVVMViewBindingListView.h"
+#include "Widgets/SMVVMViewBindingListView.h"
 
 class FMVVMDiffCustomObject_View : public IDiffCustomObject
 {
