Make module register on the compiler delegate and hook MVVMViewClass on the module's MT delegate instead
UE-288224
UE-275539

#rnx
#tests manually checked the compilation handler after registering to the module instead

diff --git a/Engine/Plugins/Runtime/ModelViewViewModel/Source/ModelViewViewModel/Private/ModelViewViewModelModule.cpp b/Engine/Plugins/Runtime/ModelViewViewModel/Source/ModelViewViewModel/Private/ModelViewViewModelModule.cpp
--- a/Engine/Plugins/Runtime/ModelViewViewModel/Source/ModelViewViewModel/Private/ModelViewViewModelModule.cpp
+++ b/Engine/Plugins/Runtime/ModelViewViewModel/Source/ModelViewViewModel/Private/ModelViewViewModelModule.cpp
@@ -5,6 +5,11 @@
 #include "Modules/ModuleManager.h"
 #include "Types/MVVMExecutionMode.h"
 
+#if WITH_EDITOR
+#include "Editor.h"
+#include "Editor/EditorEngine.h"
+#endif
+
 #define LOCTEXT_NAMESPACE "ModelViewViewModelModule"
 
 void FModelViewViewModelModule::StartupModule()
@@ -15,6 +20,25 @@ void FModelViewViewModelModule::StartupModule()
 		HandleDefaultExecutionModeChanged(CVarDefaultExecutionMode);
 		CVarDefaultExecutionMode->OnChangedDelegate().AddRaw(this, &FModelViewViewModelModule::HandleDefaultExecutionModeChanged);
 	}
+
+#if WITH_EDITOR
+	FCoreDelegates::OnPostEngineInit.AddRaw(this, &FModelViewViewModelModule::OnPostEngineInit);
+#endif
+}
+
+void FModelViewViewModelModule::OnPostEngineInit()
+{
+#if WITH_EDITOR
+	ensure(!OnBlueprintCompiledHandle.IsValid());
+	OnBlueprintCompiledHandle = GEditor->OnBlueprintCompiled().AddRaw(this, &FModelViewViewModelModule::HandleBlueprintCompiled);
+#endif
+}
+
+void FModelViewViewModelModule::HandleBlueprintCompiled()
+{
+#if WITH_EDITOR
+	OnBlueprintCompiled.Broadcast();
+#endif
 }
 
 void FModelViewViewModelModule::ShutdownModule()
@@ -26,6 +50,14 @@ void FModelViewViewModelModule::ShutdownModule()
 			CVarDefaultExecutionMode->OnChangedDelegate().RemoveAll(this);
 		}
 	}
+
+#if WITH_EDITOR
+	FCoreDelegates::OnPostEngineInit.RemoveAll(this);
+	if (GEditor != nullptr && OnBlueprintCompiledHandle.IsValid())
+	{
+		GEditor->OnBlueprintCompiled().Remove(OnBlueprintCompiledHandle);
+	}
+#endif
 }
 
 void FModelViewViewModelModule::HandleDefaultExecutionModeChanged(IConsoleVariable* Variable)
diff --git a/Engine/Plugins/Runtime/ModelViewViewModel/Source/ModelViewViewModel/Private/ModelViewViewModelModule.h b/Engine/Plugins/Runtime/ModelViewViewModel/Source/ModelViewViewModel/Private/ModelViewViewModelModule.h
--- a/Engine/Plugins/Runtime/ModelViewViewModel/Source/ModelViewViewModel/Private/ModelViewViewModelModule.h
+++ b/Engine/Plugins/Runtime/ModelViewViewModel/Source/ModelViewViewModel/Private/ModelViewViewModelModule.h
@@ -23,6 +23,18 @@ public:
 	virtual void ShutdownModule() override;
 	//~ End IModuleInterface interface
 
+#if WITH_EDITOR
+	DECLARE_TS_MULTICAST_DELEGATE(FOnBlueprintCompiled)
+	FOnBlueprintCompiled OnBlueprintCompiled;
+#endif
+
 private:
 	void HandleDefaultExecutionModeChanged(IConsoleVariable* Variable);
+	void OnPostEngineInit();
+
+	void HandleBlueprintCompiled();
+#if WITH_EDITOR
+	FDelegateHandle OnBlueprintCompiledHandle;
+#endif
+
 };
diff --git a/Engine/Plugins/Runtime/ModelViewViewModel/Source/ModelViewViewModel/Private/View/MVVMViewClass.cpp b/Engine/Plugins/Runtime/ModelViewViewModel/Source/ModelViewViewModel/Private/View/MVVMViewClass.cpp
--- a/Engine/Plugins/Runtime/ModelViewViewModel/Source/ModelViewViewModel/Private/View/MVVMViewClass.cpp
+++ b/Engine/Plugins/Runtime/ModelViewViewModel/Source/ModelViewViewModel/Private/View/MVVMViewClass.cpp
@@ -4,6 +4,8 @@
 #include "Types/MVVMFieldContext.h"
 #include "View/MVVMView.h"
 #include "View/MVVMViewModelContextResolver.h"
+#include "Modules/ModuleManager.h"
+#include "ModelViewViewModelModule.h"
 
 #include "Bindings/MVVMFieldPathHelper.h"
 #include "Blueprint/UserWidget.h"
@@ -747,7 +749,8 @@ void UMVVMViewClass::Construct(UUserWidget* UserWidget)
 #if WITH_EDITOR
 	if (GEditor)
 	{
-		BlueprintCompiledHandle = GEditor->OnBlueprintCompiled().AddUObject(this, &UMVVMViewClass::HandleBlueprintCompiled);
+		FModelViewViewModelModule& Module = FModuleManager::GetModuleChecked<FModelViewViewModelModule>("ModelViewViewModel");
+		Module.OnBlueprintCompiled.AddUObject(this, &UMVVMViewClass::HandleBlueprintCompiled);
 	}
 #endif
 }
@@ -765,8 +768,8 @@ void UMVVMViewClass::Destruct(UUserWidget* UserWidget)
 #if WITH_EDITOR
 	if (GEditor && BlueprintCompiledHandle.IsValid())
 	{
-		GEditor->OnBlueprintCompiled().Remove(BlueprintCompiledHandle);
-		BlueprintCompiledHandle.Reset();
+		FModelViewViewModelModule& Module = FModuleManager::GetModuleChecked<FModelViewViewModelModule>("ModelViewViewModel");
+		Module.OnBlueprintCompiled.Remove(BlueprintCompiledHandle);
 	}
 #endif
 }
