Fixed crash when launching the standalone editor executable

#rnx

diff --git a/Engine/Plugins/Runtime/ModelViewViewModel/Source/ModelViewViewModel/Private/ModelViewViewModelModule.cpp b/Engine/Plugins/Runtime/ModelViewViewModel/Source/ModelViewViewModel/Private/ModelViewViewModelModule.cpp
--- a/Engine/Plugins/Runtime/ModelViewViewModel/Source/ModelViewViewModel/Private/ModelViewViewModelModule.cpp
+++ b/Engine/Plugins/Runtime/ModelViewViewModel/Source/ModelViewViewModel/Private/ModelViewViewModelModule.cpp
@@ -30,7 +30,10 @@ void FModelViewViewModelModule::OnPostEngineInit()
 {
 #if WITH_EDITOR
 	ensure(!OnBlueprintCompiledHandle.IsValid());
-	OnBlueprintCompiledHandle = GEditor->OnBlueprintCompiled().AddRaw(this, &FModelViewViewModelModule::HandleBlueprintCompiled);
+	if (GEditor != nullptr)
+	{
+		OnBlueprintCompiledHandle = GEditor->OnBlueprintCompiled().AddRaw(this, &FModelViewViewModelModule::HandleBlueprintCompiled);
+	}
 #endif
 }
 
