Updated headers using UnrealCodeFixup to make sure dllstorage is on methods/staticvars instead of types. Part 4/n

diff --git a/Engine/Plugins/Runtime/ModelViewViewModel/Source/ModelViewViewModelBlueprint/Public/MVVMBlueprintPin.h b/Engine/Plugins/Runtime/ModelViewViewModel/Source/ModelViewViewModelBlueprint/Public/MVVMBlueprintPin.h
--- a/Engine/Plugins/Runtime/ModelViewViewModel/Source/ModelViewViewModelBlueprint/Public/MVVMBlueprintPin.h
+++ b/Engine/Plugins/Runtime/ModelViewViewModel/Source/ModelViewViewModelBlueprint/Public/MVVMBlueprintPin.h
@@ -6,6 +6,8 @@
 
 #include "MVVMBlueprintPin.generated.h"
 
+#define UE_API MODELVIEWVIEWMODELBLUEPRINT_API
+
 class UWidgetBlueprint;
 class UEdGraphNode;
 class UEdGraphPin;
@@ -24,15 +26,15 @@ enum class EMVVMBlueprintPinStatus : uint8
  * Pin name type to help with compare operation and moving it around.
  */
 USTRUCT()
-struct MODELVIEWVIEWMODELBLUEPRINT_API FMVVMBlueprintPinId
+struct FMVVMBlueprintPinId
 {
 	GENERATED_BODY()
 
 	FMVVMBlueprintPinId() = default;
-	explicit FMVVMBlueprintPinId(const TArrayView<const FName> Names);
-	explicit FMVVMBlueprintPinId(TArray<FName>&& Names);
+	UE_API explicit FMVVMBlueprintPinId(const TArrayView<const FName> Names);
+	UE_API explicit FMVVMBlueprintPinId(TArray<FName>&& Names);
 
-	bool IsValid() const;
+	UE_API bool IsValid() const;
 
 	const TArrayView<const FName> GetNames() const
 	{
@@ -40,15 +42,15 @@ struct MODELVIEWVIEWMODELBLUEPRINT_API FMVVMBlueprintPinId
 	}
 
 	/** return true if the Pin is part of the Other pin. It can be a grand child. */
-	bool IsChildOf(const FMVVMBlueprintPinId& Other) const;
+	UE_API bool IsChildOf(const FMVVMBlueprintPinId& Other) const;
 
 	/** return true if the Pin is the directly child of the Other pin. It can be a child but not a grand child. */
-	bool IsDirectChildOf(const FMVVMBlueprintPinId& Other) const;
+	UE_API bool IsDirectChildOf(const FMVVMBlueprintPinId& Other) const;
 
-	bool operator==(const FMVVMBlueprintPinId& Other) const;
-	bool operator==(const TArrayView<const FName> Other) const;
+	UE_API bool operator==(const FMVVMBlueprintPinId& Other) const;
+	UE_API bool operator==(const TArrayView<const FName> Other) const;
 
-	FString ToString() const;
+	UE_API FString ToString() const;
 
 private:
 	UPROPERTY(VisibleAnywhere, Category = "MVVM")
@@ -59,7 +61,7 @@ private:
  *
  */
 USTRUCT()
-struct MODELVIEWVIEWMODELBLUEPRINT_API FMVVMBlueprintPin
+struct FMVVMBlueprintPin
 {
 	GENERATED_BODY()
 
@@ -99,9 +101,9 @@ private:
 public:
 	FMVVMBlueprintPin() = default;
 	UE_DEPRECATED(5.4, "FMVVMBlueprintPin with a single name is deprecated. Use the TArrayView constructor instead")
-	FMVVMBlueprintPin(FName PinName);
-	explicit FMVVMBlueprintPin(FMVVMBlueprintPinId PinId);
-	explicit FMVVMBlueprintPin(const TArrayView<const FName> PinName);
+	UE_API FMVVMBlueprintPin(FName PinName);
+	UE_API explicit FMVVMBlueprintPin(FMVVMBlueprintPinId PinId);
+	UE_API explicit FMVVMBlueprintPin(const TArrayView<const FName> PinName);
 
 	UE_DEPRECATED(5.4, "GetName is deprecated. Use GetId().GetNames instead")
 	FName GetName() const
@@ -143,26 +145,26 @@ public:
 		return Path;
 	}
 
-	FString GetValueAsString(const UClass* SelfContext) const;
+	UE_API FString GetValueAsString(const UClass* SelfContext) const;
 
-	void SetDefaultValue(UObject* Value);
-	void SetDefaultValue(const FText& Value);
-	void SetDefaultValue(const FString& Value);
-	void SetPath(const FMVVMBlueprintPropertyPath& Value);
+	UE_API void SetDefaultValue(UObject* Value);
+	UE_API void SetDefaultValue(const FText& Value);
+	UE_API void SetDefaultValue(const FString& Value);
+	UE_API void SetPath(const FMVVMBlueprintPropertyPath& Value);
 
 public:
-	void PostSerialize(const FArchive& Ar);
+	UE_API void PostSerialize(const FArchive& Ar);
 
 public:
-	static bool IsInputPin(const UEdGraphPin* Pin);
-	static TArray<FMVVMBlueprintPin> CopyAndReturnMissingPins(UBlueprint* Blueprint, UEdGraphNode* GraphNode, const TArray<FMVVMBlueprintPin>& Pins);
-	static TArray<FMVVMBlueprintPin> CreateFromNode(UBlueprint* Blueprint, UEdGraphNode* GraphNode);
-	static FMVVMBlueprintPin CreateFromPin(const UBlueprint* Blueprint, const UEdGraphPin* Pin);
+	static UE_API bool IsInputPin(const UEdGraphPin* Pin);
+	static UE_API TArray<FMVVMBlueprintPin> CopyAndReturnMissingPins(UBlueprint* Blueprint, UEdGraphNode* GraphNode, const TArray<FMVVMBlueprintPin>& Pins);
+	static UE_API TArray<FMVVMBlueprintPin> CreateFromNode(UBlueprint* Blueprint, UEdGraphNode* GraphNode);
+	static UE_API FMVVMBlueprintPin CreateFromPin(const UBlueprint* Blueprint, const UEdGraphPin* Pin);
 
-	void CopyTo(const UBlueprint* WidgetBlueprint, UEdGraphNode* Node) const;
-	UEdGraphPin* FindGraphPin(const UEdGraph* Graph) const;
+	UE_API void CopyTo(const UBlueprint* WidgetBlueprint, UEdGraphNode* Node) const;
+	UE_API UEdGraphPin* FindGraphPin(const UEdGraph* Graph) const;
 
-	void Reset();
+	UE_API void Reset();
 };
 
 template<>
@@ -173,3 +175,5 @@ struct TStructOpsTypeTraits<FMVVMBlueprintPin> : public TStructOpsTypeTraitsBase
 		WithPostSerialize = true,
 	};
 };
+
+#undef UE_API
diff --git a/Engine/Plugins/Runtime/ModelViewViewModel/Source/ModelViewViewModelBlueprint/Public/MVVMBlueprintView.h b/Engine/Plugins/Runtime/ModelViewViewModel/Source/ModelViewViewModelBlueprint/Public/MVVMBlueprintView.h
--- a/Engine/Plugins/Runtime/ModelViewViewModel/Source/ModelViewViewModelBlueprint/Public/MVVMBlueprintView.h
+++ b/Engine/Plugins/Runtime/ModelViewViewModel/Source/ModelViewViewModelBlueprint/Public/MVVMBlueprintView.h
@@ -7,6 +7,8 @@
 
 #include "MVVMBlueprintView.generated.h"
 
+#define UE_API MODELVIEWVIEWMODELBLUEPRINT_API
+
 class UMVVMWidgetBlueprintExtension_View;
 class UMVVMBlueprintViewEvent;
 class UMVVMBlueprintViewCondition;
@@ -76,13 +78,13 @@ public:
 /**
  * 
  */
-UCLASS(Within=MVVMWidgetBlueprintExtension_View)
-class MODELVIEWVIEWMODELBLUEPRINT_API UMVVMBlueprintView : public UObject
+UCLASS(MinimalAPI, Within=MVVMWidgetBlueprintExtension_View)
+class UMVVMBlueprintView : public UObject
 {
 	GENERATED_BODY()
 
 public:
-	UMVVMBlueprintView();
+	UE_API UMVVMBlueprintView();
 
 public:
 	UMVVMBlueprintViewSettings* GetSettings()
@@ -90,40 +92,40 @@ public:
 		return Settings;
 	}
 
-	FMVVMBlueprintViewModelContext* FindViewModel(FGuid ViewModelId);
-	const FMVVMBlueprintViewModelContext* FindViewModel(FGuid ViewModelId) const;
-	const FMVVMBlueprintViewModelContext* FindViewModel(FName ViewModelName) const;
+	UE_API FMVVMBlueprintViewModelContext* FindViewModel(FGuid ViewModelId);
+	UE_API const FMVVMBlueprintViewModelContext* FindViewModel(FGuid ViewModelId) const;
+	UE_API const FMVVMBlueprintViewModelContext* FindViewModel(FName ViewModelName) const;
 
-	void AddViewModel(const FMVVMBlueprintViewModelContext& NewContext);
-	bool RemoveViewModel(FGuid ViewModelId);
-	int32 RemoveViewModels(const TArrayView<FGuid> ViewModelIds);
-	bool RenameViewModel(FName OldViewModelName, FName NewViewModelName);
-	bool ReparentViewModel(FGuid ViewModelId, const UClass* ViewModelClass);
+	UE_API void AddViewModel(const FMVVMBlueprintViewModelContext& NewContext);
+	UE_API bool RemoveViewModel(FGuid ViewModelId);
+	UE_API int32 RemoveViewModels(const TArrayView<FGuid> ViewModelIds);
+	UE_API bool RenameViewModel(FName OldViewModelName, FName NewViewModelName);
+	UE_API bool ReparentViewModel(FGuid ViewModelId, const UClass* ViewModelClass);
 
 	const TArrayView<const FMVVMBlueprintViewModelContext> GetViewModels() const
 	{
 		return AvailableViewModels; 
 	}
 
-	const FMVVMBlueprintViewBinding* FindBinding(const UWidget* Widget, const FProperty* Property) const;
-	FMVVMBlueprintViewBinding* FindBinding(const UWidget* Widget, const FProperty* Property);
+	UE_API const FMVVMBlueprintViewBinding* FindBinding(const UWidget* Widget, const FProperty* Property) const;
+	UE_API FMVVMBlueprintViewBinding* FindBinding(const UWidget* Widget, const FProperty* Property);
 
-	void RemoveBinding(const FMVVMBlueprintViewBinding* Binding);
-	const FMVVMBlueprintViewBinding* DuplicateBinding(const FMVVMBlueprintViewBinding* Binding);
+	UE_API void RemoveBinding(const FMVVMBlueprintViewBinding* Binding);
+	UE_API const FMVVMBlueprintViewBinding* DuplicateBinding(const FMVVMBlueprintViewBinding* Binding);
 
-	void RemoveBindingAt(int32 Index);
+	UE_API void RemoveBindingAt(int32 Index);
 
-	FMVVMBlueprintViewBinding& AddDefaultBinding();
+	UE_API FMVVMBlueprintViewBinding& AddDefaultBinding();
 
 	int32 GetNumBindings() const
 	{
 		return Bindings.Num();
 	}
 
-	FMVVMBlueprintViewBinding* GetBindingAt(int32 Index);
-	const FMVVMBlueprintViewBinding* GetBindingAt(int32 Index) const;
-	FMVVMBlueprintViewBinding* GetBinding(FGuid Id);
-	const FMVVMBlueprintViewBinding* GetBinding(FGuid Id) const;
+	UE_API FMVVMBlueprintViewBinding* GetBindingAt(int32 Index);
+	UE_API const FMVVMBlueprintViewBinding* GetBindingAt(int32 Index) const;
+	UE_API FMVVMBlueprintViewBinding* GetBinding(FGuid Id);
+	UE_API const FMVVMBlueprintViewBinding* GetBinding(FGuid Id) const;
 
 	TArrayView<FMVVMBlueprintViewBinding> GetBindings()
 	{
@@ -135,11 +137,11 @@ public:
 		return Bindings;
 	}
 
-	UMVVMBlueprintViewEvent* AddDefaultEvent();
-	void AddEvent(UMVVMBlueprintViewEvent* Event);
-	void RemoveEvent(UMVVMBlueprintViewEvent* Event);
-	void ReplaceEvent(UMVVMBlueprintViewEvent* OldEvent, UMVVMBlueprintViewEvent* NewEvent);
-	UMVVMBlueprintViewEvent* DuplicateEvent(UMVVMBlueprintViewEvent* Event);
+	UE_API UMVVMBlueprintViewEvent* AddDefaultEvent();
+	UE_API void AddEvent(UMVVMBlueprintViewEvent* Event);
+	UE_API void RemoveEvent(UMVVMBlueprintViewEvent* Event);
+	UE_API void ReplaceEvent(UMVVMBlueprintViewEvent* OldEvent, UMVVMBlueprintViewEvent* NewEvent);
+	UE_API UMVVMBlueprintViewEvent* DuplicateEvent(UMVVMBlueprintViewEvent* Event);
 
 	TArrayView<TObjectPtr<UMVVMBlueprintViewEvent>> GetEvents()
 	{
@@ -151,11 +153,11 @@ public:
 		return Events;
 	}
 
-	UMVVMBlueprintViewCondition* AddDefaultCondition();
-	void AddCondition(UMVVMBlueprintViewCondition* Condition);
-	void RemoveCondition(UMVVMBlueprintViewCondition* Condition);
-	void ReplaceCondition(UMVVMBlueprintViewCondition* OldCondition, UMVVMBlueprintViewCondition* NewCondition);
-	UMVVMBlueprintViewCondition* DuplicateCondition(UMVVMBlueprintViewCondition* Condition);
+	UE_API UMVVMBlueprintViewCondition* AddDefaultCondition();
+	UE_API void AddCondition(UMVVMBlueprintViewCondition* Condition);
+	UE_API void RemoveCondition(UMVVMBlueprintViewCondition* Condition);
+	UE_API void ReplaceCondition(UMVVMBlueprintViewCondition* OldCondition, UMVVMBlueprintViewCondition* NewCondition);
+	UE_API UMVVMBlueprintViewCondition* DuplicateCondition(UMVVMBlueprintViewCondition* Condition);
 
 	TArrayView<TObjectPtr<UMVVMBlueprintViewCondition>> GetConditions()
 	{
@@ -172,10 +174,10 @@ public:
 		return !Bindings.IsEmpty() || !Events.IsEmpty() || !Conditions.IsEmpty();
 	}
 
-	TArray<FText> GetBindingMessages(FGuid Id, UE::MVVM::EBindingMessageType InMessageType) const;
-	bool HasBindingMessage(FGuid Id, UE::MVVM::EBindingMessageType InMessageType) const;
-	void AddMessageToBinding(FGuid Id, UE::MVVM::FBindingMessage MessageToAdd);
-	void ResetBindingMessages();
+	UE_API TArray<FText> GetBindingMessages(FGuid Id, UE::MVVM::EBindingMessageType InMessageType) const;
+	UE_API bool HasBindingMessage(FGuid Id, UE::MVVM::EBindingMessageType InMessageType) const;
+	UE_API void AddMessageToBinding(FGuid Id, UE::MVVM::FBindingMessage MessageToAdd);
+	UE_API void ResetBindingMessages();
 
 	FGuid GetCompiledBindingLibraryId() const
 	{
@@ -183,19 +185,19 @@ public:
 	}
 
 #if WITH_EDITOR
-	virtual void PostLoad() override;
-	virtual void PreSave(FObjectPreSaveContext Context) override;
-	virtual void PostEditChangeProperty(FPropertyChangedEvent& PropertyChangedEvent) override;
-	virtual void PostEditChangeChainProperty(FPropertyChangedChainEvent& PropertyChainEvent) override;
-	virtual void PostEditUndo() override;
+	UE_API virtual void PostLoad() override;
+	UE_API virtual void PreSave(FObjectPreSaveContext Context) override;
+	UE_API virtual void PostEditChangeProperty(FPropertyChangedEvent& PropertyChangedEvent) override;
+	UE_API virtual void PostEditChangeChainProperty(FPropertyChangedChainEvent& PropertyChainEvent) override;
+	UE_API virtual void PostEditUndo() override;
 
-	void AddAssetTags(FAssetRegistryTagsContext Context) const;
+	UE_API void AddAssetTags(FAssetRegistryTagsContext Context) const;
 	UE_DEPRECATED(5.4, "Implement the version that takes FAssetRegistryTagsContext instead.")
-	void AddAssetTags(TArray<FAssetRegistryTag>& OutTags) const;
-	void OnFieldRenamed(UClass* FieldOwnerClass, FName OldObjectName, FName NewObjectName);
+	UE_API void AddAssetTags(TArray<FAssetRegistryTag>& OutTags) const;
+	UE_API void OnFieldRenamed(UClass* FieldOwnerClass, FName OldObjectName, FName NewObjectName);
 #endif
 
-	virtual void Serialize(FArchive& Ar) override;
+	UE_API virtual void Serialize(FArchive& Ar) override;
 
 	DECLARE_EVENT(UMVVMBlueprintView, FOnBindingsUpdated);
 	FOnBindingsUpdated OnBindingsUpdated;
@@ -249,3 +251,5 @@ private:
 
 	bool bIsContextSensitive;
 };
+
+#undef UE_API
diff --git a/Engine/Plugins/Runtime/ModelViewViewModel/Source/ModelViewViewModelBlueprint/Public/MVVMBlueprintViewBinding.h b/Engine/Plugins/Runtime/ModelViewViewModel/Source/ModelViewViewModelBlueprint/Public/MVVMBlueprintViewBinding.h
--- a/Engine/Plugins/Runtime/ModelViewViewModel/Source/ModelViewViewModelBlueprint/Public/MVVMBlueprintViewBinding.h
+++ b/Engine/Plugins/Runtime/ModelViewViewModel/Source/ModelViewViewModelBlueprint/Public/MVVMBlueprintViewBinding.h
@@ -8,6 +8,8 @@
 
 #include "MVVMBlueprintViewBinding.generated.h"
 
+#define UE_API MODELVIEWVIEWMODELBLUEPRINT_API
+
 class UWidgetBlueprint;
 class UMVVMBlueprintView;
 class UMVVMBlueprintViewConversionFunction;
@@ -16,7 +18,7 @@ class UMVVMBlueprintViewConversionFunction;
 *
 */
 USTRUCT(BlueprintType)
-struct MODELVIEWVIEWMODELBLUEPRINT_API FMVVMBlueprintViewConversionPath
+struct FMVVMBlueprintViewConversionPath
 {
 	GENERATED_BODY()
 
@@ -72,9 +74,9 @@ public:
 		return bSourceToDestination ? SourceToDestinationConversion : DestinationToSourceConversion;
 	}
 
-	void GenerateWrapper(UBlueprint* Blueprint);
-	void SavePinValues(UBlueprint* Blueprint);
-	void DeprecateViewConversionFunction(UBlueprint* Blueprint, FMVVMBlueprintViewBinding& Owner);
+	UE_API void GenerateWrapper(UBlueprint* Blueprint);
+	UE_API void SavePinValues(UBlueprint* Blueprint);
+	UE_API void DeprecateViewConversionFunction(UBlueprint* Blueprint, FMVVMBlueprintViewBinding& Owner);
 
 public:
 PRAGMA_DISABLE_DEPRECATION_WARNINGS
@@ -88,7 +90,7 @@ PRAGMA_ENABLE_DEPRECATION_WARNINGS
 *
 */
 USTRUCT(BlueprintType)
-struct MODELVIEWVIEWMODELBLUEPRINT_API FMVVMBlueprintViewBinding
+struct FMVVMBlueprintViewBinding
 {
 	GENERATED_BODY()
 
@@ -127,18 +129,20 @@ struct MODELVIEWVIEWMODELBLUEPRINT_API FMVVMBlueprintViewBinding
 	/**
 	 * Get an internal name. For use in the UI, use GetDisplayNameString()
 	 */
-	FName GetFName() const;
+	UE_API FName GetFName() const;
 
 	/** 
 	 * Get a string that identifies this binding. 
 	 * This is of the form: Widget.Property <- ViewModel.Property
 	 */
-	FString GetDisplayNameString(const UWidgetBlueprint* WidgetBlueprint, bool bUseDisplayName = false) const;
+	UE_API FString GetDisplayNameString(const UWidgetBlueprint* WidgetBlueprint, bool bUseDisplayName = false) const;
 
 	/**
 	 * Get a string that identifies this binding and is specifically formatted for search. 
 	 * This includes the display name and variable name of all fields and widgets, as well as all function keywords.
 	 * For use in the UI, use GetDisplayNameString()
 	 */
-	FString GetSearchableString(const UWidgetBlueprint* WidgetBlueprint) const;
+	UE_API FString GetSearchableString(const UWidgetBlueprint* WidgetBlueprint) const;
 };
+
+#undef UE_API
diff --git a/Engine/Plugins/Runtime/ModelViewViewModel/Source/ModelViewViewModelBlueprint/Public/MVVMBlueprintViewCondition.h b/Engine/Plugins/Runtime/ModelViewViewModel/Source/ModelViewViewModelBlueprint/Public/MVVMBlueprintViewCondition.h
--- a/Engine/Plugins/Runtime/ModelViewViewModel/Source/ModelViewViewModelBlueprint/Public/MVVMBlueprintViewCondition.h
+++ b/Engine/Plugins/Runtime/ModelViewViewModel/Source/ModelViewViewModelBlueprint/Public/MVVMBlueprintViewCondition.h
@@ -9,6 +9,8 @@
 
 #include "MVVMBlueprintViewCondition.generated.h"
 
+#define UE_API MODELVIEWVIEWMODELBLUEPRINT_API
+
 struct FEdGraphEditAction;
 class UEdGraph;
 class UK2Node;
@@ -21,8 +23,8 @@ class UWidgetBlueprint;
  * 
  * Ex: UButton::OnClick 
  */
-UCLASS(Within = MVVMBlueprintView)
-class MODELVIEWVIEWMODELBLUEPRINT_API UMVVMBlueprintViewCondition : public UObject
+UCLASS(MinimalAPI, Within = MVVMBlueprintView)
+class UMVVMBlueprintViewCondition : public UObject
 {
 	GENERATED_BODY()
 
@@ -54,32 +56,32 @@ public:
 	{
 		return ConditionPath;
 	}
-	void SetConditionPath(FMVVMBlueprintPropertyPath ConditionPath);
+	UE_API void SetConditionPath(FMVVMBlueprintPropertyPath ConditionPath);
 
 	const FMVVMBlueprintPropertyPath& GetDestinationPath() const
 	{
 		return DestinationPath;
 	}
-	void SetDestinationPath(FMVVMBlueprintPropertyPath DestinationPath);
+	UE_API void SetDestinationPath(FMVVMBlueprintPropertyPath DestinationPath);
 
 	EMVVMConditionOperation GetOperation() const
 	{
 		return ConditionOperation;
 	}
 
-	void SetOperation(EMVVMConditionOperation Operation);
+	UE_API void SetOperation(EMVVMConditionOperation Operation);
 
 	float GetOperationValue() const
 	{
 		return Value;
 	}
-	void SetOperationValue(float NewValue);
+	UE_API void SetOperationValue(float NewValue);
 
 	float GetOperationMaxValue() const
 	{
 		return MaxValue;
 	}
-	void SetOperationMaxValue(float NewMaxValue);
+	UE_API void SetOperationMaxValue(float NewMaxValue);
 
 public:
 	UEdGraph* GetWrapperGraph() const
@@ -97,16 +99,16 @@ public:
 		RemoveConversionFunctionCurrentValues, // when removing or changing the conversion function, we want to remove all the conversion function parameters
 		LeaveConversionFunctionCurrentValues // when we remove the wrapper graph because the destination path or condition properties are changed, we want to keep the conversion function parameters
 	};
-	void RemoveWrapperGraph(ERemoveWrapperGraphParam ActionForCurrentValues = RemoveConversionFunctionCurrentValues);
+	UE_API void RemoveWrapperGraph(ERemoveWrapperGraphParam ActionForCurrentValues = RemoveConversionFunctionCurrentValues);
 
 	UK2Node* GetWrapperDestinationNode() const
 	{
 		return CachedWrapperDestinationNode;
 	}
 
-	UEdGraph* GetOrCreateWrapperGraph();
+	UE_API UEdGraph* GetOrCreateWrapperGraph();
 
-	void RecreateWrapperGraph();
+	UE_API void RecreateWrapperGraph();
 
 public:
 	TArrayView<const FMVVMBlueprintPin> GetPins() const
@@ -115,54 +117,54 @@ public:
 	}
 
 	/** Generates SavedPins from the wrapper graph, if it exists. */
-	void SavePinValues();
+	UE_API void SavePinValues();
 	/** Keep the orphaned pins. Add the missing pins. */
-	void UpdatePinValues();
+	UE_API void UpdatePinValues();
 	/** Keep the orphaned pins. Add the missing pins. */
-	bool HasOrphanedPin() const;
+	UE_API bool HasOrphanedPin() const;
 	/** Condition sources are tested at runtime to check if they are valid. */
-	void UpdateConditionKey(FMVVMViewClass_ConditionKey ConditionKey);
+	UE_API void UpdateConditionKey(FMVVMViewClass_ConditionKey ConditionKey);
 
-	UEdGraphPin* GetOrCreateGraphPin(const FMVVMBlueprintPinId& Pin);
+	UE_API UEdGraphPin* GetOrCreateGraphPin(const FMVVMBlueprintPinId& Pin);
 
-	FMVVMBlueprintPropertyPath GetPinPath(const FMVVMBlueprintPinId& Pin) const;
-	void SetPinPath(const FMVVMBlueprintPinId& Pin, const FMVVMBlueprintPropertyPath& Path);
+	UE_API FMVVMBlueprintPropertyPath GetPinPath(const FMVVMBlueprintPinId& Pin) const;
+	UE_API void SetPinPath(const FMVVMBlueprintPinId& Pin, const FMVVMBlueprintPropertyPath& Path);
 	// To set a pin when loading the asset (no graph generation)
-	void SetPinPathNoGraphGeneration(const FMVVMBlueprintPinId& Pin, const FMVVMBlueprintPropertyPath& Path);
+	UE_API void SetPinPathNoGraphGeneration(const FMVVMBlueprintPinId& Pin, const FMVVMBlueprintPropertyPath& Path);
 
 	FSimpleMulticastDelegate OnWrapperGraphModified;
 
 public:
-	TArray<FText> GetCompilationMessages(EMessageType InMessageType) const;
-	bool HasCompilationMessage(EMessageType InMessageType) const;
-	void AddCompilationToBinding(FMessage MessageToAdd) const;
-	void ResetCompilationMessages();
+	UE_API TArray<FText> GetCompilationMessages(EMessageType InMessageType) const;
+	UE_API bool HasCompilationMessage(EMessageType InMessageType) const;
+	UE_API void AddCompilationToBinding(FMessage MessageToAdd) const;
+	UE_API void ResetCompilationMessages();
 
 public:
 	/** 
 	 * Get a string that identifies this event. 
 	 */
-	FText GetDisplayName(bool bUseDisplayName) const;
+	UE_API FText GetDisplayName(bool bUseDisplayName) const;
 
 	/**
 	 * Get a string that identifies this event and is specifically formatted for search. 
 	 * This includes the display name and variable name of all fields and widgets, as well as all function keywords.
 	 * For use in the UI, use GetDisplayNameString()
 	 */
-	FString GetSearchableString() const;
+	UE_API FString GetSearchableString() const;
 
 public:
-	virtual void PostEditChangeChainProperty(FPropertyChangedChainEvent& PropertyChainEvent) override;
+	UE_API virtual void PostEditChangeChainProperty(FPropertyChangedChainEvent& PropertyChainEvent) override;
 
 private:
-	const UFunction* GetDestinationSignature() const;
-	void HandleGraphChanged(const FEdGraphEditAction& Action);
-	void HandleUserDefinedPinRenamed(UK2Node* InNode, FName OldPinName, FName NewPinName);
-	UWidgetBlueprint* GetWidgetBlueprintInternal() const;
-	void SetCachedWrapperGraphInternal(UEdGraph* Graph, UK2Node* Node, UMVVMK2Node_IsConditionValid* SourceNode);
-	UEdGraph* CreateWrapperGraphInternal();
-	void LoadPinValuesInternal();
-	void UpdateConditionKeyInternal();
+	UE_API const UFunction* GetDestinationSignature() const;
+	UE_API void HandleGraphChanged(const FEdGraphEditAction& Action);
+	UE_API void HandleUserDefinedPinRenamed(UK2Node* InNode, FName OldPinName, FName NewPinName);
+	UE_API UWidgetBlueprint* GetWidgetBlueprintInternal() const;
+	UE_API void SetCachedWrapperGraphInternal(UEdGraph* Graph, UK2Node* Node, UMVVMK2Node_IsConditionValid* SourceNode);
+	UE_API UEdGraph* CreateWrapperGraphInternal();
+	UE_API void LoadPinValuesInternal();
+	UE_API void UpdateConditionKeyInternal();
 
 private:
 	UPROPERTY(VisibleAnywhere, Category = "Viewmodel")
@@ -210,3 +212,5 @@ private:
 
 	bool bNeedsToRegenerateChildren;
 };
+
+#undef UE_API
diff --git a/Engine/Plugins/Runtime/ModelViewViewModel/Source/ModelViewViewModelBlueprint/Public/MVVMBlueprintViewConversionFunction.h b/Engine/Plugins/Runtime/ModelViewViewModel/Source/ModelViewViewModelBlueprint/Public/MVVMBlueprintViewConversionFunction.h
--- a/Engine/Plugins/Runtime/ModelViewViewModel/Source/ModelViewViewModelBlueprint/Public/MVVMBlueprintViewConversionFunction.h
+++ b/Engine/Plugins/Runtime/ModelViewViewModel/Source/ModelViewViewModelBlueprint/Public/MVVMBlueprintViewConversionFunction.h
@@ -10,6 +10,8 @@
 
 #include "MVVMBlueprintViewConversionFunction.generated.h"
 
+#define UE_API MODELVIEWVIEWMODELBLUEPRINT_API
+
 class UK2Node;
 class UEdGraphPin;
 class UEdGraph;
@@ -21,65 +23,65 @@ struct FEdGraphEditAction;
  * 
  * Internally that function may be using native C++, K2Nodes, UFunctions, Events, etc.
  */
-UCLASS()
-class MODELVIEWVIEWMODELBLUEPRINT_API UMVVMBlueprintViewConversionFunction : public UObject
+UCLASS(MinimalAPI)
+class UMVVMBlueprintViewConversionFunction : public UObject
 {
 	GENERATED_BODY()
 
 public:
-	static bool IsValidConversionFunction(const UBlueprint* WidgetBlueprint, const UFunction* Function);
-	static bool IsValidConversionNode(const UBlueprint* WidgetBlueprint, const TSubclassOf<UK2Node> Function);
+	static UE_API bool IsValidConversionFunction(const UBlueprint* WidgetBlueprint, const UFunction* Function);
+	static UE_API bool IsValidConversionNode(const UBlueprint* WidgetBlueprint, const TSubclassOf<UK2Node> Function);
 
 public:
 	/** @return the conversion function uses at runtime. The wrapper function if complex or GetFunction is simple. */
-	const UFunction* GetCompiledFunction(const UClass* SelfContext) const;
+	UE_API const UFunction* GetCompiledFunction(const UClass* SelfContext) const;
 
 	/** @return the conversion function uses at runtime. The wrapper function if complex or GetFunction is simple. */
-	FName GetCompiledFunctionName(const UClass* SelfContext) const;
+	UE_API FName GetCompiledFunctionName(const UClass* SelfContext) const;
 
 	/** @return the conversion function. */
 	UE_DEPRECATED(5.4, "GetConversionFunction that returns a variant is deprecated.")
-	TVariant<const UFunction*, TSubclassOf<UK2Node>> GetConversionFunction(const UBlueprint* SelfContext) const;
+	UE_API TVariant<const UFunction*, TSubclassOf<UK2Node>> GetConversionFunction(const UBlueprint* SelfContext) const;
 
 	/** @return the conversion function. */
-	FMVVMBlueprintFunctionReference GetConversionFunction() const;
+	UE_API FMVVMBlueprintFunctionReference GetConversionFunction() const;
 
 	/** Set the function. Generate a Graph. */
-	void Initialize(UBlueprint* SelfContext, FName GraphName, FMVVMBlueprintFunctionReference Function);
+	UE_API void Initialize(UBlueprint* SelfContext, FName GraphName, FMVVMBlueprintFunctionReference Function);
 
 	/** Set the function. Generate a Graph. */
-	void InitializeFromFunction(UBlueprint* SelfContext, FName GraphName, const UFunction* Function);
+	UE_API void InitializeFromFunction(UBlueprint* SelfContext, FName GraphName, const UFunction* Function);
 
 	// For deprecation
-	void Deprecation_InitializeFromWrapperGraph(UBlueprint* SelfContext, UEdGraph* Graph);
+	UE_API void Deprecation_InitializeFromWrapperGraph(UBlueprint* SelfContext, UEdGraph* Graph);
 
 	// For deprecation
-	void Deprecation_InitializeFromMemberReference(UBlueprint* SelfContext, FName GraphName, FMemberReference MemberReference, const FMVVMBlueprintPropertyPath& Source);
+	UE_API void Deprecation_InitializeFromMemberReference(UBlueprint* SelfContext, FName GraphName, FMemberReference MemberReference, const FMVVMBlueprintPropertyPath& Source);
 
 	// For deprecation
-	void Deprecation_SetWrapperGraphName(UBlueprint* Context, FName GraphName, const FMVVMBlueprintPropertyPath& Source);
+	UE_API void Deprecation_SetWrapperGraphName(UBlueprint* Context, FName GraphName, const FMVVMBlueprintPropertyPath& Source);
 
 	/**
 	 * The conversion is valid.
 	 * The function was valid when created but may not be anymore.
 	 * It doesn't check if the source and destination are valid.
 	 */
-	bool IsValid(const UBlueprint* SelfContext) const;
+	UE_API bool IsValid(const UBlueprint* SelfContext) const;
 
 	/** The function has more than one argument and requires a wrapper or it uses a FunctionNode. */
-	bool NeedsWrapperGraph(const UBlueprint* SelfContext) const;
+	UE_API bool NeedsWrapperGraph(const UBlueprint* SelfContext) const;
 
 	/** The wrapper Graph is generated on load/compile and is not saved. */
-	bool IsWrapperGraphTransient() const;
+	UE_API bool IsWrapperGraphTransient() const;
 
 	/** True if the graph is going to be used for an ubergraph page. */
-	bool IsUbergraphPage() const;
+	UE_API bool IsUbergraphPage() const;
 
 	const FMVVMBlueprintPropertyPath& GetDestinationPath() const
 	{
 		return DestinationPath;
 	}
-	void SetDestinationPath(FMVVMBlueprintPropertyPath DestinationPath);
+	UE_API void SetDestinationPath(FMVVMBlueprintPropertyPath DestinationPath);
 
 	/** Return the wrapper graph, if it exists. */
 	UEdGraph* GetWrapperGraph() const
@@ -108,24 +110,24 @@ public:
 	 * If needed, create the graph and all the nodes for that graph when compiling.
 	 * Returns the existing one, if one was created from GetOrCreateWrapperGraph.
 	 */
-	UEdGraph* GetOrCreateIntermediateWrapperGraph(FKismetCompilerContext& Context);
+	UE_API UEdGraph* GetOrCreateIntermediateWrapperGraph(FKismetCompilerContext& Context);
 
 	/** If needed, create the graph and all the nodes for that graph. */
-	UEdGraph* GetOrCreateWrapperGraph(UBlueprint* Blueprint);
+	UE_API UEdGraph* GetOrCreateWrapperGraph(UBlueprint* Blueprint);
 
-	void RecreateWrapperGraph(UBlueprint* Blueprint);
+	UE_API void RecreateWrapperGraph(UBlueprint* Blueprint);
 
 	/**
 	 * The conversion function is going to be removed from the Blueprint.
 	 * Do any cleanup that is needed.
 	 */
-	void RemoveWrapperGraph(UBlueprint* Blueprint);
+	UE_API void RemoveWrapperGraph(UBlueprint* Blueprint);
 
 	/**
 	 * Returns the pin from the graph.
 	 * Create the graph and all the nodes for that graph if the graph doesn't exist and it's needed.
 	 */
-	UEdGraphPin* GetOrCreateGraphPin(UBlueprint* Blueprint, const FMVVMBlueprintPinId& PinId);
+	UE_API UEdGraphPin* GetOrCreateGraphPin(UBlueprint* Blueprint, const FMVVMBlueprintPinId& PinId);
 
 	const TArrayView<const FMVVMBlueprintPin> GetPins() const
 	{
@@ -133,32 +135,32 @@ public:
 	}
 
 	/** */
-	void SetGraphPin(UBlueprint* Blueprint, const FMVVMBlueprintPinId& PinId, const FMVVMBlueprintPropertyPath& Value);
+	UE_API void SetGraphPin(UBlueprint* Blueprint, const FMVVMBlueprintPinId& PinId, const FMVVMBlueprintPropertyPath& Value);
 
 	/** Generates SavedPins from the wrapper graph, if it exists. */
-	void SavePinValues(UBlueprint* Blueprint);
+	UE_API void SavePinValues(UBlueprint* Blueprint);
 	/** Keep the orphaned pins. Add the missing pins. */
-	void UpdatePinValues(UBlueprint* Blueprint);
+	UE_API void UpdatePinValues(UBlueprint* Blueprint);
 	/** Keep the orphaned pins. Add the missing pins. */
-	bool HasOrphanedPin() const;
+	UE_API bool HasOrphanedPin() const;
 
 	FSimpleMulticastDelegate OnWrapperGraphModified;
 
-	virtual void PostLoad() override;
+	UE_API virtual void PostLoad() override;
 
 private:
-	void HandleGraphChanged(const FEdGraphEditAction& Action, TWeakObjectPtr<UBlueprint> Context);
-	void HandleUserDefinedPinRenamed(UK2Node* InNode, FName OldPinName, FName NewPinName, TWeakObjectPtr<UBlueprint> WeakBlueprint);
-	void SetCachedWrapperGraph(UBlueprint* Blueprint, UEdGraph* CachedGraph, UK2Node* CachedNode);
-	UEdGraph* CreateWrapperGraphInternal(FKismetCompilerContext& Context);
-	UEdGraph* CreateWrapperGraphInternal(UBlueprint* Blueprint);
-	bool NeedsWrapperGraphInternal(const UClass* SkeletalSelfContext) const;
-	void LoadPinValuesInternal(UBlueprint* Blueprint);
-	FName GenerateUniqueGraphName() const;
-	void CreateWrapperGraphName();
-	void Reset();
-
-	bool IsConversionFunctionAsyncNode();
+	UE_API void HandleGraphChanged(const FEdGraphEditAction& Action, TWeakObjectPtr<UBlueprint> Context);
+	UE_API void HandleUserDefinedPinRenamed(UK2Node* InNode, FName OldPinName, FName NewPinName, TWeakObjectPtr<UBlueprint> WeakBlueprint);
+	UE_API void SetCachedWrapperGraph(UBlueprint* Blueprint, UEdGraph* CachedGraph, UK2Node* CachedNode);
+	UE_API UEdGraph* CreateWrapperGraphInternal(FKismetCompilerContext& Context);
+	UE_API UEdGraph* CreateWrapperGraphInternal(UBlueprint* Blueprint);
+	UE_API bool NeedsWrapperGraphInternal(const UClass* SkeletalSelfContext) const;
+	UE_API void LoadPinValuesInternal(UBlueprint* Blueprint);
+	UE_API FName GenerateUniqueGraphName() const;
+	UE_API void CreateWrapperGraphName();
+	UE_API void Reset();
+
+	UE_API bool IsConversionFunctionAsyncNode();
 
 private:
 
@@ -214,3 +216,5 @@ private:
 	UPROPERTY()
 	TSubclassOf<UK2Node> FunctionNode_DEPRECATED;
 };
+
+#undef UE_API
diff --git a/Engine/Plugins/Runtime/ModelViewViewModel/Source/ModelViewViewModelBlueprint/Public/MVVMBlueprintViewEvent.h b/Engine/Plugins/Runtime/ModelViewViewModel/Source/ModelViewViewModelBlueprint/Public/MVVMBlueprintViewEvent.h
--- a/Engine/Plugins/Runtime/ModelViewViewModel/Source/ModelViewViewModelBlueprint/Public/MVVMBlueprintViewEvent.h
+++ b/Engine/Plugins/Runtime/ModelViewViewModel/Source/ModelViewViewModelBlueprint/Public/MVVMBlueprintViewEvent.h
@@ -8,6 +8,8 @@
 
 #include "MVVMBlueprintViewEvent.generated.h"
 
+#define UE_API MODELVIEWVIEWMODELBLUEPRINT_API
+
 struct FEdGraphEditAction;
 class UEdGraph;
 class UK2Node;
@@ -20,8 +22,8 @@ class UWidgetBlueprint;
  * 
  * Ex: UButton::OnClick 
  */
-UCLASS(Within = MVVMBlueprintView)
-class MODELVIEWVIEWMODELBLUEPRINT_API UMVVMBlueprintViewEvent : public UObject
+UCLASS(MinimalAPI, Within = MVVMBlueprintView)
+class UMVVMBlueprintViewEvent : public UObject
 {
 	GENERATED_BODY()
 
@@ -40,7 +42,7 @@ public:
 	};
 
 	/** @return true if the property path contains a multicast delegate property. */
-	static bool Supports(const UWidgetBlueprint* WidgetBlueprint, const FMVVMBlueprintPropertyPath& PropertyPath);
+	static UE_API bool Supports(const UWidgetBlueprint* WidgetBlueprint, const FMVVMBlueprintPropertyPath& PropertyPath);
 
 public:
 	/** Whether the event is enabled or disabled by default. The instance may enable the event at runtime. */
@@ -56,13 +58,13 @@ public:
 	{
 		return EventPath;
 	}
-	void SetEventPath(FMVVMBlueprintPropertyPath EventPath);
+	UE_API void SetEventPath(FMVVMBlueprintPropertyPath EventPath);
 
 	const FMVVMBlueprintPropertyPath& GetDestinationPath() const
 	{
 		return DestinationPath;
 	}
-	void SetDestinationPath(FMVVMBlueprintPropertyPath DestinationPath);
+	UE_API void SetDestinationPath(FMVVMBlueprintPropertyPath DestinationPath);
 
 public:
 	UEdGraph* GetWrapperGraph() const
@@ -80,16 +82,16 @@ public:
 		RemoveConversionFunctionCurrentValues, // when removing or changing the conversion function, we want to remove all the conversion function parameters
 		LeaveConversionFunctionCurrentValues // when we remove the wrapper graph because the event path has changed, we want to keep the conversion function parameters
 	};
-	void RemoveWrapperGraph(ERemoveWrapperGraphParam ActionForCurrentValues = RemoveConversionFunctionCurrentValues);
+	UE_API void RemoveWrapperGraph(ERemoveWrapperGraphParam ActionForCurrentValues = RemoveConversionFunctionCurrentValues);
 
 	UK2Node* GetWrapperNode() const
 	{
 		return CachedWrapperNode;
 	}
 
-	UEdGraph* GetOrCreateWrapperGraph();
+	UE_API UEdGraph* GetOrCreateWrapperGraph();
 
-	void RecreateWrapperGraph();
+	UE_API void RecreateWrapperGraph();
 
 public:
 	TArrayView<const FMVVMBlueprintPin> GetPins() const
@@ -98,55 +100,55 @@ public:
 	}
 
 	/** Generates SavedPins from the wrapper graph, if it exists. */
-	void SavePinValues();
+	UE_API void SavePinValues();
 	/** Keep the orphaned pins. Add the missing pins. */
-	void UpdatePinValues();
+	UE_API void UpdatePinValues();
 	/** Keep the orphaned pins. Add the missing pins. */
-	bool HasOrphanedPin() const;
+	UE_API bool HasOrphanedPin() const;
 	/** Event sources are tested at runtime to check if they are valid. */
-	void UpdateEventKey(FMVVMViewClass_EventKey EventKey);
+	UE_API void UpdateEventKey(FMVVMViewClass_EventKey EventKey);
 
-	UEdGraphPin* GetOrCreateGraphPin(const FMVVMBlueprintPinId& Pin);
+	UE_API UEdGraphPin* GetOrCreateGraphPin(const FMVVMBlueprintPinId& Pin);
 
-	FMVVMBlueprintPropertyPath GetPinPath(const FMVVMBlueprintPinId& Pin) const;
-	void SetPinPath(const FMVVMBlueprintPinId& Pin, const FMVVMBlueprintPropertyPath& Path);
+	UE_API FMVVMBlueprintPropertyPath GetPinPath(const FMVVMBlueprintPinId& Pin) const;
+	UE_API void SetPinPath(const FMVVMBlueprintPinId& Pin, const FMVVMBlueprintPropertyPath& Path);
 	// To set a pin when loading the asset (no graph generation)
-	void SetPinPathNoGraphGeneration(const FMVVMBlueprintPinId& Pin, const FMVVMBlueprintPropertyPath& Path);
+	UE_API void SetPinPathNoGraphGeneration(const FMVVMBlueprintPinId& Pin, const FMVVMBlueprintPropertyPath& Path);
 
 	FSimpleMulticastDelegate OnWrapperGraphModified;
 
 public:
-	TArray<FText> GetCompilationMessages(EMessageType InMessageType) const;
-	bool HasCompilationMessage(EMessageType InMessageType) const;
-	void AddCompilationToBinding(FMessage MessageToAdd) const;
-	void ResetCompilationMessages();
+	UE_API TArray<FText> GetCompilationMessages(EMessageType InMessageType) const;
+	UE_API bool HasCompilationMessage(EMessageType InMessageType) const;
+	UE_API void AddCompilationToBinding(FMessage MessageToAdd) const;
+	UE_API void ResetCompilationMessages();
 
 public:
 	/** 
 	 * Get a string that identifies this event. 
 	 */
-	FText GetDisplayName(bool bUseDisplayName) const;
+	UE_API FText GetDisplayName(bool bUseDisplayName) const;
 
 	/**
 	 * Get a string that identifies this event and is specifically formatted for search. 
 	 * This includes the display name and variable name of all fields and widgets, as well as all function keywords.
 	 * For use in the UI, use GetDisplayNameString()
 	 */
-	FString GetSearchableString() const;
+	UE_API FString GetSearchableString() const;
 
 public:
-	virtual void PostEditChangeChainProperty(FPropertyChangedChainEvent& PropertyChainEvent) override;
+	UE_API virtual void PostEditChangeChainProperty(FPropertyChangedChainEvent& PropertyChainEvent) override;
 
 private:
-	static const UFunction* GetEventSignature(const UWidgetBlueprint* WidgetBlueprint, const FMVVMBlueprintPropertyPath& PropertyPath);
-	const UFunction* GetEventSignature() const;
-	void HandleGraphChanged(const FEdGraphEditAction& Action);
-	void HandleUserDefinedPinRenamed(UK2Node* InNode, FName OldPinName, FName NewPinName);
-	UWidgetBlueprint* GetWidgetBlueprintInternal() const;
-	void SetCachedWrapperGraphInternal(UEdGraph* Graph, UK2Node* Node, UMVVMK2Node_AreSourcesValidForEvent* SourceNode);
-	UEdGraph* CreateWrapperGraphInternal();
-	void LoadPinValuesInternal();
-	void UpdateEventKeyInternal();
+	static UE_API const UFunction* GetEventSignature(const UWidgetBlueprint* WidgetBlueprint, const FMVVMBlueprintPropertyPath& PropertyPath);
+	UE_API const UFunction* GetEventSignature() const;
+	UE_API void HandleGraphChanged(const FEdGraphEditAction& Action);
+	UE_API void HandleUserDefinedPinRenamed(UK2Node* InNode, FName OldPinName, FName NewPinName);
+	UE_API UWidgetBlueprint* GetWidgetBlueprintInternal() const;
+	UE_API void SetCachedWrapperGraphInternal(UEdGraph* Graph, UK2Node* Node, UMVVMK2Node_AreSourcesValidForEvent* SourceNode);
+	UE_API UEdGraph* CreateWrapperGraphInternal();
+	UE_API void LoadPinValuesInternal();
+	UE_API void UpdateEventKeyInternal();
 
 private:
 	UPROPERTY(VisibleAnywhere, Category = "Viewmodel")
@@ -185,3 +187,5 @@ private:
 
 	bool bNeedsToRegenerateChildren;
 };
+
+#undef UE_API
diff --git a/Engine/Plugins/Runtime/ModelViewViewModel/Source/ModelViewViewModelBlueprint/Public/MVVMBlueprintViewModel.h b/Engine/Plugins/Runtime/ModelViewViewModel/Source/ModelViewViewModelBlueprint/Public/MVVMBlueprintViewModel.h
--- a/Engine/Plugins/Runtime/ModelViewViewModel/Source/ModelViewViewModelBlueprint/Public/MVVMBlueprintViewModel.h
+++ b/Engine/Plugins/Runtime/ModelViewViewModel/Source/ModelViewViewModelBlueprint/Public/MVVMBlueprintViewModel.h
@@ -6,8 +6,8 @@
 #include "MVVMBlueprintViewModel.generated.h"
 
 
-UCLASS(DisplayName="ViewModel Blueprint")
-class MODELVIEWVIEWMODELBLUEPRINT_API UMVVMBlueprintViewModel : public UObject
+UCLASS(MinimalAPI, DisplayName="ViewModel Blueprint")
+class UMVVMBlueprintViewModel : public UObject
 {
 	GENERATED_BODY()
 
diff --git a/Engine/Plugins/Runtime/ModelViewViewModel/Source/ModelViewViewModelBlueprint/Public/MVVMBlueprintViewModelContext.h b/Engine/Plugins/Runtime/ModelViewViewModel/Source/ModelViewViewModelBlueprint/Public/MVVMBlueprintViewModelContext.h
--- a/Engine/Plugins/Runtime/ModelViewViewModel/Source/ModelViewViewModelBlueprint/Public/MVVMBlueprintViewModelContext.h
+++ b/Engine/Plugins/Runtime/ModelViewViewModel/Source/ModelViewViewModelBlueprint/Public/MVVMBlueprintViewModelContext.h
@@ -7,6 +7,8 @@
 
 #include "MVVMBlueprintViewModelContext.generated.h"
 
+#define UE_API MODELVIEWVIEWMODELBLUEPRINT_API
+
 class UMVVMBlueprintInstancedViewModelBase;
 class UMVVMViewModelContextResolver;
 
@@ -39,13 +41,13 @@ namespace UE::MVVM
  *
  */
 USTRUCT()
-struct MODELVIEWVIEWMODELBLUEPRINT_API FMVVMBlueprintViewModelContext
+struct FMVVMBlueprintViewModelContext
 {
 	GENERATED_BODY()
 
 public:
 	FMVVMBlueprintViewModelContext() = default;
-	FMVVMBlueprintViewModelContext(const UClass* InClass, FName InViewModelName);
+	UE_API FMVVMBlueprintViewModelContext(const UClass* InClass, FName InViewModelName);
 
 	FGuid GetViewModelId() const
 	{
@@ -57,7 +59,7 @@ public:
 		return ViewModelName;
 	}
 
-	FText GetDisplayName() const;
+	UE_API FText GetDisplayName() const;
 
 	UClass* GetViewModelClass() const
 	{
@@ -103,7 +105,7 @@ public:
 	}
 
 #if WITH_EDITOR
-	[[nodiscard]] TObjectPtr<UMVVMViewModelContextResolver> CreateDefaultResolver(UPackage* Package) const;
+	[[nodiscard]] UE_API TObjectPtr<UMVVMViewModelContextResolver> CreateDefaultResolver(UPackage* Package) const;
 #endif
 
 private:
@@ -213,3 +215,5 @@ struct TStructOpsTypeTraits<FMVVMBlueprintViewModelContext> : public TStructOpsT
 		WithPostSerialize = true,
 	};
 };
+
+#undef UE_API
diff --git a/Engine/Plugins/Runtime/ModelViewViewModel/Source/ModelViewViewModelBlueprint/Public/MVVMDeveloperProjectSettings.h b/Engine/Plugins/Runtime/ModelViewViewModel/Source/ModelViewViewModelBlueprint/Public/MVVMDeveloperProjectSettings.h
--- a/Engine/Plugins/Runtime/ModelViewViewModel/Source/ModelViewViewModelBlueprint/Public/MVVMDeveloperProjectSettings.h
+++ b/Engine/Plugins/Runtime/ModelViewViewModel/Source/ModelViewViewModelBlueprint/Public/MVVMDeveloperProjectSettings.h
@@ -7,6 +7,8 @@
 
 #include "MVVMDeveloperProjectSettings.generated.h"
 
+#define UE_API MODELVIEWVIEWMODELBLUEPRINT_API
+
 class UK2Node;
 class UListViewBase;
 class UMVVMViewModelContextResolver;
@@ -63,25 +65,25 @@ enum class EMVVMDeveloperConversionFunctionFilterType : uint8
 /**
  * Implements the settings for the MVVM Editor
  */
-UCLASS(config=ModelViewViewModel, defaultconfig)
-class MODELVIEWVIEWMODELBLUEPRINT_API UMVVMDeveloperProjectSettings : public UDeveloperSettings
+UCLASS(MinimalAPI, config=ModelViewViewModel, defaultconfig)
+class UMVVMDeveloperProjectSettings : public UDeveloperSettings
 {
 	GENERATED_BODY()
 
 public:
-	UMVVMDeveloperProjectSettings();
+	UE_API UMVVMDeveloperProjectSettings();
 
-	virtual FName GetCategoryName() const override;
-	virtual FText GetSectionText() const override;
+	UE_API virtual FName GetCategoryName() const override;
+	UE_API virtual FText GetSectionText() const override;
 #if WITH_EDITOR
-	virtual void PostEditChangeChainProperty(FPropertyChangedChainEvent& PropertyChangedEvent) override;
+	UE_API virtual void PostEditChangeChainProperty(FPropertyChangedChainEvent& PropertyChangedEvent) override;
 #endif
 
-	bool PropertyHasFiltering(const UStruct* ObjectStruct, const FProperty* Property) const;
-	bool IsPropertyAllowed(const UBlueprint* Context, const UStruct* ObjectStruct, const FProperty* Property) const;
-	bool IsFunctionAllowed(const UBlueprint* Context, const UClass* ObjectClass, const UFunction* Function) const;
-	bool IsConversionFunctionAllowed(const UBlueprint* Context, const UFunction* Function) const;
-	bool IsConversionFunctionAllowed(const UBlueprint* Context, const TSubclassOf<UK2Node> Function) const;
+	UE_API bool PropertyHasFiltering(const UStruct* ObjectStruct, const FProperty* Property) const;
+	UE_API bool IsPropertyAllowed(const UBlueprint* Context, const UStruct* ObjectStruct, const FProperty* Property) const;
+	UE_API bool IsFunctionAllowed(const UBlueprint* Context, const UClass* ObjectClass, const UFunction* Function) const;
+	UE_API bool IsConversionFunctionAllowed(const UBlueprint* Context, const UFunction* Function) const;
+	UE_API bool IsConversionFunctionAllowed(const UBlueprint* Context, const TSubclassOf<UK2Node> Function) const;
 
 	bool IsExecutionModeAllowed(EMVVMExecutionMode ExecutionMode) const
 	{
@@ -98,11 +100,11 @@ public:
 		return ConversionFunctionFilter;
 	}
 
-	TArray<const UClass*> GetAllowedConversionFunctionClasses() const;
-	TArray<const UClass*> GetDeniedConversionFunctionClasses() const;
+	UE_API TArray<const UClass*> GetAllowedConversionFunctionClasses() const;
+	UE_API TArray<const UClass*> GetDeniedConversionFunctionClasses() const;
 
-	bool IsExtensionSupportedForPanelClass(TSubclassOf<UPanelWidget> ClassToSupport) const;
-	bool IsExtensionSupportedForListViewBaseClass(TSubclassOf<UListViewBase> ClassToSupport) const;
+	UE_API bool IsExtensionSupportedForPanelClass(TSubclassOf<UPanelWidget> ClassToSupport) const;
+	UE_API bool IsExtensionSupportedForListViewBaseClass(TSubclassOf<UListViewBase> ClassToSupport) const;
 
 	/** Get generated functions permission list. */
 	FPathPermissionList& GetGeneratedFunctionPermissions() { return GeneratedFunctionPermissions; }
@@ -213,3 +215,5 @@ private:
 	/** MVVM Generated function permission list */
 	FPathPermissionList GeneratedFunctionPermissions;
 };
+
+#undef UE_API
diff --git a/Engine/Plugins/Runtime/ModelViewViewModel/Source/ModelViewViewModelBlueprint/Public/ViewModel/MVVMViewModelBlueprint.h b/Engine/Plugins/Runtime/ModelViewViewModel/Source/ModelViewViewModelBlueprint/Public/ViewModel/MVVMViewModelBlueprint.h
--- a/Engine/Plugins/Runtime/ModelViewViewModel/Source/ModelViewViewModelBlueprint/Public/ViewModel/MVVMViewModelBlueprint.h
+++ b/Engine/Plugins/Runtime/ModelViewViewModel/Source/ModelViewViewModelBlueprint/Public/ViewModel/MVVMViewModelBlueprint.h
@@ -16,8 +16,8 @@ class UMVVMViewModelBlueprintGeneratedClass;
 /**
  * Will be deprecated in the next version.
  */
-UCLASS()
-class MODELVIEWVIEWMODELBLUEPRINT_API UMVVMViewModelBlueprint : public UBlueprint
+UCLASS(MinimalAPI)
+class UMVVMViewModelBlueprint : public UBlueprint
 {
 	GENERATED_BODY()
 };
diff --git a/Engine/Plugins/Runtime/ModelViewViewModel/Source/ModelViewViewModelDebugger/Public/MVVMDebugSnapshot.h b/Engine/Plugins/Runtime/ModelViewViewModel/Source/ModelViewViewModelDebugger/Public/MVVMDebugSnapshot.h
--- a/Engine/Plugins/Runtime/ModelViewViewModel/Source/ModelViewViewModelDebugger/Public/MVVMDebugSnapshot.h
+++ b/Engine/Plugins/Runtime/ModelViewViewModel/Source/ModelViewViewModelDebugger/Public/MVVMDebugSnapshot.h
@@ -7,6 +7,8 @@
 #include "UObject/ObjectMacros.h"
 #include "MVVMDebugItemId.h"
 
+#define UE_API MODELVIEWVIEWMODELDEBUGGER_API
+
 struct FMVVMViewDebugEntry;
 struct FMVVMViewClassDebugEntry;
 struct FMVVMViewModelDebugEntry;
@@ -15,7 +17,7 @@ class UMVVMViewClass;
 namespace UE::MVVM
 {
 
-class MODELVIEWVIEWMODELDEBUGGER_API FDebugSnapshot
+class FDebugSnapshot
 {
 private:
 	TArray<TSharedPtr<FMVVMViewDebugEntry>> Views;
@@ -32,13 +34,15 @@ public:
 		return ViewModels;
 	}
 
-	TSharedPtr<FMVVMViewDebugEntry> FindView(FGuid Id) const;
-	TSharedPtr<FMVVMViewModelDebugEntry> FindViewModel(FGuid Id) const;
+	UE_API TSharedPtr<FMVVMViewDebugEntry> FindView(FGuid Id) const;
+	UE_API TSharedPtr<FMVVMViewModelDebugEntry> FindViewModel(FGuid Id) const;
 
-	static TSharedPtr<FDebugSnapshot> CreateSnapshot();
+	static UE_API TSharedPtr<FDebugSnapshot> CreateSnapshot();
 
 private:
-	TSharedRef<FMVVMViewClassDebugEntry> FindOrAddViewClassEntry(const UMVVMViewClass* ViewClass);
+	UE_API TSharedRef<FMVVMViewClassDebugEntry> FindOrAddViewClassEntry(const UMVVMViewClass* ViewClass);
+};
+
 };
 
-};
\ No newline at end of file
+#undef UE_API
diff --git a/Engine/Plugins/Runtime/ModelViewViewModel/Source/ModelViewViewModelEditor/Public/MVVMEditorSubsystem.h b/Engine/Plugins/Runtime/ModelViewViewModel/Source/ModelViewViewModelEditor/Public/MVVMEditorSubsystem.h
--- a/Engine/Plugins/Runtime/ModelViewViewModel/Source/ModelViewViewModelEditor/Public/MVVMEditorSubsystem.h
+++ b/Engine/Plugins/Runtime/ModelViewViewModel/Source/ModelViewViewModelEditor/Public/MVVMEditorSubsystem.h
@@ -9,6 +9,8 @@
 
 #include "MVVMEditorSubsystem.generated.h"
 
+#define UE_API MODELVIEWVIEWMODELEDITOR_API
+
 class UEdGraphPin;
 class UMVVMBlueprintView;
 enum class EMVVMBindingMode : uint8;
@@ -29,142 +31,144 @@ class UMVVMBlueprintViewCondition;
 class UWidgetBlueprint;
 
 /** */
-UCLASS(DisplayName="Viewmodel Editor Subsystem")
-class MODELVIEWVIEWMODELEDITOR_API UMVVMEditorSubsystem : public UEditorSubsystem
+UCLASS(MinimalAPI, DisplayName="Viewmodel Editor Subsystem")
+class UMVVMEditorSubsystem : public UEditorSubsystem
 {
 	GENERATED_BODY()
 
 public:
 	UFUNCTION(BlueprintCallable, Category = "Viewmodel")
-	UMVVMBlueprintView* RequestView(UWidgetBlueprint* WidgetBlueprint) const;
+	UE_API UMVVMBlueprintView* RequestView(UWidgetBlueprint* WidgetBlueprint) const;
 
 	UFUNCTION(BlueprintCallable, Category = "Viewmodel")
-	UMVVMBlueprintView* GetView(const UWidgetBlueprint* WidgetBlueprint) const;
+	UE_API UMVVMBlueprintView* GetView(const UWidgetBlueprint* WidgetBlueprint) const;
 
 	UFUNCTION(BlueprintCallable, Category = "Viewmodel")
-	FGuid AddViewModel(UWidgetBlueprint* WidgetBlueprint, const UClass* ViewModel);
+	UE_API FGuid AddViewModel(UWidgetBlueprint* WidgetBlueprint, const UClass* ViewModel);
 
 	UFUNCTION(BlueprintCallable, Category = "Viewmodel")
-	FGuid AddInstancedViewModel(UWidgetBlueprint* WidgetBlueprint);
+	UE_API FGuid AddInstancedViewModel(UWidgetBlueprint* WidgetBlueprint);
 	
 	UFUNCTION(BlueprintCallable, Category = "Viewmodel")
-	void RemoveViewModel(UWidgetBlueprint* WidgetBlueprint, FName ViewModel);
+	UE_API void RemoveViewModel(UWidgetBlueprint* WidgetBlueprint, FName ViewModel);
 
 	UFUNCTION(BlueprintCallable, Category = "Viewmodel")
-	bool VerifyViewModelRename(UWidgetBlueprint* WidgetBlueprint, FName ViewModel, FName NewViewModel, FText& OutError);
+	UE_API bool VerifyViewModelRename(UWidgetBlueprint* WidgetBlueprint, FName ViewModel, FName NewViewModel, FText& OutError);
 
 	UFUNCTION(BlueprintCallable, Category = "Viewmodel")
-	bool RenameViewModel(UWidgetBlueprint* WidgetBlueprint, FName ViewModel, FName NewViewModel, FText& OutError);
+	UE_API bool RenameViewModel(UWidgetBlueprint* WidgetBlueprint, FName ViewModel, FName NewViewModel, FText& OutError);
 
 	UFUNCTION(BlueprintCallable, Category = "Viewmodel")
-	bool ReparentViewModel(UWidgetBlueprint* WidgetBlueprint, FName ViewModel, const UClass* NewViewModel, FText& OutError);
+	UE_API bool ReparentViewModel(UWidgetBlueprint* WidgetBlueprint, FName ViewModel, const UClass* NewViewModel, FText& OutError);
 
 	UFUNCTION(BlueprintCallable, Category = "Viewmodel")
-	FMVVMBlueprintViewBinding& AddBinding(UWidgetBlueprint* WidgetBlueprint);
+	UE_API FMVVMBlueprintViewBinding& AddBinding(UWidgetBlueprint* WidgetBlueprint);
 
 	UFUNCTION(BlueprintCallable, Category = "Viewmodel")
-	void RemoveBinding(UWidgetBlueprint* WidgetBlueprint, const FMVVMBlueprintViewBinding& Binding);
+	UE_API void RemoveBinding(UWidgetBlueprint* WidgetBlueprint, const FMVVMBlueprintViewBinding& Binding);
 
 	UFUNCTION(BlueprintCallable, Category = "Viewmodel")
-	UMVVMBlueprintViewEvent* AddEvent(UWidgetBlueprint* WidgetBlueprint);
+	UE_API UMVVMBlueprintViewEvent* AddEvent(UWidgetBlueprint* WidgetBlueprint);
 
 	UFUNCTION(BlueprintCallable, Category = "Viewmodel")
-	void RemoveEvent(UWidgetBlueprint* WidgetBlueprint, UMVVMBlueprintViewEvent* Event);
+	UE_API void RemoveEvent(UWidgetBlueprint* WidgetBlueprint, UMVVMBlueprintViewEvent* Event);
 
 	UFUNCTION(BlueprintCallable, Category = "Viewmodel")
-	UMVVMBlueprintViewCondition* AddCondition(UWidgetBlueprint* WidgetBlueprint);
+	UE_API UMVVMBlueprintViewCondition* AddCondition(UWidgetBlueprint* WidgetBlueprint);
 
 	UFUNCTION(BlueprintCallable, Category = "Viewmodel")
-	void RemoveCondition(UWidgetBlueprint* WidgetBlueprint, UMVVMBlueprintViewCondition* Condition);
+	UE_API void RemoveCondition(UWidgetBlueprint* WidgetBlueprint, UMVVMBlueprintViewCondition* Condition);
 
 	UFUNCTION(BlueprintCallable, Category = "Viewmodel")
-	TArray<FMVVMAvailableBinding> GetChildViewModels(TSubclassOf<UObject> Class, TSubclassOf<UObject> Accessor);
+	UE_API TArray<FMVVMAvailableBinding> GetChildViewModels(TSubclassOf<UObject> Class, TSubclassOf<UObject> Accessor);
 
 	UE_DEPRECATED(5.4, "SetSourceToDestinationConversionFunction with a UFunction is deprecated.")
-	void SetSourceToDestinationConversionFunction(UWidgetBlueprint* WidgetBlueprint, FMVVMBlueprintViewBinding& Binding, const UFunction* ConversionFunction);
-	void SetSourceToDestinationConversionFunction(UWidgetBlueprint* WidgetBlueprint, FMVVMBlueprintViewBinding& Binding, FMVVMBlueprintFunctionReference ConversionFunction);
+	UE_API void SetSourceToDestinationConversionFunction(UWidgetBlueprint* WidgetBlueprint, FMVVMBlueprintViewBinding& Binding, const UFunction* ConversionFunction);
+	UE_API void SetSourceToDestinationConversionFunction(UWidgetBlueprint* WidgetBlueprint, FMVVMBlueprintViewBinding& Binding, FMVVMBlueprintFunctionReference ConversionFunction);
 	UE_DEPRECATED(5.4, "SetDestinationToSourceConversionFunction  with a UFunction is deprecated.")
-	void SetDestinationToSourceConversionFunction(UWidgetBlueprint* WidgetBlueprint, FMVVMBlueprintViewBinding& Binding, const UFunction* ConversionFunction);
-	void SetDestinationToSourceConversionFunction(UWidgetBlueprint* WidgetBlueprint, FMVVMBlueprintViewBinding& Binding, FMVVMBlueprintFunctionReference ConversionFunction);
-	void SetDestinationPathForBinding(UWidgetBlueprint* WidgetBlueprint, FMVVMBlueprintViewBinding& Binding, FMVVMBlueprintPropertyPath Field, bool bAllowEventConversion);
-	void SetSourcePathForBinding(UWidgetBlueprint* WidgetBlueprint, FMVVMBlueprintViewBinding& Binding, FMVVMBlueprintPropertyPath Field);
-	void OverrideExecutionModeForBinding(UWidgetBlueprint* WidgetBlueprint, FMVVMBlueprintViewBinding& Binding, EMVVMExecutionMode Mode);
-	void ResetExecutionModeForBinding(UWidgetBlueprint* WidgetBlueprint, FMVVMBlueprintViewBinding& Binding);
-	void SetBindingTypeForBinding(UWidgetBlueprint* WidgetBlueprint, FMVVMBlueprintViewBinding& Binding, EMVVMBindingMode Type);
-	void SetEnabledForBinding(UWidgetBlueprint* WidgetBlueprint, FMVVMBlueprintViewBinding& Binding, bool bEnabled);
-	void SetCompileForBinding(UWidgetBlueprint* WidgetBlueprint, FMVVMBlueprintViewBinding& Binding, bool bCompile);
-	void GenerateBindToDestinationPathsForBinding(UWidgetBlueprint* WidgetBlueprint, FMVVMBlueprintViewBinding& Binding);
+	UE_API void SetDestinationToSourceConversionFunction(UWidgetBlueprint* WidgetBlueprint, FMVVMBlueprintViewBinding& Binding, const UFunction* ConversionFunction);
+	UE_API void SetDestinationToSourceConversionFunction(UWidgetBlueprint* WidgetBlueprint, FMVVMBlueprintViewBinding& Binding, FMVVMBlueprintFunctionReference ConversionFunction);
+	UE_API void SetDestinationPathForBinding(UWidgetBlueprint* WidgetBlueprint, FMVVMBlueprintViewBinding& Binding, FMVVMBlueprintPropertyPath Field, bool bAllowEventConversion);
+	UE_API void SetSourcePathForBinding(UWidgetBlueprint* WidgetBlueprint, FMVVMBlueprintViewBinding& Binding, FMVVMBlueprintPropertyPath Field);
+	UE_API void OverrideExecutionModeForBinding(UWidgetBlueprint* WidgetBlueprint, FMVVMBlueprintViewBinding& Binding, EMVVMExecutionMode Mode);
+	UE_API void ResetExecutionModeForBinding(UWidgetBlueprint* WidgetBlueprint, FMVVMBlueprintViewBinding& Binding);
+	UE_API void SetBindingTypeForBinding(UWidgetBlueprint* WidgetBlueprint, FMVVMBlueprintViewBinding& Binding, EMVVMBindingMode Type);
+	UE_API void SetEnabledForBinding(UWidgetBlueprint* WidgetBlueprint, FMVVMBlueprintViewBinding& Binding, bool bEnabled);
+	UE_API void SetCompileForBinding(UWidgetBlueprint* WidgetBlueprint, FMVVMBlueprintViewBinding& Binding, bool bCompile);
+	UE_API void GenerateBindToDestinationPathsForBinding(UWidgetBlueprint* WidgetBlueprint, FMVVMBlueprintViewBinding& Binding);
 
-	void SetEventPath(UMVVMBlueprintViewEvent* Event, FMVVMBlueprintPropertyPath PropertyPath, bool bRequestBindingConversion);
-	void SetEventDestinationPath(UMVVMBlueprintViewEvent* Event, FMVVMBlueprintPropertyPath PropertyPath);
-	void SetEventArgumentPath(UMVVMBlueprintViewEvent* Event, const FMVVMBlueprintPinId& PinId, const FMVVMBlueprintPropertyPath& PropertyPath) const;
-	void SetEnabledForEvent(UMVVMBlueprintViewEvent* Event, bool bEnabled);
-	void SetCompileForEvent(UMVVMBlueprintViewEvent* Event, bool bCompile);
+	UE_API void SetEventPath(UMVVMBlueprintViewEvent* Event, FMVVMBlueprintPropertyPath PropertyPath, bool bRequestBindingConversion);
+	UE_API void SetEventDestinationPath(UMVVMBlueprintViewEvent* Event, FMVVMBlueprintPropertyPath PropertyPath);
+	UE_API void SetEventArgumentPath(UMVVMBlueprintViewEvent* Event, const FMVVMBlueprintPinId& PinId, const FMVVMBlueprintPropertyPath& PropertyPath) const;
+	UE_API void SetEnabledForEvent(UMVVMBlueprintViewEvent* Event, bool bEnabled);
+	UE_API void SetCompileForEvent(UMVVMBlueprintViewEvent* Event, bool bCompile);
 
-	void SetConditionPath(UMVVMBlueprintViewCondition* Condition, FMVVMBlueprintPropertyPath PropertyPath, bool bRequestBindingConversion);
-	void SetConditionDestinationPath(UMVVMBlueprintViewCondition* Condition, FMVVMBlueprintPropertyPath PropertyPath);
-	void SetConditionArgumentPath(UMVVMBlueprintViewCondition* Condition, const FMVVMBlueprintPinId& PinId, const FMVVMBlueprintPropertyPath& PropertyPath) const;
-	void SetEnabledForCondition(UMVVMBlueprintViewCondition* Condition, bool bEnabled);
-	void SetCompileForCondition(UMVVMBlueprintViewCondition* Condition, bool bCompile);
-	void SetConditionOperation(UMVVMBlueprintViewCondition* Condition, EMVVMConditionOperation Operation);
-	void SetConditionOperationValue(UMVVMBlueprintViewCondition* Condition, float Value);
-	void SetConditionOperationMaxValue(UMVVMBlueprintViewCondition* Condition, float MaxValue);
+	UE_API void SetConditionPath(UMVVMBlueprintViewCondition* Condition, FMVVMBlueprintPropertyPath PropertyPath, bool bRequestBindingConversion);
+	UE_API void SetConditionDestinationPath(UMVVMBlueprintViewCondition* Condition, FMVVMBlueprintPropertyPath PropertyPath);
+	UE_API void SetConditionArgumentPath(UMVVMBlueprintViewCondition* Condition, const FMVVMBlueprintPinId& PinId, const FMVVMBlueprintPropertyPath& PropertyPath) const;
+	UE_API void SetEnabledForCondition(UMVVMBlueprintViewCondition* Condition, bool bEnabled);
+	UE_API void SetCompileForCondition(UMVVMBlueprintViewCondition* Condition, bool bCompile);
+	UE_API void SetConditionOperation(UMVVMBlueprintViewCondition* Condition, EMVVMConditionOperation Operation);
+	UE_API void SetConditionOperationValue(UMVVMBlueprintViewCondition* Condition, float Value);
+	UE_API void SetConditionOperationMaxValue(UMVVMBlueprintViewCondition* Condition, float MaxValue);
 
 	UFUNCTION(BlueprintCallable, Category = "Viewmodel")
-	bool IsValidConversionFunction(const UWidgetBlueprint* WidgetBlueprint, const UFunction* Function, const FMVVMBlueprintPropertyPath& Source, const FMVVMBlueprintPropertyPath& Destination) const;
-	bool IsValidConversionFunction(const UWidgetBlueprint* WidgetBlueprint, const UFunction* Function, const FProperty* ExpectedArgumentType, const FProperty* ExptectedReturnType) const;
-	bool IsValidConversionFunction(const UWidgetBlueprint* WidgetBlueprint, UE::MVVM::FConversionFunctionValue Function, const FProperty* ExpectedArgumentType, const FProperty* ExptectedReturnType) const;
+	UE_API bool IsValidConversionFunction(const UWidgetBlueprint* WidgetBlueprint, const UFunction* Function, const FMVVMBlueprintPropertyPath& Source, const FMVVMBlueprintPropertyPath& Destination) const;
+	UE_API bool IsValidConversionFunction(const UWidgetBlueprint* WidgetBlueprint, const UFunction* Function, const FProperty* ExpectedArgumentType, const FProperty* ExptectedReturnType) const;
+	UE_API bool IsValidConversionFunction(const UWidgetBlueprint* WidgetBlueprint, UE::MVVM::FConversionFunctionValue Function, const FProperty* ExpectedArgumentType, const FProperty* ExptectedReturnType) const;
 
-	bool IsValidConversionNode(const UWidgetBlueprint* WidgetBlueprint, const TSubclassOf<UK2Node> Function, const FMVVMBlueprintPropertyPath& Source, const FMVVMBlueprintPropertyPath& Destination) const;
-	bool IsValidConversionNode(const UWidgetBlueprint* WidgetBlueprint, const TSubclassOf<UK2Node> Function, const FProperty* ExpectedArgumentType, const FProperty* ExptectedReturnType) const;
+	UE_API bool IsValidConversionNode(const UWidgetBlueprint* WidgetBlueprint, const TSubclassOf<UK2Node> Function, const FMVVMBlueprintPropertyPath& Source, const FMVVMBlueprintPropertyPath& Destination) const;
+	UE_API bool IsValidConversionNode(const UWidgetBlueprint* WidgetBlueprint, const TSubclassOf<UK2Node> Function, const FProperty* ExpectedArgumentType, const FProperty* ExptectedReturnType) const;
 
 	UFUNCTION(BlueprintCallable, Category = "Viewmodel")
-	bool IsSimpleConversionFunction(const UFunction* Function) const;
+	UE_API bool IsSimpleConversionFunction(const UFunction* Function) const;
 
 	UFUNCTION(BlueprintCallable, Category = "Viewmodel")
-	UEdGraph* GetConversionFunctionGraph(const UWidgetBlueprint* WidgetBlueprint, const FMVVMBlueprintViewBinding& Binding, bool bSourceToDestination) const;
+	UE_API UEdGraph* GetConversionFunctionGraph(const UWidgetBlueprint* WidgetBlueprint, const FMVVMBlueprintViewBinding& Binding, bool bSourceToDestination) const;
 
 	UE_DEPRECATED(5.4, "GetConversionFunction was moved to MVVMBlueprintViewConversionFunction.")
 	UFUNCTION(BlueprintCallable, Category = "Viewmodel")
-	UFunction* GetConversionFunction(const UWidgetBlueprint* WidgetBlueprint, const FMVVMBlueprintViewBinding& Binding, bool bSourceToDestination) const;
+	UE_API UFunction* GetConversionFunction(const UWidgetBlueprint* WidgetBlueprint, const FMVVMBlueprintViewBinding& Binding, bool bSourceToDestination) const;
 
 	UE_DEPRECATED(5.3, "GetConversionFunctionNode was moved to MVVMBlueprintViewConversionFunction.")
 	UFUNCTION(BlueprintCallable, Category = "Viewmodel")
-	UK2Node_CallFunction* GetConversionFunctionNode(const UWidgetBlueprint* WidgetBlueprint, const FMVVMBlueprintViewBinding& Binding, bool bSourceToDestination) const;
+	UE_API UK2Node_CallFunction* GetConversionFunctionNode(const UWidgetBlueprint* WidgetBlueprint, const FMVVMBlueprintViewBinding& Binding, bool bSourceToDestination) const;
 
 	UE_DEPRECATED(5.5, "GetAvailableConversionFunctions return value changes.")
 	UFUNCTION(BlueprintCallable, Category = "Viewmodel")
-	TArray<UFunction*> GetAvailableConversionFunctions(const UWidgetBlueprint* WidgetBlueprint, const FMVVMBlueprintPropertyPath& Source, const FMVVMBlueprintPropertyPath& Destination) const;
-
-	TArray<UE::MVVM::FConversionFunctionValue> GetConversionFunctions(const UWidgetBlueprint* WidgetBlueprint, const FProperty* ExpectedArgumentType, const FProperty* ExptectedReturnType) const;
-
-	FMVVMBlueprintPropertyPath GetPathForConversionFunctionArgument(const UWidgetBlueprint* WidgetBlueprint, const FMVVMBlueprintViewBinding& Binding, const FMVVMBlueprintPinId& PinId, bool bSourceToDestination) const;
-	void SetPathForConversionFunctionArgument(UWidgetBlueprint* WidgetBlueprint, FMVVMBlueprintViewBinding& Binding, const FMVVMBlueprintPinId& PinId, const FMVVMBlueprintPropertyPath& Path, bool bSourceToDestination) const;
-	UEdGraphPin* GetConversionFunctionArgumentPin(const UWidgetBlueprint* WidgetBlueprint, const FMVVMBlueprintViewBinding& Binding, const FMVVMBlueprintPinId& PinId, bool bSourceToDestination) const;
-
-	void SplitPin(UWidgetBlueprint* WidgetBlueprint, FMVVMBlueprintViewBinding& Binding, const FMVVMBlueprintPinId& PinId, bool bSourceToDestination) const;
-	bool CanSplitPin(const UWidgetBlueprint* WidgetBlueprint, const FMVVMBlueprintViewBinding& Binding, const FMVVMBlueprintPinId& PinId, bool bSourceToDestination) const;
-	void SplitPin(UWidgetBlueprint* WidgetBlueprint, UMVVMBlueprintViewEvent* Event, const FMVVMBlueprintPinId& PinId) const;
-	bool CanSplitPin(const UWidgetBlueprint* WidgetBlueprint, UMVVMBlueprintViewEvent* Event, const FMVVMBlueprintPinId& PinId) const;
-	void RecombinePin(UWidgetBlueprint* WidgetBlueprint, FMVVMBlueprintViewBinding& Binding, const FMVVMBlueprintPinId& PinId, bool bSourceToDestination) const;
-	bool CanRecombinePin(const UWidgetBlueprint* WidgetBlueprint, const FMVVMBlueprintViewBinding& Binding, const FMVVMBlueprintPinId& PinId, bool bSourceToDestination) const;
-	void RecombinePin(UWidgetBlueprint* WidgetBlueprint, UMVVMBlueprintViewEvent* Event, const FMVVMBlueprintPinId& PinId) const;
-	bool CanRecombinePin(const UWidgetBlueprint* WidgetBlueprint, UMVVMBlueprintViewEvent* Event, const FMVVMBlueprintPinId& PinId) const;
-	void ResetPinToDefaultValue(UWidgetBlueprint* WidgetBlueprint, FMVVMBlueprintViewBinding& Binding,const FMVVMBlueprintPinId& PinId, bool bSourceToDestination) const;
-	bool CanResetPinToDefaultValue(const UWidgetBlueprint* WidgetBlueprint, const FMVVMBlueprintViewBinding& Binding, const FMVVMBlueprintPinId& PinId, bool bSourceToDestination) const;
-	void ResetPinToDefaultValue(UWidgetBlueprint* WidgetBlueprint, UMVVMBlueprintViewEvent* Event, const FMVVMBlueprintPinId& PinId) const;
-	bool CanResetPinToDefaultValue(const UWidgetBlueprint* WidgetBlueprint, UMVVMBlueprintViewEvent* Event, const FMVVMBlueprintPinId& PinId) const;
-	void ResetOrphanedPin(UWidgetBlueprint* WidgetBlueprint, FMVVMBlueprintViewBinding& Binding, const FMVVMBlueprintPinId& PinId, bool bSourceToDestination) const;
-	bool CanResetOrphanedPin(const UWidgetBlueprint* WidgetBlueprint, const FMVVMBlueprintViewBinding& Binding, const FMVVMBlueprintPinId& PinId, bool bSourceToDestination) const;
-	void ResetOrphanedPin(UWidgetBlueprint* WidgetBlueprint, UMVVMBlueprintViewEvent* Event,const FMVVMBlueprintPinId& PinId) const;
-	bool CanResetOrphanedPin(const UWidgetBlueprint* WidgetBlueprint, UMVVMBlueprintViewEvent* Event, const FMVVMBlueprintPinId& PinId) const;
-
-	TArray<UE::MVVM::FBindingSource> GetBindableWidgets(const UWidgetBlueprint* WidgetBlueprint) const;
-	TArray<UE::MVVM::FBindingSource> GetAllViewModels(const UWidgetBlueprint* WidgetBlueprint) const;
-
-	FGuid GetFirstBindingThatUsesViewModel(const UWidgetBlueprint* WidgetBlueprint, FGuid ViewModelId) const;
-	static FString GetDefaultViewModelName(const UClass* ViewModelClass);
+	UE_API TArray<UFunction*> GetAvailableConversionFunctions(const UWidgetBlueprint* WidgetBlueprint, const FMVVMBlueprintPropertyPath& Source, const FMVVMBlueprintPropertyPath& Destination) const;
+
+	UE_API TArray<UE::MVVM::FConversionFunctionValue> GetConversionFunctions(const UWidgetBlueprint* WidgetBlueprint, const FProperty* ExpectedArgumentType, const FProperty* ExptectedReturnType) const;
+
+	UE_API FMVVMBlueprintPropertyPath GetPathForConversionFunctionArgument(const UWidgetBlueprint* WidgetBlueprint, const FMVVMBlueprintViewBinding& Binding, const FMVVMBlueprintPinId& PinId, bool bSourceToDestination) const;
+	UE_API void SetPathForConversionFunctionArgument(UWidgetBlueprint* WidgetBlueprint, FMVVMBlueprintViewBinding& Binding, const FMVVMBlueprintPinId& PinId, const FMVVMBlueprintPropertyPath& Path, bool bSourceToDestination) const;
+	UE_API UEdGraphPin* GetConversionFunctionArgumentPin(const UWidgetBlueprint* WidgetBlueprint, const FMVVMBlueprintViewBinding& Binding, const FMVVMBlueprintPinId& PinId, bool bSourceToDestination) const;
+
+	UE_API void SplitPin(UWidgetBlueprint* WidgetBlueprint, FMVVMBlueprintViewBinding& Binding, const FMVVMBlueprintPinId& PinId, bool bSourceToDestination) const;
+	UE_API bool CanSplitPin(const UWidgetBlueprint* WidgetBlueprint, const FMVVMBlueprintViewBinding& Binding, const FMVVMBlueprintPinId& PinId, bool bSourceToDestination) const;
+	UE_API void SplitPin(UWidgetBlueprint* WidgetBlueprint, UMVVMBlueprintViewEvent* Event, const FMVVMBlueprintPinId& PinId) const;
+	UE_API bool CanSplitPin(const UWidgetBlueprint* WidgetBlueprint, UMVVMBlueprintViewEvent* Event, const FMVVMBlueprintPinId& PinId) const;
+	UE_API void RecombinePin(UWidgetBlueprint* WidgetBlueprint, FMVVMBlueprintViewBinding& Binding, const FMVVMBlueprintPinId& PinId, bool bSourceToDestination) const;
+	UE_API bool CanRecombinePin(const UWidgetBlueprint* WidgetBlueprint, const FMVVMBlueprintViewBinding& Binding, const FMVVMBlueprintPinId& PinId, bool bSourceToDestination) const;
+	UE_API void RecombinePin(UWidgetBlueprint* WidgetBlueprint, UMVVMBlueprintViewEvent* Event, const FMVVMBlueprintPinId& PinId) const;
+	UE_API bool CanRecombinePin(const UWidgetBlueprint* WidgetBlueprint, UMVVMBlueprintViewEvent* Event, const FMVVMBlueprintPinId& PinId) const;
+	UE_API void ResetPinToDefaultValue(UWidgetBlueprint* WidgetBlueprint, FMVVMBlueprintViewBinding& Binding,const FMVVMBlueprintPinId& PinId, bool bSourceToDestination) const;
+	UE_API bool CanResetPinToDefaultValue(const UWidgetBlueprint* WidgetBlueprint, const FMVVMBlueprintViewBinding& Binding, const FMVVMBlueprintPinId& PinId, bool bSourceToDestination) const;
+	UE_API void ResetPinToDefaultValue(UWidgetBlueprint* WidgetBlueprint, UMVVMBlueprintViewEvent* Event, const FMVVMBlueprintPinId& PinId) const;
+	UE_API bool CanResetPinToDefaultValue(const UWidgetBlueprint* WidgetBlueprint, UMVVMBlueprintViewEvent* Event, const FMVVMBlueprintPinId& PinId) const;
+	UE_API void ResetOrphanedPin(UWidgetBlueprint* WidgetBlueprint, FMVVMBlueprintViewBinding& Binding, const FMVVMBlueprintPinId& PinId, bool bSourceToDestination) const;
+	UE_API bool CanResetOrphanedPin(const UWidgetBlueprint* WidgetBlueprint, const FMVVMBlueprintViewBinding& Binding, const FMVVMBlueprintPinId& PinId, bool bSourceToDestination) const;
+	UE_API void ResetOrphanedPin(UWidgetBlueprint* WidgetBlueprint, UMVVMBlueprintViewEvent* Event,const FMVVMBlueprintPinId& PinId) const;
+	UE_API bool CanResetOrphanedPin(const UWidgetBlueprint* WidgetBlueprint, UMVVMBlueprintViewEvent* Event, const FMVVMBlueprintPinId& PinId) const;
+
+	UE_API TArray<UE::MVVM::FBindingSource> GetBindableWidgets(const UWidgetBlueprint* WidgetBlueprint) const;
+	UE_API TArray<UE::MVVM::FBindingSource> GetAllViewModels(const UWidgetBlueprint* WidgetBlueprint) const;
+
+	UE_API FGuid GetFirstBindingThatUsesViewModel(const UWidgetBlueprint* WidgetBlueprint, FGuid ViewModelId) const;
+	static UE_API FString GetDefaultViewModelName(const UClass* ViewModelClass);
 
 private:
 	mutable TUniquePtr<UE::MVVM::ConversionFunctionLibrary::FCollection> ConversionFunctionCollection;
 };
+
+#undef UE_API
diff --git a/Engine/Plugins/Runtime/ModelViewViewModel/Source/ModelViewViewModelEditor/Public/Types/MVVMConversionFunctionValue.h b/Engine/Plugins/Runtime/ModelViewViewModel/Source/ModelViewViewModelEditor/Public/Types/MVVMConversionFunctionValue.h
--- a/Engine/Plugins/Runtime/ModelViewViewModel/Source/ModelViewViewModelEditor/Public/Types/MVVMConversionFunctionValue.h
+++ b/Engine/Plugins/Runtime/ModelViewViewModel/Source/ModelViewViewModelEditor/Public/Types/MVVMConversionFunctionValue.h
@@ -5,13 +5,15 @@
 #include "K2Node.h"
 #include "Templates/SubclassOf.h"
 
+#define UE_API MODELVIEWVIEWMODELEDITOR_API
+
 class UFunction;
 class UK2Node;
 
 namespace UE::MVVM
 {
 /** */
-struct MODELVIEWVIEWMODELEDITOR_API FConversionFunctionValue
+struct FConversionFunctionValue
 {
 public:
 
@@ -44,13 +46,13 @@ public:
 		return ConversionNode;
 	}
 
-	FString GetName() const;
-	FName GetFName() const;
-	FString GetFullGroupName(bool bStartWithOuter) const;
-	FText GetDisplayName() const;
-	FText GetTooltip() const;
-	FText GetCategory() const;
-	TArray<FString> GetSearchKeywords() const;
+	UE_API FString GetName() const;
+	UE_API FName GetFName() const;
+	UE_API FString GetFullGroupName(bool bStartWithOuter) const;
+	UE_API FText GetDisplayName() const;
+	UE_API FText GetTooltip() const;
+	UE_API FText GetCategory() const;
+	UE_API TArray<FString> GetSearchKeywords() const;
 
 public:
 	bool IsValid() const
@@ -85,3 +87,5 @@ private:
 };
 
 } // namespace UE::MVVM
+
+#undef UE_API
diff --git a/Engine/Plugins/Runtime/ModelViewViewModel/Source/ModelViewViewModelEditor/Public/Types/MVVMLinkedPinValue.h b/Engine/Plugins/Runtime/ModelViewViewModel/Source/ModelViewViewModelEditor/Public/Types/MVVMLinkedPinValue.h
--- a/Engine/Plugins/Runtime/ModelViewViewModel/Source/ModelViewViewModelEditor/Public/Types/MVVMLinkedPinValue.h
+++ b/Engine/Plugins/Runtime/ModelViewViewModel/Source/ModelViewViewModelEditor/Public/Types/MVVMLinkedPinValue.h
@@ -9,35 +9,37 @@
 
 #include "MVVMLinkedPinValue.generated.h"
 
+#define UE_API MODELVIEWVIEWMODELEDITOR_API
+
 class UFunction;
 class UK2Node;
 
 /** */
 USTRUCT(BlueprintType)
-struct MODELVIEWVIEWMODELEDITOR_API FMVVMLinkedPinValue
+struct FMVVMLinkedPinValue
 {
 public:
 	GENERATED_BODY()
 
 	FMVVMLinkedPinValue() = default;
-	explicit FMVVMLinkedPinValue(FMVVMBlueprintPropertyPath InPath);
-	explicit FMVVMLinkedPinValue(const UBlueprint* InBlueprint, FMVVMBlueprintFunctionReference InConversion);
-	explicit FMVVMLinkedPinValue(UE::MVVM::FConversionFunctionValue Function);
-	explicit FMVVMLinkedPinValue(const UFunction* Function);
-	explicit FMVVMLinkedPinValue(TSubclassOf<UK2Node> Node);
+	UE_API explicit FMVVMLinkedPinValue(FMVVMBlueprintPropertyPath InPath);
+	UE_API explicit FMVVMLinkedPinValue(const UBlueprint* InBlueprint, FMVVMBlueprintFunctionReference InConversion);
+	UE_API explicit FMVVMLinkedPinValue(UE::MVVM::FConversionFunctionValue Function);
+	UE_API explicit FMVVMLinkedPinValue(const UFunction* Function);
+	UE_API explicit FMVVMLinkedPinValue(TSubclassOf<UK2Node> Node);
 
 public:
-	bool IsPropertyPath() const;
-	const FMVVMBlueprintPropertyPath& GetPropertyPath() const;
+	UE_API bool IsPropertyPath() const;
+	UE_API const FMVVMBlueprintPropertyPath& GetPropertyPath() const;
 
-	bool IsConversionFunction() const;
-	const UFunction* GetConversionFunction() const;
+	UE_API bool IsConversionFunction() const;
+	UE_API const UFunction* GetConversionFunction() const;
 
-	bool IsConversionNode() const;
-	TSubclassOf<UK2Node> GetConversionNode() const;
+	UE_API bool IsConversionNode() const;
+	UE_API TSubclassOf<UK2Node> GetConversionNode() const;
 
 public:
-	bool IsValid() const;
+	UE_API bool IsValid() const;
 	bool operator== (const FMVVMLinkedPinValue& Other) const
 	{
 		return ConversionNode == Other.ConversionNode && ConversionFunction == Other.ConversionFunction && PropertyPath == Other.PropertyPath;
@@ -53,3 +55,5 @@ protected:
 	UPROPERTY(BlueprintReadWrite, Category = "MVVM")
 	TSubclassOf<UK2Node> ConversionNode;
 };
+
+#undef UE_API
diff --git a/Engine/Plugins/Runtime/OpenCV/Source/OpenCVHelper/Public/OpenCVBlueprintFunctionLibrary.h b/Engine/Plugins/Runtime/OpenCV/Source/OpenCVHelper/Public/OpenCVBlueprintFunctionLibrary.h
--- a/Engine/Plugins/Runtime/OpenCV/Source/OpenCVHelper/Public/OpenCVBlueprintFunctionLibrary.h
+++ b/Engine/Plugins/Runtime/OpenCV/Source/OpenCVHelper/Public/OpenCVBlueprintFunctionLibrary.h
@@ -8,6 +8,8 @@
 
 #include "OpenCVBlueprintFunctionLibrary.generated.h"
 
+#define UE_API OPENCVHELPER_API
+
 DECLARE_LOG_CATEGORY_EXTERN(LogOpenCVBlueprintFunctionLibrary, Log, All);
 
 class UTexture2D;
@@ -52,8 +54,8 @@ struct FOpenCVArucoDetectedMarker
 	FTransform Pose;
 };
 
-UCLASS(BlueprintType)
-class OPENCVHELPER_API UOpenCVBlueprintFunctionLibrary : public UBlueprintFunctionLibrary
+UCLASS(MinimalAPI, BlueprintType)
+class UOpenCVBlueprintFunctionLibrary : public UBlueprintFunctionLibrary
 {
 	GENERATED_BODY()
 
@@ -69,7 +71,7 @@ public:
 	 * @return Total number of corners detected in the input image
 	 */
 	UFUNCTION(BlueprintCallable, Category = "OpenCV", Meta = (DisplayName = "OpenCV Chessboard Detect Corners"))
-	static int32 OpenCVChessboardDetectCorners(const UTextureRenderTarget2D* InRenderTarget, const FIntPoint InPatternSize, const bool bDebugDrawCorners, UTexture2D*& OutDebugTexture, TArray<FVector2D>& OutDetectedCorners);
+	static UE_API int32 OpenCVChessboardDetectCorners(const UTextureRenderTarget2D* InRenderTarget, const FIntPoint InPatternSize, const bool bDebugDrawCorners, UTexture2D*& OutDebugTexture, TArray<FVector2D>& OutDetectedCorners);
 
 	/**
 	 * Detects all ArUco markers in the supplied image
@@ -86,5 +88,7 @@ public:
 	 * @return Total number of markers detected in the input image
 	 */
 	UFUNCTION(BlueprintCallable, Category = "OpenCV | ArUco", Meta = (DisplayName = "OpenCV ArUco Detect Markers"))
-	static int32 OpenCVArucoDetectMarkers(const UTextureRenderTarget2D* InRenderTarget, const EOpenCVArucoDictionary InDictionary, const EOpenCVArucoDictionarySize InDictionarySize, const bool bDebugDrawMarkers, const bool bEstimatePose, const float InMarkerLengthInMeters, const FOpenCVLensDistortionParametersBase& InLensDistortionParameters, UTexture2D*& OutDebugTexture, TArray<FOpenCVArucoDetectedMarker>& OutDetectedMarkers);
-};
\ No newline at end of file
+	static UE_API int32 OpenCVArucoDetectMarkers(const UTextureRenderTarget2D* InRenderTarget, const EOpenCVArucoDictionary InDictionary, const EOpenCVArucoDictionarySize InDictionarySize, const bool bDebugDrawMarkers, const bool bEstimatePose, const float InMarkerLengthInMeters, const FOpenCVLensDistortionParametersBase& InLensDistortionParameters, UTexture2D*& OutDebugTexture, TArray<FOpenCVArucoDetectedMarker>& OutDetectedMarkers);
+};
+
+#undef UE_API
diff --git a/Engine/Plugins/Runtime/OpenCV/Source/OpenCVHelper/Public/OpenCVHelper.h b/Engine/Plugins/Runtime/OpenCV/Source/OpenCVHelper/Public/OpenCVHelper.h
--- a/Engine/Plugins/Runtime/OpenCV/Source/OpenCVHelper/Public/OpenCVHelper.h
+++ b/Engine/Plugins/Runtime/OpenCV/Source/OpenCVHelper/Public/OpenCVHelper.h
@@ -67,13 +67,13 @@ enum class EArucoDictionary : uint8
 	DICT_ARUCO_ORIGINAL	 UMETA(DisplayName = "DICT_ARUCO_ORIGINAL")
 };
 
-struct OPENCVHELPER_API FArucoMarker
+struct FArucoMarker
 {
 	int32 MarkerID = 0;
 	FVector2f Corners[4];
 };
 
-class OPENCVHELPER_API FOpenCVHelper
+class FOpenCVHelper
 {
 public:
 	/** Enumeration to specify any cartesian axis in positive or negative directions */
@@ -100,19 +100,19 @@ public:
 	};
 
 	/** Converts in-place the coordinate system of the given FTransform by specifying the source axes in terms of the destination axes */
-	static void ConvertCoordinateSystem(FTransform& Transform, const EAxis DstXInSrcAxis, const EAxis DstYInSrcAxis, const EAxis DstZInSrcAxis);
+	static OPENCVHELPER_API void ConvertCoordinateSystem(FTransform& Transform, const EAxis DstXInSrcAxis, const EAxis DstYInSrcAxis, const EAxis DstZInSrcAxis);
 
 	/** Converts in-place an FTransform in Unreal coordinates to OpenCV coordinates */
-	static void ConvertUnrealToOpenCV(FTransform& Transform);
+	static OPENCVHELPER_API void ConvertUnrealToOpenCV(FTransform& Transform);
 
 	/** Converts in-place an FTransform in OpenCV coordinates to Unreal coordinates */
-	static void ConvertOpenCVToUnreal(FTransform& Transform);
+	static OPENCVHELPER_API void ConvertOpenCVToUnreal(FTransform& Transform);
 
 	/** Converts an FVector in Unreal coordinates to OpenCV coordinates */
-	static FVector ConvertUnrealToOpenCV(const FVector& Transform);
+	static OPENCVHELPER_API FVector ConvertUnrealToOpenCV(const FVector& Transform);
 
 	/** Converts an FTransform in OpenCV coordinates to Unreal coordinates */
-	static FVector ConvertOpenCVToUnreal(const FVector& Transform);
+	static OPENCVHELPER_API FVector ConvertOpenCVToUnreal(const FVector& Transform);
 
 public:
 #if WITH_OPENCV
@@ -125,61 +125,61 @@ public:
 	 * 
 	 * @return Texture created out of the given OpenCV Mat.
 	 */
-	static UTexture2D* TextureFromCvMat(cv::Mat& Mat, const FString* PackagePath = nullptr, const FName* TextureName = nullptr);
-	static UTexture2D* TextureFromCvMat(cv::Mat& Mat, UTexture2D* InTexture);
+	static OPENCVHELPER_API UTexture2D* TextureFromCvMat(cv::Mat& Mat, const FString* PackagePath = nullptr, const FName* TextureName = nullptr);
+	static OPENCVHELPER_API UTexture2D* TextureFromCvMat(cv::Mat& Mat, UTexture2D* InTexture);
 
 	/** 
 	 * Takes a rotation vector (in rodrigues form) and translation vector that represent a change of basis from object space to camera space (in OpenCV's coordinate system),
 	 * inverts them to generate a camera pose, and converts it to an FTransform in UE's coordinate system. 
 	 * The input rotation and translation vectors are expected to match the form of an "rvec" and "tvec" that are generated functions like cv::calibrateCamera() and cv::solvePnP()
 	 */
-	static void MakeCameraPoseFromObjectVectors(const cv::Mat& InRotation, const cv::Mat& InTranslation, FTransform& OutTransform);
+	static OPENCVHELPER_API void MakeCameraPoseFromObjectVectors(const cv::Mat& InRotation, const cv::Mat& InTranslation, FTransform& OutTransform);
 
 	/**
 	 * Takes a camera pose transform in UE's coordinate system, converts it to a rotation vector (in rodrigues form) and translation vector (in OpenCV's coordinate system),
 	 * and inverts them to generate vectors that represent a change of basis from object space to camera space.
 	 * The output rotation and translation vectors are expected to match the form of an "rvec" and "tvec" that are consumed by functions like cv::calibrateCamera() and cv::solvePnP()
 	 */
-	static void MakeObjectVectorsFromCameraPose(const FTransform& InTransform, cv::Mat& OutRotation, cv::Mat& OutTranslation);
+	static OPENCVHELPER_API void MakeObjectVectorsFromCameraPose(const FTransform& InTransform, cv::Mat& OutRotation, cv::Mat& OutTranslation);
 #endif	// WITH_OPENCV
 
 	/** Identify a set of aruco markers in the input image that belong to the input aruco dictionary, and output the marker IDs and the 2D coordinates of the 4 corners of each marker */
-	static bool IdentifyArucoMarkers(TArray<FColor>& Image, FIntPoint ImageSize, EArucoDictionary DictionaryName, TArray<FArucoMarker>& OutMarkers);
+	static OPENCVHELPER_API bool IdentifyArucoMarkers(TArray<FColor>& Image, FIntPoint ImageSize, EArucoDictionary DictionaryName, TArray<FArucoMarker>& OutMarkers);
 
 	/** Draw a debug view of the input aruco markers on top of the input texture */
-	static bool DrawArucoMarkers(const TArray<FArucoMarker>& Markers, UTexture2D* DebugTexture);
+	static OPENCVHELPER_API bool DrawArucoMarkers(const TArray<FArucoMarker>& Markers, UTexture2D* DebugTexture);
 
 	/** Identify a checkerboard pattern in the input image with the given checkerboard dimensions (columns x rows), and output the 2D coordinates of the intersections between each checkerboard square */
-	static bool IdentifyCheckerboard(TArray<FColor>& Image, FIntPoint ImageSize, FIntPoint CheckerboardDimensions, TArray<FVector2f>& OutCorners);
+	static OPENCVHELPER_API bool IdentifyCheckerboard(TArray<FColor>& Image, FIntPoint ImageSize, FIntPoint CheckerboardDimensions, TArray<FVector2f>& OutCorners);
 
 	/** Identify a checkerboard pattern in a region of interest of the input image with the given checkerboard dimensions (columns x rows), and output the 2D coordinates of the intersections between each checkerboard square */
-	static bool IdentifyCheckerboard(TArray<FColor>& Image, FIntPoint ImageSize, FIntRect RegionOfInterest, FIntPoint CheckerboardDimensions, TArray<FVector2f>& OutCorners);
+	static OPENCVHELPER_API bool IdentifyCheckerboard(TArray<FColor>& Image, FIntPoint ImageSize, FIntRect RegionOfInterest, FIntPoint CheckerboardDimensions, TArray<FVector2f>& OutCorners);
 
 	/** Draw a debug view of the input checkerboard corners on top of the input texture */
-	static bool DrawCheckerboardCorners(const TArray<FVector2f>& Corners, FIntPoint CheckerboardDimensions, UTexture2D* DebugTexture);
+	static OPENCVHELPER_API bool DrawCheckerboardCorners(const TArray<FVector2f>& Corners, FIntPoint CheckerboardDimensions, UTexture2D* DebugTexture);
 
 	/** Draw a debug view of the input checkerboard corners on top of the input texture */
-	static bool DrawCheckerboardCorners(const TArray<FVector2D>& Corners, FIntPoint CheckerboardDimensions, UTexture2D* DebugTexture);
+	static OPENCVHELPER_API bool DrawCheckerboardCorners(const TArray<FVector2D>& Corners, FIntPoint CheckerboardDimensions, UTexture2D* DebugTexture);
 
 	/** Compute the camera pose that minimizes the reprojection error of the input object points and image points */
-	static bool SolvePnP(const TArray<FVector>& ObjectPoints, const TArray<FVector2f>& ImagePoints, const FVector2D& FocalLength, const FVector2D& ImageCenter, const TArray<float>& DistortionParameters, FTransform& OutCameraPose);
+	static OPENCVHELPER_API bool SolvePnP(const TArray<FVector>& ObjectPoints, const TArray<FVector2f>& ImagePoints, const FVector2D& FocalLength, const FVector2D& ImageCenter, const TArray<float>& DistortionParameters, FTransform& OutCameraPose);
 
 	/** Project the input object points to the 2D image plane defined by the input camera intrinsics and camera pose */
-	static bool ProjectPoints(const TArray<FVector>& ObjectPoints, const FVector2D& FocalLength, const FVector2D& ImageCenter, const TArray<float>& DistortionParameters, const FTransform& CameraPose, TArray<FVector2f>& OutImagePoints);
+	static OPENCVHELPER_API bool ProjectPoints(const TArray<FVector>& ObjectPoints, const FVector2D& FocalLength, const FVector2D& ImageCenter, const TArray<float>& DistortionParameters, const FTransform& CameraPose, TArray<FVector2f>& OutImagePoints);
 
 	/** Project the input object points to the 2D image plane defined by the input camera intrinsics and camera pose */
-	static bool ProjectPoints(const TArray<FVector>& ObjectPoints, const FVector2D& FocalLength, const FVector2D& ImageCenter, const TArray<float>& DistortionParameters, const FTransform& CameraPose, TArray<FVector2D>& OutImagePoints);
+	static OPENCVHELPER_API bool ProjectPoints(const TArray<FVector>& ObjectPoints, const FVector2D& FocalLength, const FVector2D& ImageCenter, const TArray<float>& DistortionParameters, const FTransform& CameraPose, TArray<FVector2D>& OutImagePoints);
 
 	/** Find a 3D fit line that passes through the input points, as well as a point on that line */
-	static bool FitLine3D(const TArray<FVector>& InPoints, FVector& OutLine, FVector& OutPointOnLine);
+	static OPENCVHELPER_API bool FitLine3D(const TArray<FVector>& InPoints, FVector& OutLine, FVector& OutPointOnLine);
 
 #if WITH_OPENCV
 	UE_DEPRECATED(5.4, "The version of ComputeReprojectionError takes OpenCV types as input parameters is deprecated. Please use the version that takes all UE types.")
-	static double ComputeReprojectionError(const FTransform& CameraPose, const cv::Mat& CameraIntrinsicMatrix, const std::vector<cv::Point3f>& Points3d, const std::vector<cv::Point2f>& Points2d);
+	static OPENCVHELPER_API double ComputeReprojectionError(const FTransform& CameraPose, const cv::Mat& CameraIntrinsicMatrix, const std::vector<cv::Point3f>& Points3d, const std::vector<cv::Point2f>& Points2d);
 #endif
 
 	/** Project the 3D objects points to the 2D image plane represented by the input camera intrinsics and pose, and compute the reprojection error (euclidean distance) between the input image points and the reprojected points */
-	static double ComputeReprojectionError(const TArray<FVector>& ObjectPoints, const TArray<FVector2f>& ImagePoints, const FVector2D& FocalLength, const FVector2D& ImageCenter, const FTransform& CameraPose);
+	static OPENCVHELPER_API double ComputeReprojectionError(const TArray<FVector>& ObjectPoints, const TArray<FVector2f>& ImagePoints, const FVector2D& FocalLength, const FVector2D& ImageCenter, const FTransform& CameraPose);
 };
 
 /**
@@ -191,7 +191,7 @@ public:
  * where F and C are normalized.
  */
 USTRUCT(BlueprintType)
-struct OPENCVHELPER_API FOpenCVLensDistortionParametersBase
+struct FOpenCVLensDistortionParametersBase
 {
 	GENERATED_USTRUCT_BODY()
 
@@ -214,10 +214,10 @@ public:
 public:
 #if WITH_OPENCV
 	/** Convert internal coefficients to OpenCV matrix representation */
-	cv::Mat ConvertToOpenCVDistortionCoefficients() const;
+	OPENCVHELPER_API cv::Mat ConvertToOpenCVDistortionCoefficients() const;
 
 	/** Convert internal normalized camera matrix to OpenCV pixel scaled matrix representation. */
-	cv::Mat CreateOpenCVCameraMatrix(const FVector2D& InImageSize) const;
+	OPENCVHELPER_API cv::Mat CreateOpenCVCameraMatrix(const FVector2D& InImageSize) const;
 #endif //WITH_OPENCV
 
 public:
diff --git a/Engine/Plugins/Runtime/RigVM/Source/RigVMDeveloper/Public/EdGraph/RigVMEdGraph.h b/Engine/Plugins/Runtime/RigVM/Source/RigVMDeveloper/Public/EdGraph/RigVMEdGraph.h
--- a/Engine/Plugins/Runtime/RigVM/Source/RigVMDeveloper/Public/EdGraph/RigVMEdGraph.h
+++ b/Engine/Plugins/Runtime/RigVM/Source/RigVMDeveloper/Public/EdGraph/RigVMEdGraph.h
@@ -9,6 +9,8 @@
 #include "RigVMModel/RigVMClient.h"
 #include "RigVMEdGraph.generated.h"
 
+#define UE_API RIGVMDEVELOPER_API
+
 class URigVMBlueprint;
 class URigVMEdGraphSchema;
 class URigVMController;
@@ -127,31 +129,31 @@ private:
 
 DECLARE_MULTICAST_DELEGATE_OneParam(FRigVMEdGraphNodeClicked, URigVMEdGraphNode*);
 
-UCLASS()
-class RIGVMDEVELOPER_API URigVMEdGraph : public UEdGraph, public IRigVMEditorSideObject
+UCLASS(MinimalAPI)
+class URigVMEdGraph : public UEdGraph, public IRigVMEditorSideObject
 {
 	GENERATED_BODY()
 
 public:
-	URigVMEdGraph();
+	UE_API URigVMEdGraph();
 
 	/** IRigVMEditorSideObject interface */
-	virtual FRigVMClient* GetRigVMClient() const override;
-	virtual FString GetRigVMNodePath() const override;
-	virtual void HandleRigVMGraphRenamed(const FString& InOldNodePath, const FString& InNewNodePath) override;
+	UE_API virtual FRigVMClient* GetRigVMClient() const override;
+	UE_API virtual FString GetRigVMNodePath() const override;
+	UE_API virtual void HandleRigVMGraphRenamed(const FString& InOldNodePath, const FString& InNewNodePath) override;
 
 	/** Set up this graph */
-	const URigVMBlueprint* GetBlueprintDefaultObject() const;
-	void SetBlueprintClass(const UClass* InClass);
-	virtual void InitializeFromBlueprint(URigVMBlueprint* InBlueprint);
-	bool IsPreviewGraph() const;
+	UE_API const URigVMBlueprint* GetBlueprintDefaultObject() const;
+	UE_API void SetBlueprintClass(const UClass* InClass);
+	UE_API virtual void InitializeFromBlueprint(URigVMBlueprint* InBlueprint);
+	UE_API bool IsPreviewGraph() const;
 
 	/** Get the ed graph schema */
-	const URigVMEdGraphSchema* GetRigVMEdGraphSchema();
+	UE_API const URigVMEdGraphSchema* GetRigVMEdGraphSchema();
 
 #if WITH_EDITORONLY_DATA
 	/** Customize blueprint changes based on backwards compatibility */
-	virtual void Serialize(FArchive& Ar) override;
+	UE_API virtual void Serialize(FArchive& Ar) override;
 #endif
 
 #if WITH_EDITOR
@@ -160,25 +162,25 @@ public:
 	bool bIsTemporaryGraphForCopyPaste;
 	bool bIsSelecting;
 
-	UEdGraphNode* FindNodeForModelNodeName(const FName& InModelNodeName, const bool bCacheIfRequired = true);
+	UE_API UEdGraphNode* FindNodeForModelNodeName(const FName& InModelNodeName, const bool bCacheIfRequired = true);
 
-	URigVMBlueprint* GetBlueprint() const;
-	URigVMGraph* GetModel() const;
-	URigVMController* GetController() const;
+	UE_API URigVMBlueprint* GetBlueprint() const;
+	UE_API URigVMGraph* GetModel() const;
+	UE_API URigVMController* GetController() const;
 	bool IsRootGraph() const { return GetRootGraph() == this; }
-	const URigVMEdGraph* GetRootGraph() const;
+	UE_API const URigVMEdGraph* GetRootGraph() const;
 
-	void HandleModifiedEvent(ERigVMGraphNotifType InNotifType, URigVMGraph* InGraph, UObject* InSubject);
-	virtual bool HandleModifiedEvent_Internal(ERigVMGraphNotifType InNotifType, URigVMGraph* InGraph, UObject* InSubject);
+	UE_API void HandleModifiedEvent(ERigVMGraphNotifType InNotifType, URigVMGraph* InGraph, UObject* InSubject);
+	UE_API virtual bool HandleModifiedEvent_Internal(ERigVMGraphNotifType InNotifType, URigVMGraph* InGraph, UObject* InSubject);
 
-	int32 GetInstructionIndex(const URigVMEdGraphNode* InNode, bool bAsInput);
+	UE_API int32 GetInstructionIndex(const URigVMEdGraphNode* InNode, bool bAsInput);
 
-	void CacheEntryNameList();
-	const TArray<TSharedPtr<FRigVMStringWithTag>>* GetEntryNameList(URigVMPin* InPin = nullptr) const;
+	UE_API void CacheEntryNameList();
+	UE_API const TArray<TSharedPtr<FRigVMStringWithTag>>* GetEntryNameList(URigVMPin* InPin = nullptr) const;
 
 	virtual const TArray<TSharedPtr<FRigVMStringWithTag>>* GetNameListForWidget(const FString& InWidgetName) const { return nullptr; }
 
-	void AddLocalVariableSearchMetaDataInfo(const FName InVariableName, TArray<UBlueprintExtension::FSearchTagDataPair>& OutTaggedMetaData) const;
+	UE_API void AddLocalVariableSearchMetaDataInfo(const FName InVariableName, TArray<UBlueprintExtension::FSearchTagDataPair>& OutTaggedMetaData) const;
 
 	UPROPERTY()
 	FString ModelNodePath;
@@ -188,7 +190,7 @@ public:
 
 protected:
 	using Super::AddNode;
-	virtual void AddNode(UEdGraphNode* NodeToAdd, bool bUserAction = false, bool bSelectNewNode = true) override;
+	UE_API virtual void AddNode(UEdGraphNode* NodeToAdd, bool bUserAction = false, bool bSelectNewNode = true) override;
 
 private:
 
@@ -196,10 +198,10 @@ private:
 
 	TMap<URigVMNode*, TPair<int32, int32>> CachedInstructionIndices;
 
-	void RemoveNode(UEdGraphNode* InNode);
+	UE_API void RemoveNode(UEdGraphNode* InNode);
 
 protected:
-	void HandleVMCompiledEvent(UObject* InCompiledObject, URigVM* InVM, FRigVMExtendedExecuteContext& InContext);
+	UE_API void HandleVMCompiledEvent(UObject* InCompiledObject, URigVM* InVM, FRigVMExtendedExecuteContext& InContext);
 
 private:
 	TMap<FName, UEdGraphNode*> ModelNodePathToEdNode;
@@ -226,3 +228,5 @@ private:
 	friend class URigVMEdGraphArrayNodeSpawner;
 	friend class URigVMEdGraphInvokeEntryNodeSpawner;
 };
+
+#undef UE_API
