MVVM: Construct pins for MVVM if destination or condition path is valid. Fix issue of broken/missing pings in BindingList

#rnx

diff --git a/Engine/Plugins/Runtime/ModelViewViewModel/Source/ModelViewViewModelBlueprint/Private/MVVMBlueprintViewCondition.cpp b/Engine/Plugins/Runtime/ModelViewViewModel/Source/ModelViewViewModelBlueprint/Private/MVVMBlueprintViewCondition.cpp
--- a/Engine/Plugins/Runtime/ModelViewViewModel/Source/ModelViewViewModelBlueprint/Private/MVVMBlueprintViewCondition.cpp
+++ b/Engine/Plugins/Runtime/ModelViewViewModel/Source/ModelViewViewModelBlueprint/Private/MVVMBlueprintViewCondition.cpp
@@ -40,7 +40,7 @@ void UMVVMBlueprintViewCondition::SetConditionPath(FMVVMBlueprintPropertyPath In
 	ConditionPath = MoveTemp(InConditionPath);
 	GraphName = FName();
 
-	if (ConditionPath.IsValid())
+	if (ConditionPath.IsValid() || DestinationPath.IsValid())
 	{
 		TStringBuilder<256> StringBuilder;
 		StringBuilder << TEXT("__");
@@ -64,6 +64,15 @@ void UMVVMBlueprintViewCondition::SetDestinationPath(FMVVMBlueprintPropertyPath
 	RemoveWrapperGraph(RemoveConversionFunctionCurrentValues);
 
 	DestinationPath = MoveTemp(InDestinationPath);
+	GraphName = FName();
+
+	if (ConditionPath.IsValid() || DestinationPath.IsValid())
+	{
+		TStringBuilder<256> StringBuilder;
+		StringBuilder << TEXT("__");
+		StringBuilder << FGuid::NewGuid().ToString(EGuidFormats::DigitsWithHyphensLower);
+		GraphName = StringBuilder.ToString();
+	}
 
 	bNeedsToRegenerateChildren = true;
 
@@ -329,7 +338,7 @@ const UFunction* UMVVMBlueprintViewCondition::GetDestinationSignature() const
 
 UEdGraph* UMVVMBlueprintViewCondition::CreateWrapperGraphInternal()
 {
-	if (GraphName.IsNone() || !DestinationPath.IsValid() || !ConditionPath.IsValid())
+	if (GraphName.IsNone() || (!ConditionPath.IsValid() && !DestinationPath.IsValid() ) )
 	{
 		return nullptr;
 	}
