Fix or silence false positive PVS warnings in 7.36

diff --git a/Engine/Plugins/Compositing/Composure/Source/Composure/Private/ComposurePlayerCompositingTarget.cpp b/Engine/Plugins/Compositing/Composure/Source/Composure/Private/ComposurePlayerCompositingTarget.cpp
--- a/Engine/Plugins/Compositing/Composure/Source/Composure/Private/ComposurePlayerCompositingTarget.cpp
+++ b/Engine/Plugins/Compositing/Composure/Source/Composure/Private/ComposurePlayerCompositingTarget.cpp
@@ -84,8 +84,8 @@ APlayerCameraManager* UComposurePlayerCompositingTarget::SetPlayerCameraManager(
 	if (PlayerCameraManager)
 	{
 		// Stop rendering any component.
-		check(InPlayerCameraManager->PCOwner);
-		InPlayerCameraManager->PCOwner->bRenderPrimitiveComponents = false;
+		check(PlayerCameraManager->PCOwner);
+		PlayerCameraManager->PCOwner->bRenderPrimitiveComponents = false;
 
 		// Setup camera modifier to the camera manager.
 		check(!PlayerCameraModifier);
diff --git a/Engine/Plugins/Enterprise/DatasmithCADImporter/Source/DatasmithOpenNurbsTranslator/Private/DatasmithOpenNurbsTranslator.cpp b/Engine/Plugins/Enterprise/DatasmithCADImporter/Source/DatasmithOpenNurbsTranslator/Private/DatasmithOpenNurbsTranslator.cpp
--- a/Engine/Plugins/Enterprise/DatasmithCADImporter/Source/DatasmithOpenNurbsTranslator/Private/DatasmithOpenNurbsTranslator.cpp
+++ b/Engine/Plugins/Enterprise/DatasmithCADImporter/Source/DatasmithOpenNurbsTranslator/Private/DatasmithOpenNurbsTranslator.cpp
@@ -441,6 +441,7 @@ namespace DatasmithOpenNurbsTranslatorUtils
 
 			if (bHasFaceMaterialChannel)
 			{
+				checkSlow(FaceMaterialChannel);
 				PolyGroupId = MeshDescription.CreatePolygonGroup();
 				PolygonGroupImportedMaterialSlotNames[PolyGroupId] = DatasmithMeshHelper::DefaultSlotName(FaceMaterialChannel[MeshIndex]);
 			}
diff --git a/Engine/Plugins/Enterprise/DatasmithImporter/Source/DatasmithImporter/Private/DatasmithMaterialExpressions.cpp b/Engine/Plugins/Enterprise/DatasmithImporter/Source/DatasmithImporter/Private/DatasmithMaterialExpressions.cpp
--- a/Engine/Plugins/Enterprise/DatasmithImporter/Source/DatasmithImporter/Private/DatasmithMaterialExpressions.cpp
+++ b/Engine/Plugins/Enterprise/DatasmithImporter/Source/DatasmithImporter/Private/DatasmithMaterialExpressions.cpp
@@ -571,6 +571,7 @@ UMaterialExpression* FDatasmithMaterialExpressions::Lerp(UObject* UnrealMaterial
 	}
 	else
 	{
+		checkSlow(ToBeConnectedA && ToBeConnectedB);
 		ToBeConnectedB->ConnectExpression(LerpExpression->GetInput(1), 0);
 		ToBeConnectedA->ConnectExpression(LerpExpression->GetInput(0), 0);
 	}
diff --git a/Engine/Plugins/Media/ElectraPlayer/Source/ElectraPlayerRuntime/Private/Runtime/Demuxer/ParserISO14496-12.cpp b/Engine/Plugins/Media/ElectraPlayer/Source/ElectraPlayerRuntime/Private/Runtime/Demuxer/ParserISO14496-12.cpp
--- a/Engine/Plugins/Media/ElectraPlayer/Source/ElectraPlayerRuntime/Private/Runtime/Demuxer/ParserISO14496-12.cpp
+++ b/Engine/Plugins/Media/ElectraPlayer/Source/ElectraPlayerRuntime/Private/Runtime/Demuxer/ParserISO14496-12.cpp
@@ -8462,11 +8462,11 @@ namespace Electra
 							case FMP4Box::kSample_tmcd:
 							{
 								UEMediaError Error = ParseTMCDSampleType(PlayerSession, Track.Get(), STSDFirstChildBox);
-								if (Error == UEMEDIA_ERROR_NOT_SUPPORTED)
+								if (Error == UEMEDIA_ERROR_NOT_SUPPORTED) //-V547
 								{
 									bIsSupported = false;
 								}
-								else if (Error != UEMEDIA_ERROR_OK)
+								else if (Error != UEMEDIA_ERROR_OK) //-V547
 								{
 									return Error;
 								}
diff --git a/Engine/Plugins/PCG/Source/PCG/Private/Grid/PCGLandscapeCache.cpp b/Engine/Plugins/PCG/Source/PCG/Private/Grid/PCGLandscapeCache.cpp
--- a/Engine/Plugins/PCG/Source/PCG/Private/Grid/PCGLandscapeCache.cpp
+++ b/Engine/Plugins/PCG/Source/PCG/Private/Grid/PCGLandscapeCache.cpp
@@ -130,7 +130,7 @@ FPCGLandscapeCacheEntry* FPCGLandscapeCacheEntry::CreateCacheEntry(ULandscapeInf
 			if (CDI.GetWeightmapTextureData(LayerInfo, LayerCache, /*bUseEditingLayer=*/false, /*bRemoveSubsectionDuplicates=*/true))
 			{
 				Result->LayerDataNames.Add(Layer.LayerName);
-				Result->LayerData.Emplace(std::move(LayerCache));
+				Result->LayerData.Emplace(MoveTemp(LayerCache));
 			}
 
 			LayerCache.Reset();
diff --git a/Engine/Plugins/PCG/Source/PCGEditor/Private/Nodes/PCGEditorGraphNodeBase.cpp b/Engine/Plugins/PCG/Source/PCGEditor/Private/Nodes/PCGEditorGraphNodeBase.cpp
--- a/Engine/Plugins/PCG/Source/PCGEditor/Private/Nodes/PCGEditorGraphNodeBase.cpp
+++ b/Engine/Plugins/PCG/Source/PCGEditor/Private/Nodes/PCGEditorGraphNodeBase.cpp
@@ -318,12 +318,12 @@ void UPCGEditorGraphNodeBase::GetNodeContextMenuActions(UToolMenu* Menu, class U
 			}
 
 			// Operation to convert a node in place into something else, such as a different node.
-			const TArray<FPCGPreconfiguredInfo> ConversionInfo = Settings->GetConversionInfo();
+			TArray<FPCGPreconfiguredInfo> ConversionInfo = Settings->GetConversionInfo();
 			if (!ConversionInfo.IsEmpty())
 			{
 				const int32 NumConversions = ConversionInfo.Num();
 
-				auto AddEntries = [this, &Section, ConversionInfo = std::move(ConversionInfo)](UToolMenu* AlignmentMenu)
+				auto AddEntries = [this, &Section, ConversionInfo = MoveTemp(ConversionInfo)](UToolMenu* AlignmentMenu)
 				{
 					FToolMenuSection& SubSection = (ConversionInfo.Num() > 1) ? AlignmentMenu->AddSection("EdGraphSchemaConversion", ConversionHeaderLabel) : Section;
 
diff --git a/Engine/Plugins/Runtime/ModelViewViewModel/Source/ModelViewViewModelBlueprint/Private/Bindings/MVVMConversionFunctionHelper.cpp b/Engine/Plugins/Runtime/ModelViewViewModel/Source/ModelViewViewModelBlueprint/Private/Bindings/MVVMConversionFunctionHelper.cpp
--- a/Engine/Plugins/Runtime/ModelViewViewModel/Source/ModelViewViewModelBlueprint/Private/Bindings/MVVMConversionFunctionHelper.cpp
+++ b/Engine/Plugins/Runtime/ModelViewViewModel/Source/ModelViewViewModelBlueprint/Private/Bindings/MVVMConversionFunctionHelper.cpp
@@ -1521,10 +1521,15 @@ void SetPropertyPathForPin(const UBlueprint* Blueprint, const FMVVMBlueprintProp
 		}
 	}
 
-	UK2Node_FunctionEntry* ConversionFunctionEntry = ConversionNode ? Private::FindFunctionEntry(FunctionGraph) : nullptr;
-	UK2Node_FunctionResult* ConversionFunctionResult = ConversionNode ? Private::FindFunctionResult(FunctionGraph) : nullptr;
+	if (!FunctionGraph || !ConversionNode)
+	{
+		return;
+	}
+
+	UK2Node_FunctionEntry* ConversionFunctionEntry = Private::FindFunctionEntry(FunctionGraph);
+	UK2Node_FunctionResult* ConversionFunctionResult = Private::FindFunctionResult(FunctionGraph);
 
-	if (!FunctionGraph || !ConversionNode || !ConversionFunctionEntry || !ConversionFunctionResult)
+	if (!ConversionFunctionEntry || !ConversionFunctionResult)
 	{
 		return;
 	}
diff --git a/Engine/Plugins/VirtualProduction/Avalanche/Source/AvalancheRemoteControl/Private/AvaRemoteControlRebind.cpp b/Engine/Plugins/VirtualProduction/Avalanche/Source/AvalancheRemoteControl/Private/AvaRemoteControlRebind.cpp
--- a/Engine/Plugins/VirtualProduction/Avalanche/Source/AvalancheRemoteControl/Private/AvaRemoteControlRebind.cpp
+++ b/Engine/Plugins/VirtualProduction/Avalanche/Source/AvalancheRemoteControl/Private/AvaRemoteControlRebind.cpp
@@ -142,7 +142,7 @@ public:
 		}
 
 		BindFailureReason = FString::Printf(TEXT("Another binding is already bound to {Owner:%s, SubObject:%s}.\n"),
-			*InActor->GetPathName(), *InSubObject->GetPathName());
+			*GetPathNameSafe(InActor), *GetPathNameSafe(InSubObject));
 		return false;
 	}
 	
diff --git a/Engine/Source/Developer/PakFileUtilities/Private/PakFileUtilities.cpp b/Engine/Source/Developer/PakFileUtilities/Private/PakFileUtilities.cpp
--- a/Engine/Source/Developer/PakFileUtilities/Private/PakFileUtilities.cpp
+++ b/Engine/Source/Developer/PakFileUtilities/Private/PakFileUtilities.cpp
@@ -4103,6 +4103,7 @@ bool ExtractFilesFromPak(const TCHAR* InPakFilename, TMap<FString, FFileInfo>& I
 					const FPakEntry& Entry = It.Info();
 					if (Entry.IsDeleteRecord())
 					{
+						checkSlow(OutDeletedEntries);
 						UE_LOG(LogPakFile, Display, TEXT("Found delete record for \"%s\"."), *EntryFileName);
 
 						FPakInputPair DeleteRecord;
diff --git a/Engine/Source/Developer/ShaderPreprocessor/Private/stb_preprocess/preprocessor.c b/Engine/Source/Developer/ShaderPreprocessor/Private/stb_preprocess/preprocessor.c
--- a/Engine/Source/Developer/ShaderPreprocessor/Private/stb_preprocess/preprocessor.c
+++ b/Engine/Source/Developer/ShaderPreprocessor/Private/stb_preprocess/preprocessor.c
@@ -4732,7 +4732,7 @@ static void transition(int a, int b, int c)
 	else
 	{
 		if (b == DEFAULT)
-			memset(pp_transition_table[a], c, PP_CHAR_CLASS_count);
+			memset(pp_transition_table[a], c, PP_CHAR_CLASS_count); //-V557
 		else if (b == ANY_LETTER)
 		{
 			transition(a, PP_CHAR_CLASS_idtype, c);
@@ -4761,7 +4761,7 @@ static void initscan_transition(int a, int b, int c)
 	else
 	{
 		if (b == DEFAULT)
-			memset(pp_transition_initscan[a], c, PP_CHAR_CLASS_count);
+			memset(pp_transition_initscan[a], c, PP_CHAR_CLASS_count); //-V557
 		else if (b == ANY_LETTER)
 		{
 			initscan_transition(a, PP_CHAR_CLASS_idtype, c);
diff --git a/Engine/Source/Runtime/Engine/Private/Particles/ParticleSystemComponent.cpp b/Engine/Source/Runtime/Engine/Private/Particles/ParticleSystemComponent.cpp
--- a/Engine/Source/Runtime/Engine/Private/Particles/ParticleSystemComponent.cpp
+++ b/Engine/Source/Runtime/Engine/Private/Particles/ParticleSystemComponent.cpp
@@ -2993,7 +2993,7 @@ void UParticleSystemComponent::SetTemplate(class UParticleSystem* NewTemplate)
 	if (PoolingMethod != EPSCPoolMethod::None)
 	{
 		UE_LOG(LogParticles, Warning, TEXT("Changing template on pooled PSC! This will cause a reinit of the system, eliminating the benefits of pooling! Please avoid doing this.\nPSC: %s\nOld Template: %s\nNew Template: %s")
-		, *GetFullName(), *Template->GetFullName(), *NewTemplate->GetFullName());
+		, *GetFullName(), *GetFullNameSafe(Template), *GetFullNameSafe(NewTemplate));
 	}
 
 	if( GIsAllowingParticles || GIsEditor ) 
diff --git a/Engine/Source/Runtime/Engine/Private/Polygon.cpp b/Engine/Source/Runtime/Engine/Private/Polygon.cpp
--- a/Engine/Source/Runtime/Engine/Private/Polygon.cpp
+++ b/Engine/Source/Runtime/Engine/Private/Polygon.cpp
@@ -215,7 +215,7 @@ int32 FPoly::SplitWithPlane
 	else
 	{
 		// Split.
-		if( FrontPoly==NULL )
+		if( FrontPoly == nullptr || BackPoly == nullptr )
 			return SP_Split; // Caller only wanted status.
 
 		*FrontPoly = *this; // Copy all info.
diff --git a/Engine/Source/Runtime/Engine/Private/TextureDerivedData.cpp b/Engine/Source/Runtime/Engine/Private/TextureDerivedData.cpp
--- a/Engine/Source/Runtime/Engine/Private/TextureDerivedData.cpp
+++ b/Engine/Source/Runtime/Engine/Private/TextureDerivedData.cpp
@@ -1596,7 +1596,7 @@ static void GetBuildSettingsPerFormat(
 			}
 
 			FTexturePlatformData::FTextureEncodeResultMetadata* OutResultMetadata = nullptr;
-			if (OutResultMetadataPerLayer)
+			if (OutResultMetadataPerLayer && SourceResultMetadata)
 			{
 				OutResultMetadata = &OutResultMetadataPerLayer->Add_GetRef(*SourceResultMetadata);			
 			}
@@ -3377,7 +3377,7 @@ void FTexturePlatformData::SerializeCooked(FArchive& Ar, UTexture* Owner, bool b
 			SizeY = Mips[0].SizeY;
 			
 			// SizeZ is not the same as NumSlices for texture arrays and cubemaps.
-			if (Owner && Owner->IsA(UVolumeTexture::StaticClass()))
+			if (Cast<UVolumeTexture>(Owner))
 			{
 				SetNumSlices(Mips[0].SizeZ);
 			}
diff --git a/Engine/Source/Runtime/GeometryCore/Private/CompGeom/ConvexDecomposition3.cpp b/Engine/Source/Runtime/GeometryCore/Private/CompGeom/ConvexDecomposition3.cpp
--- a/Engine/Source/Runtime/GeometryCore/Private/CompGeom/ConvexDecomposition3.cpp
+++ b/Engine/Source/Runtime/GeometryCore/Private/CompGeom/ConvexDecomposition3.cpp
@@ -176,7 +176,7 @@ bool FConvexDecomposition3::ConvexPartVsSphereOverlap(const FConvexDecomposition
 		}
 		if (bMustTransformConvex)
 		{
-			Plane.Transform(*TransformIntoSphereSpace);
+			Plane.Transform(*TransformIntoSphereSpace); //-V522
 		}
 		double PlaneDist = Plane.DistanceTo(Center);
 		MaxPlaneDist = FMath::Max(MaxPlaneDist, PlaneDist);
diff --git a/Engine/Source/Runtime/JsonUtilities/Private/JsonObjectConverter.cpp b/Engine/Source/Runtime/JsonUtilities/Private/JsonObjectConverter.cpp
--- a/Engine/Source/Runtime/JsonUtilities/Private/JsonObjectConverter.cpp
+++ b/Engine/Source/Runtime/JsonUtilities/Private/JsonObjectConverter.cpp
@@ -842,7 +842,7 @@ namespace
 			if (JsonValue->Type == EJson::Object)
 			{
 				UObject* Outer = GetTransientPackage();
-				if (ContainerStruct->IsChildOf(UObject::StaticClass()))
+				if (ContainerStruct && ContainerStruct->IsChildOf(UObject::StaticClass()))
 				{
 					Outer = (UObject*)Container;
 				}
diff --git a/Engine/Source/Runtime/MovieScene/Private/EntitySystem/MovieSceneEntityLedger.cpp b/Engine/Source/Runtime/MovieScene/Private/EntitySystem/MovieSceneEntityLedger.cpp
--- a/Engine/Source/Runtime/MovieScene/Private/EntitySystem/MovieSceneEntityLedger.cpp
+++ b/Engine/Source/Runtime/MovieScene/Private/EntitySystem/MovieSceneEntityLedger.cpp
@@ -152,7 +152,7 @@ void FEntityLedger::UpdateOneShotEntities(UMovieSceneEntitySystemLinker* Linker,
 
 void FEntityLedger::UpdateConditionalEntities(UMovieSceneEntitySystemLinker* Linker, const FEntityImportSequenceParams& ImportParams, const FMovieSceneEntityComponentField* EntityField, const FMovieSceneEvaluationFieldEntitySet& ConditionalEntities)
 {
-	if (ConditionalEntities.Num() == 0)
+	if (ConditionalEntities.Num() == 0 || EntityField == nullptr)
 	{
 		return;
 	}
diff --git a/Engine/Source/Runtime/Navmesh/Private/Recast/RecastArea.cpp b/Engine/Source/Runtime/Navmesh/Private/Recast/RecastArea.cpp
--- a/Engine/Source/Runtime/Navmesh/Private/Recast/RecastArea.cpp
+++ b/Engine/Source/Runtime/Navmesh/Private/Recast/RecastArea.cpp
@@ -382,12 +382,18 @@ bool rcErodeWalkableAndLowAreas(rcContext* ctx, int radius, unsigned int height,
 	ctx->startTimer(RC_TIMER_ERODE_AREA);
 
 	unsigned char* dist = (unsigned char*)rcAlloc(sizeof(unsigned char)*chf.spanCount, RC_ALLOC_TEMP);
-	unsigned char* seed = (unsigned char*)rcAlloc(sizeof(unsigned char)*chf.spanCount, RC_ALLOC_TEMP);
-	if (!dist || !seed)
+	if (!dist)
 	{
 		ctx->log(RC_LOG_ERROR, "erodeWalkableArea: Out of memory 'dist' (%d).", chf.spanCount);
 		return false;
 	}
+	unsigned char* seed = (unsigned char*)rcAlloc(sizeof(unsigned char)*chf.spanCount, RC_ALLOC_TEMP);
+	if (!seed)
+	{
+		rcFree(dist);
+		ctx->log(RC_LOG_ERROR, "erodeWalkableArea: Out of memory 'seed' (%d).", chf.spanCount);
+		return false;
+	}
 
 	// Init distance.
 	memset(dist, 0xff, sizeof(unsigned char)*chf.spanCount);
diff --git a/Engine/Source/Runtime/RHI/Public/RHICommandList.h b/Engine/Source/Runtime/RHI/Public/RHICommandList.h
--- a/Engine/Source/Runtime/RHI/Public/RHICommandList.h
+++ b/Engine/Source/Runtime/RHI/Public/RHICommandList.h
@@ -4526,6 +4526,7 @@ public:
 
 		if (NumUniformBuffers)
 		{
+			checkSlow(UniformBuffers);
 			InlineBindings->UniformBuffers = (FRHIUniformBuffer**)Alloc(sizeof(FRHIUniformBuffer*) * NumUniformBuffers, alignof(FRHIUniformBuffer*));
 			for (uint32 Index = 0; Index < NumUniformBuffers; ++Index)
 			{
diff --git a/Engine/Source/Runtime/RenderCore/Private/RenderGraphPrivate.cpp b/Engine/Source/Runtime/RenderCore/Private/RenderGraphPrivate.cpp
--- a/Engine/Source/Runtime/RenderCore/Private/RenderGraphPrivate.cpp
+++ b/Engine/Source/Runtime/RenderCore/Private/RenderGraphPrivate.cpp
@@ -674,7 +674,7 @@ FRDGScopeState::FState::FState(bool bInImmediate, bool bInParallelExecute)
 		case ERDGScopeMode::TopLevelOnly:
 		case ERDGScopeMode::AllEvents:
 			// Override to a higher level in some cases
-			if (bRDGChannelEnabled != 0)
+			if (bRDGChannelEnabled != 0) //-V547
 			{
 				LocalScopeMode = ERDGScopeMode::AllEventsAndPassNames;
 			}
diff --git a/Engine/Source/Runtime/Renderer/Private/OIT/OIT.cpp b/Engine/Source/Runtime/Renderer/Private/OIT/OIT.cpp
--- a/Engine/Source/Runtime/Renderer/Private/OIT/OIT.cpp
+++ b/Engine/Source/Runtime/Renderer/Private/OIT/OIT.cpp
@@ -217,7 +217,7 @@ static void AddInternalOITComposePass(
 		return;
 
 	FIntPoint InResolution = OITData.SampleDataTexture->Desc.Extent;
-	FIntPoint OutResolution = bModulate ? OutLightingTexture->Desc.Extent : OutColorTexture->Desc.Extent;
+	FIntPoint OutResolution = bModulate ? OutLightingTexture->Desc.Extent : OutColorTexture->Desc.Extent; //-V522
 	FIntPoint Resolution = FIntPoint(FMath::Min(InResolution.X, OutResolution.X), FMath::Min(InResolution.Y, OutResolution.Y));
 
 	FOITPixelCombineCS::FParameters* Parameters = GraphBuilder.AllocParameters<FOITPixelCombineCS::FParameters>();
diff --git a/Engine/Source/Runtime/Renderer/Private/SingleLayerWaterRendering.cpp b/Engine/Source/Runtime/Renderer/Private/SingleLayerWaterRendering.cpp
--- a/Engine/Source/Runtime/Renderer/Private/SingleLayerWaterRendering.cpp
+++ b/Engine/Source/Runtime/Renderer/Private/SingleLayerWaterRendering.cpp
@@ -851,7 +851,7 @@ static FSingleLayerWaterTileClassification ClassifyTiles(FRDGBuilder& GraphBuild
 			PassParameters->TileMaskBufferOut = TileMaskBufferUAV;
 			if (bProduceFroxelData)
 			{
-				PassParameters->FroxelBuilder = OutFroxelViewData->GetBuilderParameters(GraphBuilder);
+				PassParameters->FroxelBuilder = OutFroxelViewData->GetBuilderParameters(GraphBuilder); //-V522
 			}
 			FComputeShaderUtils::AddPass(
 				GraphBuilder,
