MVVM: Fix various undo/redo issues.
1. Ensure conversion function pin changes are transacted and the UI is updated

2. Ensure `MVVMBlueprintViewConditions` and `MVVMBlueprintViewEvents` are transactional.

On property changes, `OnConditionPreEditChange` and `OnEventPreEditChange` would not write object to transaction buffer. can be properly tracked in undo/redo.
The `PostEditChangeChainProperty` call would trigger `SnapshotTransactionBuffer`, but this would fail the `Object->HasAnyFlags(RF_Transactional);` check.

#rnx

diff --git a/Engine/Plugins/Runtime/ModelViewViewModel/Source/ModelViewViewModelEditor/Private/MVVMEditorSubsystem.cpp b/Engine/Plugins/Runtime/ModelViewViewModel/Source/ModelViewViewModelEditor/Private/MVVMEditorSubsystem.cpp
--- a/Engine/Plugins/Runtime/ModelViewViewModel/Source/ModelViewViewModelEditor/Private/MVVMEditorSubsystem.cpp
+++ b/Engine/Plugins/Runtime/ModelViewViewModel/Source/ModelViewViewModelEditor/Private/MVVMEditorSubsystem.cpp
@@ -77,6 +77,7 @@ namespace UE::MVVM::Private
 		EditChain.AddTail(ChangedProperty);
 		EditChain.SetActivePropertyNode(ChangedProperty);
 
+		Event->SetFlags(RF_Transactional);
 		Event->PreEditChange(EditChain);
 	}
 
@@ -103,6 +104,7 @@ namespace UE::MVVM::Private
 		EditChain.AddTail(ChangedProperty);
 		EditChain.SetActivePropertyNode(ChangedProperty);
 
+		Condition->SetFlags(RF_Transactional);
 		Condition->PreEditChange(EditChain);
 	}
 
@@ -1408,12 +1410,20 @@ FMVVMBlueprintPropertyPath UMVVMEditorSubsystem::GetPathForConversionFunctionArg
 
 void UMVVMEditorSubsystem::SetPathForConversionFunctionArgument(UWidgetBlueprint* WidgetBlueprint, FMVVMBlueprintViewBinding& Binding,  const FMVVMBlueprintPinId& ParameterId, const FMVVMBlueprintPropertyPath& Path, bool bSourceToDestination) const
 {
-	UMVVMBlueprintViewConversionFunction* ConversionFunction = Binding.Conversion.GetConversionFunction(bSourceToDestination);
-	if (ConversionFunction)
+	if (UMVVMBlueprintView* View = GetView(WidgetBlueprint))
 	{
-		ConversionFunction->SetGraphPin(WidgetBlueprint, ParameterId, Path);
+		FName NAME_Conversion = GET_MEMBER_NAME_CHECKED(FMVVMBlueprintViewBinding, Conversion);
+		UE::MVVM::Private::OnBindingPreEditChange(View, NAME_Conversion);
+
+		UMVVMBlueprintViewConversionFunction* ConversionFunction = Binding.Conversion.GetConversionFunction(bSourceToDestination);
+		if (ConversionFunction)
+		{
+			ConversionFunction->SetGraphPin(WidgetBlueprint, ParameterId, Path);
+		}
+		FBlueprintEditorUtils::MarkBlueprintAsModified(WidgetBlueprint);
+
+		UE::MVVM::Private::OnBindingPostEditChange(View, NAME_Conversion);
 	}
-	FBlueprintEditorUtils::MarkBlueprintAsModified(WidgetBlueprint);
 }
 
 namespace Private
