MVVM: Fix graph creation with empty conditions.

As a consequence of 43158377, incomplete conditions will attempt to rebuild incomplete conditions to support pin editting beforce condition is set. This change fixes the inverse use case. This will prevent attempting to create graph with empty destination (which is causing as assertion to be thrown)

#rnx

diff --git a/Engine/Plugins/Runtime/ModelViewViewModel/Source/ModelViewViewModelBlueprint/Private/MVVMBlueprintViewCondition.cpp b/Engine/Plugins/Runtime/ModelViewViewModel/Source/ModelViewViewModelBlueprint/Private/MVVMBlueprintViewCondition.cpp
--- a/Engine/Plugins/Runtime/ModelViewViewModel/Source/ModelViewViewModelBlueprint/Private/MVVMBlueprintViewCondition.cpp
+++ b/Engine/Plugins/Runtime/ModelViewViewModel/Source/ModelViewViewModelBlueprint/Private/MVVMBlueprintViewCondition.cpp
@@ -338,7 +338,7 @@ const UFunction* UMVVMBlueprintViewCondition::GetDestinationSignature() const
 
 UEdGraph* UMVVMBlueprintViewCondition::CreateWrapperGraphInternal()
 {
-	if (GraphName.IsNone() || (!ConditionPath.IsValid() && !DestinationPath.IsValid() ) )
+	if (GraphName.IsNone() || !DestinationPath.IsValid())
 	{
 		return nullptr;
 	}
