Fix for static analysis warnings

#tests Manually tried with nullptr to both values
#rnx

diff --git a/Engine/Plugins/Runtime/ModelViewViewModel/Source/ModelViewViewModelEditor/Private/MVVMEditorSubsystem.cpp b/Engine/Plugins/Runtime/ModelViewViewModel/Source/ModelViewViewModelEditor/Private/MVVMEditorSubsystem.cpp
--- a/Engine/Plugins/Runtime/ModelViewViewModel/Source/ModelViewViewModelEditor/Private/MVVMEditorSubsystem.cpp
+++ b/Engine/Plugins/Runtime/ModelViewViewModel/Source/ModelViewViewModelEditor/Private/MVVMEditorSubsystem.cpp
@@ -1451,6 +1451,10 @@ namespace Private
 		, FText TransactionName
 		, TFunctionRef<bool(const UEdGraphSchema_K2*, UEdGraphPin*)> Test, TFunctionRef<void(const UEdGraphSchema_K2*, UEdGraphPin*)> Action)
 	{
+		if (ViewEvent == nullptr && Binding == nullptr)
+		{
+			return;
+		}
 		UEdGraphPin* GraphPin = ViewEvent
 			? Private::GetGraphPin(WidgetBlueprint, ViewEvent, ParameterId)
 			: Private::GetGraphPin(Subsystem, WidgetBlueprint, *Binding, ParameterId, bSourceToDestination);
@@ -1563,67 +1567,67 @@ namespace Private
 	}
 } //namespace
 
-void UMVVMEditorSubsystem::SplitPin(UWidgetBlueprint* WidgetBlueprint, FMVVMBlueprintViewBinding& Binding,  const FMVVMBlueprintPinId& ParameterId, bool bSourceToDestination) const
+void UMVVMEditorSubsystem::SplitPin(UWidgetBlueprint* WidgetBlueprint, FMVVMBlueprintViewBinding& Binding, const FMVVMBlueprintPinId& ParameterId, bool bSourceToDestination) const
 {
 	Private::SplitPin(this, WidgetBlueprint, nullptr, &Binding, ParameterId, bSourceToDestination);
 }
 
-bool UMVVMEditorSubsystem::CanSplitPin(const UWidgetBlueprint* WidgetBlueprint, const FMVVMBlueprintViewBinding& Binding,  const FMVVMBlueprintPinId& ParameterId, bool bSourceToDestination) const
+bool UMVVMEditorSubsystem::CanSplitPin(const UWidgetBlueprint* WidgetBlueprint, const FMVVMBlueprintViewBinding& Binding, const FMVVMBlueprintPinId& ParameterId, bool bSourceToDestination) const
 {
 	UEdGraphPin* GraphPin = Private::GetGraphPin(this, WidgetBlueprint, Binding, ParameterId, bSourceToDestination);
 	return Private::CanSplitPin(GraphPin);
 }
 
-void UMVVMEditorSubsystem::SplitPin(UWidgetBlueprint* WidgetBlueprint, UMVVMBlueprintViewEvent* ViewEvent,  const FMVVMBlueprintPinId& ParameterId) const
+void UMVVMEditorSubsystem::SplitPin(UWidgetBlueprint* WidgetBlueprint, UMVVMBlueprintViewEvent* ViewEvent, const FMVVMBlueprintPinId& ParameterId) const
 {
 	Private::SplitPin(this, WidgetBlueprint, ViewEvent, nullptr, ParameterId, true);
 }
 
-bool UMVVMEditorSubsystem::CanSplitPin(const UWidgetBlueprint* WidgetBlueprint, UMVVMBlueprintViewEvent* Event,  const FMVVMBlueprintPinId& ParameterId) const
+bool UMVVMEditorSubsystem::CanSplitPin(const UWidgetBlueprint* WidgetBlueprint, UMVVMBlueprintViewEvent* Event, const FMVVMBlueprintPinId& ParameterId) const
 {
 	UEdGraphPin* GraphPin = Private::GetGraphPin(WidgetBlueprint, Event, ParameterId);
 	return Private::CanSplitPin(GraphPin);
 }
 
-void UMVVMEditorSubsystem::RecombinePin(UWidgetBlueprint* WidgetBlueprint, FMVVMBlueprintViewBinding& Binding,  const FMVVMBlueprintPinId& ParameterId, bool bSourceToDestination) const
+void UMVVMEditorSubsystem::RecombinePin(UWidgetBlueprint* WidgetBlueprint, FMVVMBlueprintViewBinding& Binding, const FMVVMBlueprintPinId& ParameterId, bool bSourceToDestination) const
 {
 	Private::RecombinePin(this, WidgetBlueprint, nullptr, &Binding, ParameterId, bSourceToDestination);
 }
 
-bool UMVVMEditorSubsystem::CanRecombinePin(const UWidgetBlueprint* WidgetBlueprint, const FMVVMBlueprintViewBinding& Binding,  const FMVVMBlueprintPinId& ParameterId, bool bSourceToDestination) const
+bool UMVVMEditorSubsystem::CanRecombinePin(const UWidgetBlueprint* WidgetBlueprint, const FMVVMBlueprintViewBinding& Binding, const FMVVMBlueprintPinId& ParameterId, bool bSourceToDestination) const
 {
 	UEdGraphPin* GraphPin = Private::GetGraphPin(this, WidgetBlueprint, Binding, ParameterId, bSourceToDestination);
 	return Private::CanRecombinePin(GraphPin);
 }
 
-void UMVVMEditorSubsystem::RecombinePin(UWidgetBlueprint* WidgetBlueprint, UMVVMBlueprintViewEvent* ViewEvent,   const FMVVMBlueprintPinId& ParameterId) const
+void UMVVMEditorSubsystem::RecombinePin(UWidgetBlueprint* WidgetBlueprint, UMVVMBlueprintViewEvent* ViewEvent, const FMVVMBlueprintPinId& ParameterId) const
 {
 	Private::RecombinePin(this, WidgetBlueprint, ViewEvent, nullptr, ParameterId, true);
 }
 
-bool UMVVMEditorSubsystem::CanRecombinePin(const UWidgetBlueprint* WidgetBlueprint, UMVVMBlueprintViewEvent* Event,  const FMVVMBlueprintPinId& ParameterId) const
+bool UMVVMEditorSubsystem::CanRecombinePin(const UWidgetBlueprint* WidgetBlueprint, UMVVMBlueprintViewEvent* Event, const FMVVMBlueprintPinId& ParameterId) const
 {
 	UEdGraphPin* GraphPin = Private::GetGraphPin(WidgetBlueprint, Event, ParameterId);
 	return Private::CanRecombinePin(GraphPin);
 }
 
-void UMVVMEditorSubsystem::ResetPinToDefaultValue(UWidgetBlueprint* WidgetBlueprint, FMVVMBlueprintViewBinding& Binding,  const FMVVMBlueprintPinId& ParameterId, bool bSourceToDestination) const
+void UMVVMEditorSubsystem::ResetPinToDefaultValue(UWidgetBlueprint* WidgetBlueprint, FMVVMBlueprintViewBinding& Binding, const FMVVMBlueprintPinId& ParameterId, bool bSourceToDestination) const
 {
 	Private::ResetPinToDefaultValue(this, WidgetBlueprint, nullptr, &Binding, ParameterId, bSourceToDestination);
 }
 
-bool UMVVMEditorSubsystem::CanResetPinToDefaultValue(const UWidgetBlueprint* WidgetBlueprint, const FMVVMBlueprintViewBinding& Binding,  const FMVVMBlueprintPinId& ParameterId, bool bSourceToDestination) const
+bool UMVVMEditorSubsystem::CanResetPinToDefaultValue(const UWidgetBlueprint* WidgetBlueprint, const FMVVMBlueprintViewBinding& Binding, const FMVVMBlueprintPinId& ParameterId, bool bSourceToDestination) const
 {
 	UEdGraphPin* GraphPin = Private::GetGraphPin(this, WidgetBlueprint, Binding, ParameterId, bSourceToDestination);
 	return Private::CanResetPinToDefaultValue(GraphPin);
 }
 
-void UMVVMEditorSubsystem::ResetPinToDefaultValue(UWidgetBlueprint* WidgetBlueprint, UMVVMBlueprintViewEvent* ViewEvent,  const FMVVMBlueprintPinId& ParameterId) const
+void UMVVMEditorSubsystem::ResetPinToDefaultValue(UWidgetBlueprint* WidgetBlueprint, UMVVMBlueprintViewEvent* ViewEvent, const FMVVMBlueprintPinId& ParameterId) const
 {
 	Private::ResetPinToDefaultValue(this, WidgetBlueprint, ViewEvent, nullptr, ParameterId, true);
 }
 
-bool UMVVMEditorSubsystem::CanResetPinToDefaultValue(const UWidgetBlueprint* WidgetBlueprint, UMVVMBlueprintViewEvent* Event,  const FMVVMBlueprintPinId& ParameterId) const
+bool UMVVMEditorSubsystem::CanResetPinToDefaultValue(const UWidgetBlueprint* WidgetBlueprint, UMVVMBlueprintViewEvent* Event, const FMVVMBlueprintPinId& ParameterId) const
 {
 	UEdGraphPin* GraphPin = Private::GetGraphPin(WidgetBlueprint, Event, ParameterId);
 	return Private::CanResetPinToDefaultValue(GraphPin);
