[MVVM]
- Fix race condition by dirtying the compiled library for reload instead of reloading it right away

diff --git a/Engine/Plugins/Runtime/ModelViewViewModel/Source/ModelViewViewModel/Private/View/MVVMViewClass.cpp b/Engine/Plugins/Runtime/ModelViewViewModel/Source/ModelViewViewModel/Private/View/MVVMViewClass.cpp
--- a/Engine/Plugins/Runtime/ModelViewViewModel/Source/ModelViewViewModel/Private/View/MVVMViewClass.cpp
+++ b/Engine/Plugins/Runtime/ModelViewViewModel/Source/ModelViewViewModel/Private/View/MVVMViewClass.cpp
@@ -790,10 +790,22 @@ void UMVVMViewClass::HandleBlueprintCompiled()
 	// do not keep property alive of potential dead BP class.
 	if (ViewCounter > 0)
 	{
-		BindingLibrary.Unload();
-		BindingLibrary.Load();
+		BindingLibraryDirty.store(true, std::memory_order_relaxed);
+	}
+}
+
+void UMVVMViewClass::ReloadIfDirty() const
+{
+	if (BindingLibraryDirty.load(std::memory_order_relaxed))
+	{
+		TRACE_CPUPROFILER_EVENT_SCOPE(UMVVMViewClass::ReloadIfDirty);
+		UMVVMViewClass* MutableThis = const_cast<UMVVMViewClass*>(this);
+		MutableThis->BindingLibraryDirty.store(false, std::memory_order_relaxed);
+		MutableThis->BindingLibrary.Unload();
+		MutableThis->BindingLibrary.Load();
 	}
 }
+
 #endif
 
 #if UE_WITH_MVVM_DEBUGGING
diff --git a/Engine/Plugins/Runtime/ModelViewViewModel/Source/ModelViewViewModel/Public/View/MVVMViewClass.h b/Engine/Plugins/Runtime/ModelViewViewModel/Source/ModelViewViewModel/Public/View/MVVMViewClass.h
--- a/Engine/Plugins/Runtime/ModelViewViewModel/Source/ModelViewViewModel/Public/View/MVVMViewClass.h
+++ b/Engine/Plugins/Runtime/ModelViewViewModel/Source/ModelViewViewModel/Public/View/MVVMViewClass.h
@@ -651,6 +651,9 @@ public:
 	/** Get the container of all the bindings. */
 	[[nodiscard]] const FMVVMCompiledBindingLibrary& GetBindingLibrary() const
 	{
+#if WITH_EDITOR
+		ReloadIfDirty();
+#endif
 		return BindingLibrary;
 	}
 
@@ -748,6 +751,8 @@ public:
 private:
 #if WITH_EDITOR
 	UE_API void HandleBlueprintCompiled();
+	UE_API void ReloadIfDirty() const;
+	std::atomic<bool> BindingLibraryDirty = false;
 #endif
 
 private:
