Add MVVM Conversion functions for changing material parameters of Font Material in text block.

diff --git a/Engine/Plugins/Runtime/ModelViewViewModel/Source/ModelViewViewModel/Private/Bindings/ConversionLibraries/MVVMFontConversionLibrary.cpp b/Engine/Plugins/Runtime/ModelViewViewModel/Source/ModelViewViewModel/Private/Bindings/ConversionLibraries/MVVMFontConversionLibrary.cpp
new file mode 100644
--- /dev/null
+++ b/Engine/Plugins/Runtime/ModelViewViewModel/Source/ModelViewViewModel/Private/Bindings/ConversionLibraries/MVVMFontConversionLibrary.cpp
@@ -0,0 +1,96 @@
+// Copyright Epic Games, Inc. All Rights Reserved.
+
+#include "Bindings/ConversionLibraries/MVVMFontConversionLibrary.h"
+
+#include "Materials/MaterialInterface.h"
+#include "Materials/MaterialInstanceDynamic.h"
+
+FSlateFontInfo UMVVMFontConversionLibrary::Conv_SetScalarParameter(FSlateFontInfo Font, FName ParameterName, float Value)
+{
+	if (UMaterialInstanceDynamic* DynamicMaterial = TryGetDynamicMaterial(Font, nullptr))
+	{
+		DynamicMaterial->SetScalarParameterValue(ParameterName, Value);
+	}
+
+	return Font;
+}
+
+FSlateFontInfo UMVVMFontConversionLibrary::Conv_SetVectorParameter(FSlateFontInfo Font, FName ParameterName, FLinearColor Value)
+{
+	if (UMaterialInstanceDynamic* DynamicMaterial = TryGetDynamicMaterial(Font, nullptr))
+	{
+		DynamicMaterial->SetVectorParameterValue(ParameterName, Value);
+	}
+
+	return Font;
+}
+
+FSlateFontInfo UMVVMFontConversionLibrary::Conv_SetTextureParameter(FSlateFontInfo Font, FName ParameterName, UTexture* Value)
+{
+	if (UMaterialInstanceDynamic* DynamicMaterial = TryGetDynamicMaterial(Font, nullptr))
+	{
+		DynamicMaterial->SetTextureParameterValue(ParameterName, Value);
+	}
+
+	return Font;
+}
+
+FSlateFontInfo UMVVMFontConversionLibrary::Conv_SetScalarParameterMID(FSlateFontInfo Font, UMaterialInterface* Material, FName ParameterName, float Value)
+{
+	if (UMaterialInstanceDynamic* DynamicMaterial = TryGetDynamicMaterial(Font, nullptr, Material))
+	{
+		DynamicMaterial->SetScalarParameterValue(ParameterName, Value);
+	}
+
+	return Font;
+}
+
+FSlateFontInfo UMVVMFontConversionLibrary::Conv_SetVectorParameterMID(FSlateFontInfo Font, UMaterialInterface* Material, FName ParameterName, FLinearColor Value)
+{
+	if (UMaterialInstanceDynamic* DynamicMaterial = TryGetDynamicMaterial(Font, nullptr, Material))
+	{
+		DynamicMaterial->SetVectorParameterValue(ParameterName, Value);
+	}
+
+	return Font;
+}
+
+FSlateFontInfo UMVVMFontConversionLibrary::Conv_SetTextureParameterMID(FSlateFontInfo Font, UMaterialInterface* Material, FName ParameterName, UTexture* Value)
+{
+	if (UMaterialInstanceDynamic* DynamicMaterial = TryGetDynamicMaterial(Font, nullptr, Material))
+	{
+		DynamicMaterial->SetTextureParameterValue(ParameterName, Value);
+	}
+
+	return Font;
+}
+
+UMaterialInstanceDynamic* UMVVMFontConversionLibrary::TryGetDynamicMaterial(FSlateFontInfo& InFont, UObject* InOuter, UMaterialInterface* InTargetMaterial)
+{
+	 UMaterialInterface* SourceMaterial = Cast<UMaterialInterface>(InFont.FontMaterial);
+
+	 // Want to either make the provided material if different, or make a MID of the existing material
+	 UMaterialInterface* DesiredMaterial = nullptr;
+	 if (InTargetMaterial)
+	 {
+		 DesiredMaterial = InTargetMaterial;
+	 }
+	 else
+	 {
+		 DesiredMaterial = SourceMaterial;
+	 }
+
+	 if (DesiredMaterial)
+	 {
+	 	UMaterialInstanceDynamic* DynamicMaterial = Cast<UMaterialInstanceDynamic>(DesiredMaterial);
+
+	 	if ( !DynamicMaterial )
+	 	{
+			DynamicMaterial = UMaterialInstanceDynamic::Create(DesiredMaterial, InOuter);
+			InFont.FontMaterial = DynamicMaterial;
+	 	}
+	 	return DynamicMaterial;
+	 }
+
+	return nullptr;
+}
\ No newline at end of file
diff --git a/Engine/Plugins/Runtime/ModelViewViewModel/Source/ModelViewViewModel/Private/Bindings/ConversionLibraries/MVVMSlateBrushConversionLibrary.cpp b/Engine/Plugins/Runtime/ModelViewViewModel/Source/ModelViewViewModel/Private/Bindings/ConversionLibraries/MVVMSlateBrushConversionLibrary.cpp
--- a/Engine/Plugins/Runtime/ModelViewViewModel/Source/ModelViewViewModel/Private/Bindings/ConversionLibraries/MVVMSlateBrushConversionLibrary.cpp
+++ b/Engine/Plugins/Runtime/ModelViewViewModel/Source/ModelViewViewModel/Private/Bindings/ConversionLibraries/MVVMSlateBrushConversionLibrary.cpp
@@ -7,9 +7,9 @@
 
 FSlateBrush UMVVMSlateBrushConversionLibrary::Conv_SetScalarParameter(FSlateBrush Brush, FName ParameterName, float Value)
 {
-	if (UMaterialInstanceDynamic* Material = TryGetDynamicMaterial(Brush, nullptr))
+	if (UMaterialInstanceDynamic* DynamicMaterial = TryGetDynamicMaterial(Brush, nullptr))
 	{
-		Material->SetScalarParameterValue(ParameterName, Value);
+		DynamicMaterial->SetScalarParameterValue(ParameterName, Value);
 	}
 
 	return Brush;
@@ -17,9 +17,9 @@ FSlateBrush UMVVMSlateBrushConversionLibrary::Conv_SetScalarParameter(FSlateBrus
 
 FSlateBrush UMVVMSlateBrushConversionLibrary::Conv_SetVectorParameter(FSlateBrush Brush, FName ParameterName, FLinearColor Value)
 {
-	if (UMaterialInstanceDynamic* Material = TryGetDynamicMaterial(Brush, nullptr))
+	if (UMaterialInstanceDynamic* DynamicMaterial = TryGetDynamicMaterial(Brush, nullptr))
 	{
-		Material->SetVectorParameterValue(ParameterName, Value);
+		DynamicMaterial->SetVectorParameterValue(ParameterName, Value);
 	}
 
 	return Brush;
@@ -27,9 +27,9 @@ FSlateBrush UMVVMSlateBrushConversionLibrary::Conv_SetVectorParameter(FSlateBrus
 
 FSlateBrush UMVVMSlateBrushConversionLibrary::Conv_SetTextureParameter(FSlateBrush Brush, FName ParameterName, UTexture* Value)
 {
-	if (UMaterialInstanceDynamic* Material = TryGetDynamicMaterial(Brush, nullptr))
+	if (UMaterialInstanceDynamic* DynamicMaterial = TryGetDynamicMaterial(Brush, nullptr))
 	{
-		Material->SetTextureParameterValue(ParameterName, Value);
+		DynamicMaterial->SetTextureParameterValue(ParameterName, Value);
 	}
 
 	return Brush;
@@ -71,7 +71,7 @@ UMaterialInstanceDynamic* UMVVMSlateBrushConversionLibrary::TryGetDynamicMateria
 
 	 // Want to either make the provided material if different, or make a MID of the existing material
 	 UMaterialInterface* DesiredMaterial = nullptr;
-	 if (InTargetMaterial && (!SourceMaterial || SourceMaterial->GetMaterial() != InTargetMaterial->GetMaterial()))
+	 if (InTargetMaterial)
 	 {
 		 DesiredMaterial = InTargetMaterial;
 	 }
diff --git a/Engine/Plugins/Runtime/ModelViewViewModel/Source/ModelViewViewModel/Public/Bindings/ConversionLibraries/MVVMFontConversionLibrary.h b/Engine/Plugins/Runtime/ModelViewViewModel/Source/ModelViewViewModel/Public/Bindings/ConversionLibraries/MVVMFontConversionLibrary.h
new file mode 100644
--- /dev/null
+++ b/Engine/Plugins/Runtime/ModelViewViewModel/Source/ModelViewViewModel/Public/Bindings/ConversionLibraries/MVVMFontConversionLibrary.h
@@ -0,0 +1,101 @@
+// Copyright Epic Games, Inc. All Rights Reserved.
+
+#pragma once
+
+#include "CoreMinimal.h"
+#include "Components/SlateWrapperTypes.h"
+#include "Kismet/BlueprintFunctionLibrary.h"
+#include "Fonts/SlateFontInfo.h"
+
+#include "MVVMFontConversionLibrary.generated.h"
+
+#define UE_API MODELVIEWVIEWMODEL_API
+
+class UMaterialInstanceDynamic;
+class UMaterialInterface;
+class UTexture;
+
+/**
+ * Conversion library that contains methods for FSlateFontInfo
+ * 
+ * Primarily consists of methods to set material params on an existing FSlateFontInfo
+ */
+UCLASS(MinimalAPI)
+class UMVVMFontConversionLibrary : public UBlueprintFunctionLibrary
+{
+	GENERATED_BODY()
+
+public:
+	
+	/**
+	 * Sets a scalar value on a font material assuming it exists, handles MID existance appropriately
+	 * 
+	 * @param Font			Font with a material we want to set a parameter on 
+	 * @param ParameterName	Name of material parameter to set
+	 * @param Value			Value to set material parameter to
+	 */
+	UFUNCTION(BlueprintPure, Category = "Widget", meta = (DisplayName = "Set Scalar Parameter (Font)", MVVMBindToDestination = "Font"))
+	static UE_API FSlateFontInfo Conv_SetScalarParameter(FSlateFontInfo Font, FName ParameterName, float Value);
+	
+	/**
+	 * Sets a vector value on a font material assuming it exists, handles MID existance appropriately
+	 * 
+	 * @param Font			Font with a material we want to set a parameter on 
+	 * @param ParameterName	Name of material parameter to set
+	 * @param Value			Value to set material parameter to
+	 */
+	UFUNCTION(BlueprintPure, Category = "Widget", meta = (DisplayName = "Set Vector Parameter (Font)", MVVMBindToDestination = "Font"))
+	static UE_API FSlateFontInfo Conv_SetVectorParameter(FSlateFontInfo Font, FName ParameterName, FLinearColor Value);
+	
+	/**
+	 * Sets a texture value on a font material assuming it exists, handles MID existance appropriately
+	 * 
+	 * @param Font			Font with a material we want to set a parameter on 
+	 * @param ParameterName	Name of material parameter to set
+	 * @param Value			Value to set material parameter to
+	 */
+	UFUNCTION(BlueprintPure, Category = "Widget", meta = (DisplayName = "Set Texture Parameter (Font)", MVVMBindToDestination = "Font"))
+	static UE_API FSlateFontInfo Conv_SetTextureParameter(FSlateFontInfo Font, FName ParameterName, UTexture* Value);
+
+	/**
+	 * Sets a scalar value on a font material assuming it exists, handles MID existance appropriately
+	 *
+	 * @param Font			Font with a material we want to set a parameter on
+	 * @param Material		Material to set on the font
+	 * @param ParameterName	Name of material parameter to set
+	 * @param Value			Value to set material parameter to
+	 */
+	UFUNCTION(BlueprintPure, Category = "Widget", meta = (DisplayName = "Set Scalar Parameter MID (Font)", MVVMBindToDestination = "Font"))
+	static UE_API FSlateFontInfo Conv_SetScalarParameterMID(FSlateFontInfo Font, UMaterialInterface* Material, FName ParameterName, float Value);
+
+	/**
+	 * Sets a vector value on a font material assuming it exists, handles MID existance appropriately
+	 *
+	 * @param Font			Font with a material we want to set a parameter on
+	 * @param Material		Material to set on the font
+	 * @param ParameterName	Name of material parameter to set
+	 * @param Value			Value to set material parameter to
+	 */
+	UFUNCTION(BlueprintPure, Category = "Widget", meta = (DisplayName = "Set Vector Parameter MID (Font)", MVVMBindToDestination = "Font"))
+	static UE_API FSlateFontInfo Conv_SetVectorParameterMID(FSlateFontInfo Font, UMaterialInterface* Material, FName ParameterName, FLinearColor Value);
+
+	/**
+	 * Sets a texture value on a font material assuming it exists, handles MID existance appropriately
+	 *
+	 * @param Font			Font with a material we want to set a parameter on
+	 * @param Material		Material to set on the font
+	 * @param ParameterName	Name of material parameter to set
+	 * @param Value			Value to set material parameter to
+	 */
+	UFUNCTION(BlueprintPure, Category = "Widget", meta = (DisplayName = "Set Texture Parameter MID (Font)", MVVMBindToDestination = "Font"))
+	static UE_API FSlateFontInfo Conv_SetTextureParameterMID(FSlateFontInfo Font, UMaterialInterface* Material, FName ParameterName, UTexture* Value);
+
+private:
+
+	/**
+	 * Tries to get a Material Instance Dynamic from the font, will create one if possible using the provided outer, else returns nullptr
+	 */
+	static UMaterialInstanceDynamic* TryGetDynamicMaterial(FSlateFontInfo& InFont, UObject* InOuter, UMaterialInterface* InMaterial = nullptr);
+};
+
+#undef UE_API
diff --git a/Engine/Plugins/Runtime/ModelViewViewModel/Source/ModelViewViewModel/Public/Bindings/ConversionLibraries/MVVMSlateBrushConversionLibrary.h b/Engine/Plugins/Runtime/ModelViewViewModel/Source/ModelViewViewModel/Public/Bindings/ConversionLibraries/MVVMSlateBrushConversionLibrary.h
--- a/Engine/Plugins/Runtime/ModelViewViewModel/Source/ModelViewViewModel/Public/Bindings/ConversionLibraries/MVVMSlateBrushConversionLibrary.h
+++ b/Engine/Plugins/Runtime/ModelViewViewModel/Source/ModelViewViewModel/Public/Bindings/ConversionLibraries/MVVMSlateBrushConversionLibrary.h
@@ -34,7 +34,7 @@ public:
 	 * @param ParameterName	Name of material parameter to set
 	 * @param Value			Value to set material parameter to
 	 */
-	UFUNCTION(BlueprintPure, Category = "Widget", meta = (DisplayName = "Set Scalar Parameter", MVVMBindToDestination = "Brush"))
+	UFUNCTION(BlueprintPure, Category = "Widget", meta = (DisplayName = "Set Scalar Parameter (Brush)", MVVMBindToDestination = "Brush"))
 	static UE_API FSlateBrush Conv_SetScalarParameter(FSlateBrush Brush, FName ParameterName, float Value);
 	
 	/**
@@ -44,7 +44,7 @@ public:
 	 * @param ParameterName	Name of material parameter to set
 	 * @param Value			Value to set material parameter to
 	 */
-	UFUNCTION(BlueprintPure, Category = "Widget", meta = (DisplayName = "Set Vector Parameter", MVVMBindToDestination = "Brush"))
+	UFUNCTION(BlueprintPure, Category = "Widget", meta = (DisplayName = "Set Vector Parameter (Brush)", MVVMBindToDestination = "Brush"))
 	static UE_API FSlateBrush Conv_SetVectorParameter(FSlateBrush Brush, FName ParameterName, FLinearColor Value);
 	
 	/**
@@ -54,7 +54,7 @@ public:
 	 * @param ParameterName	Name of material parameter to set
 	 * @param Value			Value to set material parameter to
 	 */
-	UFUNCTION(BlueprintPure, Category = "Widget", meta = (DisplayName = "Set Texture Parameter", MVVMBindToDestination = "Brush"))
+	UFUNCTION(BlueprintPure, Category = "Widget", meta = (DisplayName = "Set Texture Parameter (Brush)", MVVMBindToDestination = "Brush"))
 	static UE_API FSlateBrush Conv_SetTextureParameter(FSlateBrush Brush, FName ParameterName, UTexture* Value);
 
 	/**
@@ -65,7 +65,7 @@ public:
 	 * @param ParameterName	Name of material parameter to set
 	 * @param Value			Value to set material parameter to
 	 */
-	UFUNCTION(BlueprintPure, Category = "Widget", meta = (DisplayName = "Set Scalar Parameter MID", MVVMBindToDestination = "Brush"))
+	UFUNCTION(BlueprintPure, Category = "Widget", meta = (DisplayName = "Set Scalar Parameter MID (Brush)", MVVMBindToDestination = "Brush"))
 	static UE_API FSlateBrush Conv_SetScalarParameterMID(FSlateBrush Brush, UMaterialInterface* Material, FName ParameterName, float Value);
 
 	/**
@@ -76,7 +76,7 @@ public:
 	 * @param ParameterName	Name of material parameter to set
 	 * @param Value			Value to set material parameter to
 	 */
-	UFUNCTION(BlueprintPure, Category = "Widget", meta = (DisplayName = "Set Vector Parameter MID", MVVMBindToDestination = "Brush"))
+	UFUNCTION(BlueprintPure, Category = "Widget", meta = (DisplayName = "Set Vector Parameter MID (Brush)", MVVMBindToDestination = "Brush"))
 	static UE_API FSlateBrush Conv_SetVectorParameterMID(FSlateBrush Brush, UMaterialInterface* Material, FName ParameterName, FLinearColor Value);
 
 	/**
@@ -87,7 +87,7 @@ public:
 	 * @param ParameterName	Name of material parameter to set
 	 * @param Value			Value to set material parameter to
 	 */
-	UFUNCTION(BlueprintPure, Category = "Widget", meta = (DisplayName = "Set Texture Parameter MID", MVVMBindToDestination = "Brush"))
+	UFUNCTION(BlueprintPure, Category = "Widget", meta = (DisplayName = "Set Texture Parameter MID (Brush)", MVVMBindToDestination = "Brush"))
 	static UE_API FSlateBrush Conv_SetTextureParameterMID(FSlateBrush Brush, UMaterialInterface* Material, FName ParameterName, UTexture* Value);
 
 private:
