PVS 7.37 fixes
#rnx

diff --git a/Engine/Plugins/Runtime/ModelViewViewModel/Source/ModelViewViewModelEditor/Private/MVVMEditorSubsystem.cpp b/Engine/Plugins/Runtime/ModelViewViewModel/Source/ModelViewViewModelEditor/Private/MVVMEditorSubsystem.cpp
--- a/Engine/Plugins/Runtime/ModelViewViewModel/Source/ModelViewViewModelEditor/Private/MVVMEditorSubsystem.cpp
+++ b/Engine/Plugins/Runtime/ModelViewViewModel/Source/ModelViewViewModelEditor/Private/MVVMEditorSubsystem.cpp
@@ -1451,13 +1451,20 @@ namespace Private
 		, FText TransactionName
 		, TFunctionRef<bool(const UEdGraphSchema_K2*, UEdGraphPin*)> Test, TFunctionRef<void(const UEdGraphSchema_K2*, UEdGraphPin*)> Action)
 	{
-		if (ViewEvent == nullptr && Binding == nullptr)
+		UEdGraphPin* GraphPin;
+		if (ViewEvent)
+		{
+			GraphPin = Private::GetGraphPin(WidgetBlueprint, ViewEvent, ParameterId);
+		}
+		else if (Binding)
+		{
+			GraphPin = Private::GetGraphPin(Subsystem, WidgetBlueprint, *Binding, ParameterId, bSourceToDestination);
+		}
+		else
 		{
 			return;
 		}
-		UEdGraphPin* GraphPin = ViewEvent
-			? Private::GetGraphPin(WidgetBlueprint, ViewEvent, ParameterId)
-			: Private::GetGraphPin(Subsystem, WidgetBlueprint, *Binding, ParameterId, bSourceToDestination);
+
 		const UEdGraphSchema_K2* K2Schema = GetDefault<UEdGraphSchema_K2>();
 		if (!Test(K2Schema, GraphPin))
 		{
@@ -1473,7 +1480,7 @@ namespace Private
 			ViewEvent->SavePinValues();
 			UE::MVVM::Private::OnEventPostEditChange(ViewEvent, NAME_SavedPin);
 		}
-		else
+		else if (Binding)
 		{
 			FName NAME_Conversion = GET_MEMBER_NAME_CHECKED(FMVVMBlueprintViewBinding, Conversion);
 			UMVVMBlueprintView* View = Subsystem->GetView(WidgetBlueprint);
diff --git a/Engine/Source/Runtime/Windows/D3D11RHI/Private/D3D11Commands.cpp b/Engine/Source/Runtime/Windows/D3D11RHI/Private/D3D11Commands.cpp
--- a/Engine/Source/Runtime/Windows/D3D11RHI/Private/D3D11Commands.cpp
+++ b/Engine/Source/Runtime/Windows/D3D11RHI/Private/D3D11Commands.cpp
@@ -826,6 +826,7 @@ void FD3D11DynamicRHI::SetRenderTargets(
 	ID3D11RenderTargetView* NewRenderTargetViews[MaxSimultaneousRenderTargets];
 	for(uint32 RenderTargetIndex = 0;RenderTargetIndex < MaxSimultaneousRenderTargets;++RenderTargetIndex)
 	{
+		check(NewRenderTargetsRHI);
 		ID3D11RenderTargetView* RenderTargetView = NULL;
 		if(RenderTargetIndex < NewNumSimultaneousRenderTargets && NewRenderTargetsRHI[RenderTargetIndex].Texture != nullptr)
 		{
