MVVM: Update Pin Values when pasting into events with same event path. Ensure cached pins are up to date.

#rnx

diff --git a/Engine/Plugins/Runtime/ModelViewViewModel/Source/ModelViewViewModelEditor/Private/Framework/MVVMRowHelper.cpp b/Engine/Plugins/Runtime/ModelViewViewModel/Source/ModelViewViewModelEditor/Private/Framework/MVVMRowHelper.cpp
--- a/Engine/Plugins/Runtime/ModelViewViewModel/Source/ModelViewViewModelEditor/Private/Framework/MVVMRowHelper.cpp
+++ b/Engine/Plugins/Runtime/ModelViewViewModel/Source/ModelViewViewModelEditor/Private/Framework/MVVMRowHelper.cpp
@@ -703,6 +703,9 @@ void FRowHelper::PasteEntries(UWidgetBlueprint* WidgetBlueprint, UMVVMBlueprintV
 
 		const FMVVMBlueprintPropertyPath EventPath = EventSelected->GetEventPath();
 		NewEvent->SetEventPath(EventPath);
+
+		NewEvent->RecreateWrapperGraph();
+		NewEvent->UpdatePinValues();
 		BlueprintView->ReplaceEvent(EventSelected, NewEvent);
 	}
 }
