PR #12009 Fix incorrect depreciation message inside MVVMSubsystem

diff --git a/Engine/Plugins/Runtime/ModelViewViewModel/Source/ModelViewViewModel/Public/MVVMSubsystem.h b/Engine/Plugins/Runtime/ModelViewViewModel/Source/ModelViewViewModel/Public/MVVMSubsystem.h
--- a/Engine/Plugins/Runtime/ModelViewViewModel/Source/ModelViewViewModel/Public/MVVMSubsystem.h
+++ b/Engine/Plugins/Runtime/ModelViewViewModel/Source/ModelViewViewModel/Public/MVVMSubsystem.h
@@ -74,7 +74,7 @@ public:
 	static UE_API FMVVMAvailableBinding GetAvailableBindingForEvent(const UClass* Class, FMVVMBindingName BindingName, const UClass* Accessor);
 
 	UE_DEPRECATED(5.3, "GetGlobalViewModelCollection has been deprecated, please use the game instance subsystem.")
-	UFUNCTION(BlueprintCallable, Category = "Viewmodel", meta=(DeprecatedFunction, DeprecatedMessage = "This version of GetGlobalViewModelCollection has been deprecated, please use GetGlobalViewModelCollection from the Viewmodel Game subsystem."))
+	UFUNCTION(BlueprintCallable, Category = "Viewmodel", meta=(DeprecatedFunction, DeprecatedMessage = "This version of GetGlobalViewModelCollection has been deprecated, please use GetViewModelCollection from the Viewmodel Game subsystem."))
 	UE_API UMVVMViewModelCollectionObject* GetGlobalViewModelCollection() const;
 
 	struct FConstDirectionalBindingArgs
