MVVM: fix EUW crash on compile

diff --git a/Engine/Plugins/Runtime/ModelViewViewModel/Source/ModelViewViewModel/Private/ModelViewViewModelModule.cpp b/Engine/Plugins/Runtime/ModelViewViewModel/Source/ModelViewViewModel/Private/ModelViewViewModelModule.cpp
--- a/Engine/Plugins/Runtime/ModelViewViewModel/Source/ModelViewViewModel/Private/ModelViewViewModelModule.cpp
+++ b/Engine/Plugins/Runtime/ModelViewViewModel/Source/ModelViewViewModel/Private/ModelViewViewModelModule.cpp
@@ -33,6 +33,7 @@ void FModelViewViewModelModule::OnPostEngineInit()
 	if (GEditor != nullptr)
 	{
 		OnBlueprintCompiledHandle = GEditor->OnBlueprintCompiled().AddRaw(this, &FModelViewViewModelModule::HandleBlueprintCompiled);
+		OnBlueprintPreCompileHandle = GEditor->OnBlueprintPreCompile().AddRaw(this, &FModelViewViewModelModule::HandleBlueprintPreCompile);
 	}
 #endif
 }
@@ -44,6 +45,14 @@ void FModelViewViewModelModule::HandleBlueprintCompiled()
 #endif
 }
 
+void FModelViewViewModelModule::HandleBlueprintPreCompile(UBlueprint* Blueprint)
+{
+#if WITH_EDITOR
+	OnBlueprintPreCompile.Broadcast(Blueprint);
+#endif
+}
+
+
 void FModelViewViewModelModule::ShutdownModule()
 {
 	if (!IsEngineExitRequested())
@@ -56,9 +65,16 @@ void FModelViewViewModelModule::ShutdownModule()
 
 #if WITH_EDITOR
 	FCoreDelegates::OnPostEngineInit.RemoveAll(this);
-	if (GEditor != nullptr && OnBlueprintCompiledHandle.IsValid())
+	if (GEditor != nullptr)
 	{
-		GEditor->OnBlueprintCompiled().Remove(OnBlueprintCompiledHandle);
+		if (OnBlueprintCompiledHandle.IsValid())
+		{
+			GEditor->OnBlueprintCompiled().Remove(OnBlueprintCompiledHandle);
+		}
+		if (OnBlueprintPreCompileHandle.IsValid())
+		{
+			GEditor->OnBlueprintPreCompile().Remove(OnBlueprintPreCompileHandle);
+		}
 	}
 #endif
 }
diff --git a/Engine/Plugins/Runtime/ModelViewViewModel/Source/ModelViewViewModel/Private/ModelViewViewModelModule.h b/Engine/Plugins/Runtime/ModelViewViewModel/Source/ModelViewViewModel/Private/ModelViewViewModelModule.h
--- a/Engine/Plugins/Runtime/ModelViewViewModel/Source/ModelViewViewModel/Private/ModelViewViewModelModule.h
+++ b/Engine/Plugins/Runtime/ModelViewViewModel/Source/ModelViewViewModel/Private/ModelViewViewModelModule.h
@@ -6,6 +6,7 @@
 #include "Stats/Stats.h"
 
 class IConsoleVariable;
+class UBlueprint;
 
 /** */
 DECLARE_STATS_GROUP(TEXT("UMG Viewmodel"), STATGROUP_UMG_Viewmodel, STATCAT_Advanced);
@@ -26,6 +27,9 @@ public:
 #if WITH_EDITOR
 	DECLARE_TS_MULTICAST_DELEGATE(FOnBlueprintCompiled)
 	FOnBlueprintCompiled OnBlueprintCompiled;
+
+	DECLARE_TS_MULTICAST_DELEGATE_OneParam(FOnBlueprintPreCompile, UBlueprint*)
+	FOnBlueprintPreCompile OnBlueprintPreCompile;
 #endif
 
 private:
@@ -33,8 +37,10 @@ private:
 	void OnPostEngineInit();
 
 	void HandleBlueprintCompiled();
+	void HandleBlueprintPreCompile(UBlueprint* Blueprint);
 #if WITH_EDITOR
 	FDelegateHandle OnBlueprintCompiledHandle;
+	FDelegateHandle OnBlueprintPreCompileHandle;
 #endif
 
 };
diff --git a/Engine/Plugins/Runtime/ModelViewViewModel/Source/ModelViewViewModel/Private/View/MVVMView.cpp b/Engine/Plugins/Runtime/ModelViewViewModel/Source/ModelViewViewModel/Private/View/MVVMView.cpp
--- a/Engine/Plugins/Runtime/ModelViewViewModel/Source/ModelViewViewModel/Private/View/MVVMView.cpp
+++ b/Engine/Plugins/Runtime/ModelViewViewModel/Source/ModelViewViewModel/Private/View/MVVMView.cpp
@@ -14,10 +14,16 @@
 #include "MVVMGameSubsystem.h"
 #include "MVVMMessageLog.h"
 #include "ModelViewViewModelModule.h"
+#include "Modules/ModuleManager.h"
 #include "Templates/ValueOrError.h"
 #include "Types/MVVMFieldContext.h"
 #include "Types/MVVMViewModelCollection.h"
 #include "Bindings/MVVMBindingHelper.h"
+#if WITH_EDITOR
+#include "Engine/Blueprint.h"
+#include "Editor.h"
+#include "Editor/EditorEngine.h"
+#endif
 
 #include UE_INLINE_GENERATED_CPP_BY_NAME(MVVMView)
 
@@ -96,6 +102,17 @@ void UMVVMView::Construct()
 
 	bConstructed = true;
 
+#if WITH_EDITOR
+	if (GEditor)
+	{
+		if (GetWorld()->WorldType == EWorldType::Editor)
+		{
+			FModelViewViewModelModule& Module = FModuleManager::GetModuleChecked<FModelViewViewModelModule>("ModelViewViewModel");
+			BlueprintPreCompileHandle = Module.OnBlueprintPreCompile.AddUObject(this, &UMVVMView::HandleBlueprintPreCompile);
+		}
+	}
+#endif
+
 #if UE_WITH_MVVM_DEBUGGING
 	UE::MVVM::FDebugging::BroadcastViewConstructed(this);
 #endif
@@ -104,6 +121,20 @@ void UMVVMView::Construct()
 
 void UMVVMView::Destruct()
 {
+#if WITH_EDITOR
+	if (GEditor && BlueprintPreCompileHandle.IsValid())
+	{
+		FModelViewViewModelModule& Module = FModuleManager::GetModuleChecked<FModelViewViewModelModule>("ModelViewViewModel");
+		Module.OnBlueprintPreCompile.Remove(BlueprintPreCompileHandle);
+		BlueprintPreCompileHandle.Reset();
+	}
+
+	if (!GeneratedViewClass && GetWorld()->WorldType == EWorldType::Editor)
+	{
+		return;
+	}
+#endif
+
 	check(bConstructed == true);
 	bConstructed = false;
 
@@ -1320,4 +1351,15 @@ bool UMVVMView::AreSourcesValidForBinding(int32 InBindingKey) const
 	return false;
 }
 
+
+#if WITH_EDITOR
+void UMVVMView::HandleBlueprintPreCompile(UBlueprint* Blueprint)
+{
+	if (Blueprint && GeneratedViewClass->GetOuter() == Blueprint->GeneratedClass)
+	{
+		Destruct();
+	}
+}
+#endif
+
 #undef LOCTEXT_NAMESPACE
diff --git a/Engine/Plugins/Runtime/ModelViewViewModel/Source/ModelViewViewModel/Private/View/MVVMViewClass.cpp b/Engine/Plugins/Runtime/ModelViewViewModel/Source/ModelViewViewModel/Private/View/MVVMViewClass.cpp
--- a/Engine/Plugins/Runtime/ModelViewViewModel/Source/ModelViewViewModel/Private/View/MVVMViewClass.cpp
+++ b/Engine/Plugins/Runtime/ModelViewViewModel/Source/ModelViewViewModel/Private/View/MVVMViewClass.cpp
@@ -752,17 +752,23 @@ void UMVVMViewClass::Construct(UUserWidget* UserWidget)
 		FModelViewViewModelModule& Module = FModuleManager::GetModuleChecked<FModelViewViewModelModule>("ModelViewViewModel");
 		Module.OnBlueprintCompiled.AddUObject(this, &UMVVMViewClass::HandleBlueprintCompiled);
 	}
+	bConstructed = true;
 #endif
 }
 
 
 void UMVVMViewClass::Destruct(UUserWidget* UserWidget)
 {
-	--ViewCounter;
-	ensure(ViewCounter >= 0);
-	if (ViewCounter == 0)
+#if WITH_EDITOR
+	if (bConstructed)
+#endif
 	{
-		BindingLibrary.Unload();
+		--ViewCounter;
+		ensure(ViewCounter >= 0);
+		if (ViewCounter == 0)
+		{
+			BindingLibrary.Unload();
+		}
 	}
 
 #if WITH_EDITOR
diff --git a/Engine/Plugins/Runtime/ModelViewViewModel/Source/ModelViewViewModel/Public/View/MVVMView.h b/Engine/Plugins/Runtime/ModelViewViewModel/Source/ModelViewViewModel/Public/View/MVVMView.h
--- a/Engine/Plugins/Runtime/ModelViewViewModel/Source/ModelViewViewModel/Public/View/MVVMView.h
+++ b/Engine/Plugins/Runtime/ModelViewViewModel/Source/ModelViewViewModel/Public/View/MVVMView.h
@@ -217,6 +217,9 @@ private:
 	UE_API void BindEvent(const FMVVMViewClass_Event& ClassItem, FMVVMViewClass_EventKey KeyForLog);
 	UE_API void UnbindEvent(int32 BoundEventIndex);
 	UE_API void ReinitializeEvents(FMVVMViewClass_SourceKey SourceKey, UObject* PreviousValue, UObject* NewValue);
+#if WITH_EDITOR
+	void HandleBlueprintPreCompile(UBlueprint* Blueprint);
+#endif
 
 private:
 	UPROPERTY(Transient)
@@ -268,6 +271,10 @@ private:
 	/** Is the view is registered to the binding subsystem for tick. */
 	UPROPERTY(VisibleAnywhere, Transient, Category = "View")
 	bool bHasDefaultTickBinding = false;
+
+#if WITH_EDITORONLY_DATA
+	FDelegateHandle BlueprintPreCompileHandle;
+#endif
 };
 
 #undef UE_API
diff --git a/Engine/Plugins/Runtime/ModelViewViewModel/Source/ModelViewViewModel/Public/View/MVVMViewClass.h b/Engine/Plugins/Runtime/ModelViewViewModel/Source/ModelViewViewModel/Public/View/MVVMViewClass.h
--- a/Engine/Plugins/Runtime/ModelViewViewModel/Source/ModelViewViewModel/Public/View/MVVMViewClass.h
+++ b/Engine/Plugins/Runtime/ModelViewViewModel/Source/ModelViewViewModel/Public/View/MVVMViewClass.h
@@ -798,6 +798,7 @@ private:
 
 #if WITH_EDITORONLY_DATA
 	FDelegateHandle BlueprintCompiledHandle;
+	bool bConstructed = false;
 #endif
 };
 
