MVVM: Support copy/paste-ing source property paths across widgets.

#rnx

diff --git a/Engine/Plugins/Runtime/ModelViewViewModel/Source/ModelViewViewModelBlueprint/Public/MVVMPropertyPath.h b/Engine/Plugins/Runtime/ModelViewViewModel/Source/ModelViewViewModelBlueprint/Public/MVVMPropertyPath.h
--- a/Engine/Plugins/Runtime/ModelViewViewModel/Source/ModelViewViewModelBlueprint/Public/MVVMPropertyPath.h
+++ b/Engine/Plugins/Runtime/ModelViewViewModel/Source/ModelViewViewModelBlueprint/Public/MVVMPropertyPath.h
@@ -316,8 +316,8 @@ public:
 	}
 
 public:
-	FText ToText(const UWidgetBlueprint* Blueprint, bool bUseDisplayName) const;
-	FString ToString(const UWidgetBlueprint* Blueprint, bool bUseDisplayName, bool bIncludeMetaData) const;
+	MODELVIEWVIEWMODELBLUEPRINT_API FText ToText(const UWidgetBlueprint* Blueprint, bool bUseDisplayName) const;
+	MODELVIEWVIEWMODELBLUEPRINT_API FString ToString(const UWidgetBlueprint* Blueprint, bool bUseDisplayName, bool bIncludeMetaData) const;
 
 private:
 	MODELVIEWVIEWMODELBLUEPRINT_API void DeprecationUpdateSource(const UBlueprint* InContext);
diff --git a/Engine/Plugins/Runtime/ModelViewViewModel/Source/ModelViewViewModelEditor/Private/Framework/MVVMRowHelper.cpp b/Engine/Plugins/Runtime/ModelViewViewModel/Source/ModelViewViewModelEditor/Private/Framework/MVVMRowHelper.cpp
--- a/Engine/Plugins/Runtime/ModelViewViewModel/Source/ModelViewViewModelEditor/Private/Framework/MVVMRowHelper.cpp
+++ b/Engine/Plugins/Runtime/ModelViewViewModel/Source/ModelViewViewModelEditor/Private/Framework/MVVMRowHelper.cpp
@@ -34,6 +34,7 @@
 #include "MVVMMessageLog.h"
 #include "MVVMPropertyPath.h"
 #include "MVVMWidgetBlueprintExtension_View.h"
+#include "Misc/Base64.h"
 #include "Misc/MessageDialog.h"
 #include "Types/MVVMBindingEntry.h"
 #include "Types/MVVMBindingMode.h"
@@ -49,10 +50,12 @@ namespace Private
 {
 // Text import/export uses the unit seperator ascii character code as a delimiter. 
 // Using a non-printable character because all objects and structs are exported as readable text.
-static constexpr TStaticArray<TCHAR, 2> BindingClipboardDelimiter = { 0x1f, TEXT('\0') };
+static constexpr TStaticArray<FString::ElementType, 1> BindingClipboardDelimiter = { 0x1f };
 
 struct FBindingClipboardData
 {
+	UWidgetBlueprint* Blueprint = nullptr;
+
 	using Type = TVariant<FMVVMBlueprintViewBinding, UMVVMBlueprintViewCondition*, UMVVMBlueprintViewEvent*>;
 	TArray<Type> Items;
 };
@@ -119,36 +122,65 @@ bool ExportClipboardData(UMVVMBlueprintView* BlueprintView, const TArray<TShared
 		}
 	}
 
+	if (const UMVVMWidgetBlueprintExtension_View* Extension = BlueprintView->GetOuterUMVVMWidgetBlueprintExtension_View())
+	{
+		OutClipboardData.Blueprint = Extension->GetWidgetBlueprint();
+	}
 	return !OutClipboardData.Items.IsEmpty();
 }
 
 bool ImportClipboardData(UMVVMBlueprintView* BlueprintView, FBindingClipboardData& OutClipboardData)
 {
-	class FBindingEntryTextFactory : public FCustomizableTextObjectFactory
+	class FClipboardObjectFactory : public FCustomizableTextObjectFactory
 	{
 	public:
-		FBindingEntryTextFactory()
+		FClipboardObjectFactory(FBindingClipboardData& ClipboardData)
 			: FCustomizableTextObjectFactory(GWarn)
+			, ClipboardData(ClipboardData)
 		{}
 
 		virtual bool CanCreateClass(UClass* ObjectClass, bool& bOmitSubObjs) const override final
 		{
-			return ObjectClass == UMVVMBlueprintViewEvent::StaticClass() || ObjectClass == UMVVMBlueprintViewCondition::StaticClass();
+			return ObjectClass == UMVVMBlueprintViewEvent::StaticClass() 
+				|| ObjectClass == UMVVMBlueprintViewCondition::StaticClass() 
+				|| ObjectClass == UWidgetBlueprint::StaticClass();
 		}
 
 		virtual void ProcessConstructedObject(UObject* NewObject) override final
 		{
-			Objects.Add(NewObject);
+			if (NewObject->GetClass() == UWidgetBlueprint::StaticClass())
+			{
+				ClipboardData.Blueprint = StaticCast<UWidgetBlueprint*>(NewObject);
+			}
+			else if (NewObject->GetClass() == UMVVMBlueprintViewCondition::StaticClass())
+			{
+				FBindingClipboardData::Type Item;
+				Item.Set<UMVVMBlueprintViewCondition*>(StaticCast<UMVVMBlueprintViewCondition*>(NewObject));
+				ClipboardData.Items.Add(Item);
+			}
+			else if (NewObject->GetClass() == UMVVMBlueprintViewEvent::StaticClass())
+			{
+				FBindingClipboardData::Type Item;
+				Item.Set<UMVVMBlueprintViewEvent*>(StaticCast<UMVVMBlueprintViewEvent*>(NewObject));
+				ClipboardData.Items.Add(Item);
+			}
+			else
+			{				
+				UE_LOG(LogMVVM, Warning, TEXT("Failed to import binding %s from clipboard"), *(NewObject->GetFName().ToString()));
+			}
 		}
-
-		TArray<UObject*> Objects;
+	private:
+		FBindingClipboardData& ClipboardData;
 	};
 
 	FString TextToImport;
 	FPlatformApplicationMisc::ClipboardPaste(TextToImport);
 
+	FString TextToParse;
+	FBase64::Decode(TextToImport, TextToParse);
+
 	TArray<FString> Items;
-	TextToImport.ParseIntoArray(Items, BindingClipboardDelimiter.GetData());
+	TextToParse.ParseIntoArray(Items, BindingClipboardDelimiter.GetData());
 
 	for (const FString& ItemText : Items)
 	{
@@ -168,27 +200,8 @@ bool ImportClipboardData(UMVVMBlueprintView* BlueprintView, FBindingClipboardDat
 		}
 		else
 		{
-			FBindingEntryTextFactory Factory;
-			Factory.ProcessBuffer(BlueprintView, RF_Transactional, ItemText);
-			if (Factory.Objects.Num() != 1)
-			{
-				UE_LOG(LogMVVM, Warning, TEXT("Failed to import object from clipboard"));
-				continue;
-			}
-
-			UObject* Object = Factory.Objects.Last();
-			if (Object->GetClass() == UMVVMBlueprintViewCondition::StaticClass())
-			{
-				FBindingClipboardData::Type Item;
-				Item.Set<UMVVMBlueprintViewCondition*>(StaticCast<UMVVMBlueprintViewCondition*>(Object));
-				OutClipboardData.Items.Add(Item);
-			}
-			else if (Object->GetClass() == UMVVMBlueprintViewEvent::StaticClass())
-			{
-				FBindingClipboardData::Type Item;
-				Item.Set<UMVVMBlueprintViewEvent*>(StaticCast<UMVVMBlueprintViewEvent*>(Object));
-				OutClipboardData.Items.Add(Item);
-			}
+			FClipboardObjectFactory Factory(OutClipboardData);
+			Factory.ProcessBuffer(GetTransientPackage(), RF_Transient, ItemText);
 		}
 	}
 
@@ -296,6 +309,73 @@ FMVVMBlueprintPropertyPath TryGetCommonPropertyPath(UWidgetBlueprint* WidgetBlue
 
 	return NewPath;
 }
+
+FMVVMBlueprintPropertyPath MigratePropertyPathSource(const UMVVMBlueprintView* TargetBlueprintView, const UMVVMBlueprintView* SourceBlueprintView, const FMVVMBlueprintPropertyPath& SourcePropertyPath)
+{
+	if (!ensure(SourceBlueprintView) || !ensure(TargetBlueprintView))
+	{
+		return SourcePropertyPath;
+	}
+
+	// If the views are the same. Note: Object instances are not the same when deserialized via clipboard. Fallback to fully qualified object name to determine equivalence
+	if (SourceBlueprintView->GetFullName() == TargetBlueprintView->GetFullName())
+	{
+		return SourcePropertyPath;
+	}
+
+
+	UWidgetBlueprint* SourceWidgetBlueprint = SourceBlueprintView->GetOuterUMVVMWidgetBlueprintExtension_View() ? SourceBlueprintView->GetOuterUMVVMWidgetBlueprintExtension_View()->GetWidgetBlueprint() : nullptr;
+	if (!ensureMsgf(SourceWidgetBlueprint != nullptr, TEXT("MVVM: Failed to migrate source property path")))
+	{
+		return SourcePropertyPath;
+	}
+				
+	FMVVMBlueprintPropertyPath TargetPath = SourcePropertyPath;
+
+	switch(SourcePropertyPath.GetSource(SourceWidgetBlueprint))
+	{
+		case EMVVMBlueprintFieldPathSource::ViewModel:
+		{
+			// Try to find the source viewmodel in the target view. If found, reset source to this viewmodel.
+			TargetPath.ResetSource();
+
+			const FMVVMBlueprintViewModelContext* SourceViewModel = SourceBlueprintView->FindViewModel(SourcePropertyPath.GetViewModelId());
+			if (SourceViewModel == nullptr)
+			{
+				UE_LOG(LogMVVM, Warning, TEXT("MVVM: Failed to find source viewmodel for source property path %s"), *(SourcePropertyPath.ToString(SourceWidgetBlueprint, true, false)));
+				return TargetPath;
+			}
+
+			if (const FMVVMBlueprintViewModelContext* TargetViewModel = TargetBlueprintView->FindViewModel(SourceViewModel->GetViewModelName()))
+			{
+				TargetPath.SetViewModelId(TargetViewModel->GetViewModelId());
+			}
+			break;
+		}
+		case EMVVMBlueprintFieldPathSource::SelfContext:
+		case EMVVMBlueprintFieldPathSource::Widget:
+		case EMVVMBlueprintFieldPathSource::None:
+		{
+			// Do nothing. SelfContext source path is agnostic to widgetblueprint.
+			// Relative widget paths will work if hierarchy pasted is the same across blueprints. Otherwise will fail as expected.
+			break;
+		}
+	}
+	return TargetPath;
+}
+
+const UMVVMBlueprintView* FindMVVMBlueprintView(const UBlueprint* Blueprint)
+{
+	for (const UBlueprintExtension* Extension : Blueprint->GetExtensions())
+	{
+		if (Extension && Extension->GetClass() == UMVVMWidgetBlueprintExtension_View::StaticClass())
+		{
+			const UMVVMWidgetBlueprintExtension_View* ExtensionView = StaticCast<const UMVVMWidgetBlueprintExtension_View*>(Extension);
+			return ExtensionView->GetBlueprintView();
+		}
+	}
+	return nullptr;
+}
 }
 
 void FRowHelper::GatherAllChildBindings(UMVVMBlueprintView* BlueprintView, const TConstArrayView<TSharedPtr<FBindingEntry>> Entries, TArray<FGuid>& OutBindingIds, TArray<UMVVMBlueprintViewEvent*>& OutEvents, TArray<UMVVMBlueprintViewCondition*>& OutConditions, bool bRecurse)
@@ -467,6 +547,24 @@ void FRowHelper::CopyEntries(UWidgetBlueprint* WidgetBlueprint, UMVVMBlueprintVi
 	
 	FString CopyText;
 
+	if(ClipboardData.Blueprint != nullptr)
+	{
+		FStringOutputDevice Output;
+		const FExportObjectInnerContext Context;
+
+		UnMarkAllObjects(EObjectMark(OBJECTMARK_TagExp | OBJECTMARK_TagImp));
+		UExporter::ExportToOutputDevice(&Context, ClipboardData.Blueprint, nullptr, Output, TEXT("copy"), 0, PPF_ExportsNotFullyQualified | PPF_Copy | PPF_Delimited, false, ClipboardData.Blueprint->GetOuter());
+
+		if (Output.IsEmpty())
+		{
+			UE_LOG(LogMVVM, Warning, TEXT("Failed to export mvvm view for copy"));
+		}
+		else
+		{
+			CopyText.Append(*Output);
+		}
+	}
+
 	for (Private::FBindingClipboardData::Type Item : ClipboardData.Items)
 	{
 		bool bFailed = false;
@@ -561,7 +659,8 @@ void FRowHelper::CopyEntries(UWidgetBlueprint* WidgetBlueprint, UMVVMBlueprintVi
 
 	if (!CopyText.IsEmpty())
 	{
-		FPlatformApplicationMisc::ClipboardCopy(*CopyText);
+		FString Encoded = FBase64::Encode(CopyText);
+		FPlatformApplicationMisc::ClipboardCopy(*Encoded);
 	}
 }
 
@@ -622,10 +721,12 @@ void FRowHelper::PasteEntries(UWidgetBlueprint* WidgetBlueprint, UMVVMBlueprintV
 	}
 
 	const TSharedPtr<FBindingEntry> Entry = Entries.Last();
-	const Private::FBindingClipboardData::Type& Item = ClipboardData.Items.Last();
+	const Private::FBindingClipboardData::Type& ClipboardItem = ClipboardData.Items.Last();
 	UMVVMEditorSubsystem* EditorSubsystem = GEditor->GetEditorSubsystem<UMVVMEditorSubsystem>();
 
-	if (Item.IsType<FMVVMBlueprintViewBinding>())
+	const UMVVMBlueprintView* ClipboardBlueprintView = Private::FindMVVMBlueprintView(ClipboardData.Blueprint);
+
+	if (ClipboardItem.IsType<FMVVMBlueprintViewBinding>())
 	{
 		if (Entry->GetRowType() != FBindingEntry::ERowType::Binding)
 		{
@@ -636,11 +737,11 @@ void FRowHelper::PasteEntries(UWidgetBlueprint* WidgetBlueprint, UMVVMBlueprintV
 		FScopedTransaction Transaction(LOCTEXT("PasteBindingTransaction", "Paste Binding"));
 		BlueprintView->Modify();
 
-		const FMVVMBlueprintViewBinding& BindingToPaste = Item.Get<FMVVMBlueprintViewBinding>();
+		const FMVVMBlueprintViewBinding& BindingToPaste = ClipboardItem.Get<FMVVMBlueprintViewBinding>();
 
 		FMVVMBlueprintViewBinding* BindingSelected = Entry->GetBinding(BlueprintView);
 		BindingSelected->BindingType = BindingToPaste.BindingType;
-		BindingSelected->SourcePath = BindingToPaste.SourcePath;
+		BindingSelected->SourcePath = Private::MigratePropertyPathSource(BlueprintView, ClipboardBlueprintView, BindingToPaste.SourcePath);
 
 		// Copy over conversion functions
 		if (UMVVMBlueprintViewConversionFunction* ConversionFunction = BindingToPaste.Conversion.SourceToDestinationConversion)
@@ -663,7 +764,7 @@ void FRowHelper::PasteEntries(UWidgetBlueprint* WidgetBlueprint, UMVVMBlueprintV
 
 		BlueprintView->OnBindingsUpdated.Broadcast();
 	}
-	else if (Item.IsType<UMVVMBlueprintViewCondition*>())
+	else if (ClipboardItem.IsType<UMVVMBlueprintViewCondition*>())
 	{
 		if (Entry->GetRowType() != FBindingEntry::ERowType::Condition)
 		{
@@ -676,16 +777,20 @@ void FRowHelper::PasteEntries(UWidgetBlueprint* WidgetBlueprint, UMVVMBlueprintV
 
 		// Retain the condition path when pasting. Copy over all other properties
 		UMVVMBlueprintViewCondition* ConditionSelected = Entry->GetCondition();
-		UMVVMBlueprintViewCondition* ConditionToPaste = Item.Get<UMVVMBlueprintViewCondition*>();
+		UMVVMBlueprintViewCondition* ConditionToPaste = ClipboardItem.Get<UMVVMBlueprintViewCondition*>();
 		UMVVMBlueprintViewCondition* NewCondition = DuplicateObject<UMVVMBlueprintViewCondition>(ConditionToPaste, BlueprintView);
 
-		NewCondition->SetConditionPath(ConditionSelected->GetConditionPath());
+		FMVVMBlueprintPropertyPath ConditionPathMigrated = Private::MigratePropertyPathSource(BlueprintView, ClipboardBlueprintView, ConditionSelected->GetConditionPath());
+
+		// Set the old operation path values to the newly duplicated condition. 
+		NewCondition->SetConditionPath(ConditionPathMigrated);
 		NewCondition->SetOperation(ConditionSelected->GetOperation());
 		NewCondition->SetOperationValue(ConditionSelected->GetOperationValue());
 		NewCondition->SetOperationMaxValue(ConditionSelected->GetOperationMaxValue());
+		
 		BlueprintView->ReplaceCondition(ConditionSelected, NewCondition);
 	}
-	else if (Item.IsType<UMVVMBlueprintViewEvent*>())
+	else if (ClipboardItem.IsType<UMVVMBlueprintViewEvent*>())
 	{
 		if (Entry->GetRowType() != FBindingEntry::ERowType::Event)
 		{
@@ -698,14 +803,16 @@ void FRowHelper::PasteEntries(UWidgetBlueprint* WidgetBlueprint, UMVVMBlueprintV
 
 		// Retain the event path when pasting. Copy over all other properties
 		UMVVMBlueprintViewEvent* EventSelected = Entry->GetEvent();
-		UMVVMBlueprintViewEvent* EventToPaste = Item.Get<UMVVMBlueprintViewEvent*>();
+		UMVVMBlueprintViewEvent* EventToPaste = ClipboardItem.Get<UMVVMBlueprintViewEvent*>();
 		UMVVMBlueprintViewEvent* NewEvent = DuplicateObject<UMVVMBlueprintViewEvent>(EventToPaste, BlueprintView);
 
-		const FMVVMBlueprintPropertyPath EventPath = EventSelected->GetEventPath();
-		NewEvent->SetEventPath(EventPath);
-
+		const FMVVMBlueprintPropertyPath EventPathMigrated = Private::MigratePropertyPathSource(BlueprintView, ClipboardBlueprintView, EventSelected->GetEventPath());
+		
+		// Set the old event path values to the newly duplicated event. Rebuild to update conversion graph 
+		NewEvent->SetEventPath(EventPathMigrated);
 		NewEvent->RecreateWrapperGraph();
 		NewEvent->UpdatePinValues();
+
 		BlueprintView->ReplaceEvent(EventSelected, NewEvent);
 	}
 }
