UMG Bindings: Add new Conversion Function SetVectorParameter that takes a FColor and converts to FLinearColor

diff --git a/Engine/Plugins/Runtime/ModelViewViewModel/Source/ModelViewViewModel/Private/Bindings/ConversionLibraries/MVVMFontConversionLibrary.cpp b/Engine/Plugins/Runtime/ModelViewViewModel/Source/ModelViewViewModel/Private/Bindings/ConversionLibraries/MVVMFontConversionLibrary.cpp
--- a/Engine/Plugins/Runtime/ModelViewViewModel/Source/ModelViewViewModel/Private/Bindings/ConversionLibraries/MVVMFontConversionLibrary.cpp
+++ b/Engine/Plugins/Runtime/ModelViewViewModel/Source/ModelViewViewModel/Private/Bindings/ConversionLibraries/MVVMFontConversionLibrary.cpp
@@ -27,6 +27,17 @@ FSlateFontInfo UMVVMFontConversionLibrary::Conv_SetVectorParameter(FSlateFontInf
 	return Font;
 }
 
+FSlateFontInfo UMVVMFontConversionLibrary::Conv_SetVectorParameter_FColor(FSlateFontInfo Font, FName ParameterName, FColor Value)
+{
+	if (UMaterialInstanceDynamic* DynamicMaterial = TryGetDynamicMaterial(Font, nullptr))
+	{
+		FLinearColor Color(Value);
+		DynamicMaterial->SetVectorParameterValue(ParameterName, Color);
+	}
+
+	return Font;
+}
+
 FSlateFontInfo UMVVMFontConversionLibrary::Conv_SetTextureParameter(FSlateFontInfo Font, FName ParameterName, UTexture* Value)
 {
 	if (UMaterialInstanceDynamic* DynamicMaterial = TryGetDynamicMaterial(Font, nullptr))
@@ -57,6 +68,17 @@ FSlateFontInfo UMVVMFontConversionLibrary::Conv_SetVectorParameterMID(FSlateFont
 	return Font;
 }
 
+FSlateFontInfo UMVVMFontConversionLibrary::Conv_SetVectorParameterMID_FColor(FSlateFontInfo Font, UMaterialInterface* Material, FName ParameterName, FColor Value)
+{
+	if (UMaterialInstanceDynamic* DynamicMaterial = TryGetDynamicMaterial(Font, nullptr, Material))
+	{
+		FLinearColor Color(Value);
+		DynamicMaterial->SetVectorParameterValue(ParameterName, Color);
+	}
+
+	return Font;
+}
+
 FSlateFontInfo UMVVMFontConversionLibrary::Conv_SetTextureParameterMID(FSlateFontInfo Font, UMaterialInterface* Material, FName ParameterName, UTexture* Value)
 {
 	if (UMaterialInstanceDynamic* DynamicMaterial = TryGetDynamicMaterial(Font, nullptr, Material))
diff --git a/Engine/Plugins/Runtime/ModelViewViewModel/Source/ModelViewViewModel/Private/Bindings/ConversionLibraries/MVVMSlateBrushConversionLibrary.cpp b/Engine/Plugins/Runtime/ModelViewViewModel/Source/ModelViewViewModel/Private/Bindings/ConversionLibraries/MVVMSlateBrushConversionLibrary.cpp
--- a/Engine/Plugins/Runtime/ModelViewViewModel/Source/ModelViewViewModel/Private/Bindings/ConversionLibraries/MVVMSlateBrushConversionLibrary.cpp
+++ b/Engine/Plugins/Runtime/ModelViewViewModel/Source/ModelViewViewModel/Private/Bindings/ConversionLibraries/MVVMSlateBrushConversionLibrary.cpp
@@ -27,6 +27,17 @@ FSlateBrush UMVVMSlateBrushConversionLibrary::Conv_SetVectorParameter(FSlateBrus
 	return Brush;
 }
 
+FSlateBrush UMVVMSlateBrushConversionLibrary::Conv_SetVectorParameter_FColor(FSlateBrush Brush, FName ParameterName, FColor Value)
+{
+	if (UMaterialInstanceDynamic* DynamicMaterial = TryGetDynamicMaterial(Brush, nullptr))
+	{
+		FLinearColor Color(Value);
+		DynamicMaterial->SetVectorParameterValue(ParameterName, Color);
+	}
+
+	return Brush;
+}
+
 FSlateBrush UMVVMSlateBrushConversionLibrary::Conv_SetTextureParameter(FSlateBrush Brush, FName ParameterName, UTexture* Value)
 {
 	if (UMaterialInstanceDynamic* DynamicMaterial = TryGetDynamicMaterial(Brush, nullptr))
@@ -57,6 +68,17 @@ FSlateBrush UMVVMSlateBrushConversionLibrary::Conv_SetVectorParameterMID(FSlateB
 	return Brush;
 }
 
+FSlateBrush UMVVMSlateBrushConversionLibrary::Conv_SetVectorParameterMID_FColor(FSlateBrush Brush, UMaterialInterface* Material, FName ParameterName, FColor Value)
+{
+	if (UMaterialInstanceDynamic* DynamicMaterial = TryGetDynamicMaterial(Brush, nullptr, Material))
+	{
+		FLinearColor Color(Value);
+		DynamicMaterial->SetVectorParameterValue(ParameterName, Color);
+	}
+
+	return Brush;
+}
+
 FSlateBrush UMVVMSlateBrushConversionLibrary::Conv_SetTextureParameterMID(FSlateBrush Brush, UMaterialInterface* Material, FName ParameterName, UTexture* Value)
 {
 	if (UMaterialInstanceDynamic* DynamicMaterial = TryGetDynamicMaterial(Brush, nullptr, Material))
diff --git a/Engine/Plugins/Runtime/ModelViewViewModel/Source/ModelViewViewModel/Public/Bindings/ConversionLibraries/MVVMFontConversionLibrary.h b/Engine/Plugins/Runtime/ModelViewViewModel/Source/ModelViewViewModel/Public/Bindings/ConversionLibraries/MVVMFontConversionLibrary.h
--- a/Engine/Plugins/Runtime/ModelViewViewModel/Source/ModelViewViewModel/Public/Bindings/ConversionLibraries/MVVMFontConversionLibrary.h
+++ b/Engine/Plugins/Runtime/ModelViewViewModel/Source/ModelViewViewModel/Public/Bindings/ConversionLibraries/MVVMFontConversionLibrary.h
@@ -46,6 +46,16 @@ public:
 	 */
 	UFUNCTION(BlueprintPure, Category = "Widget", meta = (DisplayName = "Set Vector Parameter (Font)", MVVMBindToDestination = "Font"))
 	static UE_API FSlateFontInfo Conv_SetVectorParameter(FSlateFontInfo Font, FName ParameterName, FLinearColor Value);
+
+	/**
+	 * Sets a vector value on a font material assuming it exists, handles MID existance appropriately
+	 * 
+	 * @param Font			Font with a material we want to set a parameter on 
+	 * @param ParameterName	Name of material parameter to set
+	 * @param Value			Value to set material parameter to
+	 */
+	UFUNCTION(BlueprintPure, Category = "Widget", meta = (DisplayName = "Set Vector Parameter (Font) from Color", MVVMBindToDestination = "Font"))
+	static UE_API FSlateFontInfo Conv_SetVectorParameter_FColor(FSlateFontInfo Font, FName ParameterName, FColor Value);
 	
 	/**
 	 * Sets a texture value on a font material assuming it exists, handles MID existance appropriately
@@ -79,6 +89,17 @@ public:
 	UFUNCTION(BlueprintPure, Category = "Widget", meta = (DisplayName = "Set Vector Parameter MID (Font)", MVVMBindToDestination = "Font"))
 	static UE_API FSlateFontInfo Conv_SetVectorParameterMID(FSlateFontInfo Font, UMaterialInterface* Material, FName ParameterName, FLinearColor Value);
 
+	/**
+	 * Sets a vector value on a font material assuming it exists, handles MID existance appropriately
+	 *
+	 * @param Font			Font with a material we want to set a parameter on
+	 * @param Material		Material to set on the font
+	 * @param ParameterName	Name of material parameter to set
+	 * @param Value			Value to set material parameter to
+	 */
+	UFUNCTION(BlueprintPure, Category = "Widget", meta = (DisplayName = "Set Vector Parameter MID (Font) from Color", MVVMBindToDestination = "Font"))
+	static UE_API FSlateFontInfo Conv_SetVectorParameterMID_FColor(FSlateFontInfo Font, UMaterialInterface* Material, FName ParameterName, FColor Value);
+
 	/**
 	 * Sets a texture value on a font material assuming it exists, handles MID existance appropriately
 	 *
diff --git a/Engine/Plugins/Runtime/ModelViewViewModel/Source/ModelViewViewModel/Public/Bindings/ConversionLibraries/MVVMSlateBrushConversionLibrary.h b/Engine/Plugins/Runtime/ModelViewViewModel/Source/ModelViewViewModel/Public/Bindings/ConversionLibraries/MVVMSlateBrushConversionLibrary.h
--- a/Engine/Plugins/Runtime/ModelViewViewModel/Source/ModelViewViewModel/Public/Bindings/ConversionLibraries/MVVMSlateBrushConversionLibrary.h
+++ b/Engine/Plugins/Runtime/ModelViewViewModel/Source/ModelViewViewModel/Public/Bindings/ConversionLibraries/MVVMSlateBrushConversionLibrary.h
@@ -47,6 +47,17 @@ public:
 	UFUNCTION(BlueprintPure, Category = "Widget", meta = (DisplayName = "Set Vector Parameter (Brush)", MVVMBindToDestination = "Brush"))
 	static UE_API FSlateBrush Conv_SetVectorParameter(FSlateBrush Brush, FName ParameterName, FLinearColor Value);
 	
+	/**
+     * Sets a vector value on a brush material assuming it exists, handles MID existance appropriately
+     * 
+     * @param Brush			Brush with a material we want to set a parameter on 
+     * @param ParameterName	Name of material parameter to set
+     * @param Value			Value to set material parameter to
+     */
+    UFUNCTION(BlueprintPure, Category = "Widget", meta = (DisplayName = "Set Vector Parameter (Brush) from Color", MVVMBindToDestination = "Brush"))
+	static UE_API FSlateBrush Conv_SetVectorParameter_FColor(FSlateBrush Brush, FName ParameterName, FColor Value);
+    	
+    	
 	/**
 	 * Sets a texture value on a brush material assuming it exists, handles MID existance appropriately
 	 * 
@@ -78,6 +89,18 @@ public:
 	 */
 	UFUNCTION(BlueprintPure, Category = "Widget", meta = (DisplayName = "Set Vector Parameter MID (Brush)", MVVMBindToDestination = "Brush"))
 	static UE_API FSlateBrush Conv_SetVectorParameterMID(FSlateBrush Brush, UMaterialInterface* Material, FName ParameterName, FLinearColor Value);
+	
+	/**
+	* Sets a vector value on a brush material assuming it exists, handles MID existance appropriately
+	*
+	* @param Brush			Brush with a material we want to set a parameter on
+	* @param Material		Material to set on the brush
+	* @param ParameterName	Name of material parameter to set
+	* @param Value			Value to set material parameter to
+	*/
+	UFUNCTION(BlueprintPure, Category = "Widget", meta = (DisplayName = "Set Vector Parameter MID (Brush) from Color", MVVMBindToDestination = "Brush"))
+	static UE_API FSlateBrush Conv_SetVectorParameterMID_FColor(FSlateBrush Brush, UMaterialInterface* Material, FName ParameterName, FColor Value);
+
 
 	/**
 	 * Sets a texture value on a brush material assuming it exists, handles MID existance appropriately
