MVVM: fix widget renaming breaks component bindings

diff --git a/Engine/Plugins/Runtime/ModelViewViewModel/Source/ModelViewViewModelBlueprint/Private/MVVMPropertyPath.cpp b/Engine/Plugins/Runtime/ModelViewViewModel/Source/ModelViewViewModelBlueprint/Private/MVVMPropertyPath.cpp
--- a/Engine/Plugins/Runtime/ModelViewViewModel/Source/ModelViewViewModelBlueprint/Private/MVVMPropertyPath.cpp
+++ b/Engine/Plugins/Runtime/ModelViewViewModel/Source/ModelViewViewModelBlueprint/Private/MVVMPropertyPath.cpp
@@ -570,7 +570,7 @@ bool FMVVMBlueprintPropertyPath::OnFieldRenamed(const UBlueprint* InContext, UCl
 		// Check the path components for the renamed field
 		for (FMVVMBlueprintFieldPath& Path : Paths)
 		{
-			if (bFoundRenamedField)
+			if (bFoundRenamedField && !bIsComponent)
 			{
 				if (!Path.IsFieldSelfContext() && OwnerStruct)
 				{
@@ -604,6 +604,32 @@ bool FMVVMBlueprintPropertyPath::OnFieldRenamed(const UBlueprint* InContext, UCl
 				Path.SetMemberName(NewObjectName);
 				bFoundRenamedField = true;
 			}
+
+			if (bIsComponent)
+			{
+				const UWidgetBlueprintGeneratedClass* WidgetBPG = Cast<const UWidgetBlueprintGeneratedClass>(OwnerStruct);
+				const UWidgetBlueprint* WidgetBP = WidgetBPG ? Cast<UWidgetBlueprint>(WidgetBPG->ClassGeneratedBy) : Cast<UWidgetBlueprint>(InContext);
+				if (!WidgetBPG)
+				{
+					WidgetBPG = Cast<const UWidgetBlueprintGeneratedClass>(SelfContext);
+				}
+				const FProperty* ComponentProperty = WidgetBPG->FindPropertyByName(Path.GetFieldName(SelfContext));
+				if (const FObjectPropertyBase* ComponentObjectProperty = CastField<FObjectPropertyBase>(ComponentProperty))
+				{
+					if (ComponentObjectProperty->PropertyClass->IsChildOf(UUIComponent::StaticClass()))
+					{ 
+						if (UUIComponentWidgetBlueprintExtension* ComponentWidgetBlueprintExtension = UWidgetBlueprintExtension::GetExtension<UUIComponentWidgetBlueprintExtension>(WidgetBP))
+						{
+							if (UUIComponent* FoundComponent = ComponentWidgetBlueprintExtension->GetComponent(ComponentObjectProperty->PropertyClass, NewObjectName))
+							{
+								const FName ComponentPropertyName = UUIComponentContainer::GetPropertyNameForComponent(FoundComponent, NewObjectName);
+								Path.SetMemberName(ComponentPropertyName);
+								break;
+							}
+						}
+					}
+				}
+			}
 		}
 	}
 
diff --git a/Engine/Plugins/Runtime/ModelViewViewModel/Source/ModelViewViewModelEditor/Private/Widgets/SMVVMCachedViewBindingPropertyPath.cpp b/Engine/Plugins/Runtime/ModelViewViewModel/Source/ModelViewViewModelEditor/Private/Widgets/SMVVMCachedViewBindingPropertyPath.cpp
--- a/Engine/Plugins/Runtime/ModelViewViewModel/Source/ModelViewViewModelEditor/Private/Widgets/SMVVMCachedViewBindingPropertyPath.cpp
+++ b/Engine/Plugins/Runtime/ModelViewViewModel/Source/ModelViewViewModelEditor/Private/Widgets/SMVVMCachedViewBindingPropertyPath.cpp
@@ -40,8 +40,8 @@ void SCachedViewBindingPropertyPath::Tick(const FGeometry& AllottedGeometry, con
 	FMVVMBlueprintPropertyPath NewPropertyPath = OnGetPropertyPath.Execute();
 	if (NewPropertyPath != CachedPropertyPath)
 	{
+		CachedPropertyPath = NewPropertyPath;
 		PropertyPathWidget->SetPropertyPath(NewPropertyPath);
-		CachedPropertyPath = MoveTemp(NewPropertyPath);
 	}
 }
 
diff --git a/Engine/Plugins/Runtime/ModelViewViewModel/Source/ModelViewViewModelEditor/Private/Widgets/SMVVMViewModelBindingListWidget.cpp b/Engine/Plugins/Runtime/ModelViewViewModel/Source/ModelViewViewModelEditor/Private/Widgets/SMVVMViewModelBindingListWidget.cpp
--- a/Engine/Plugins/Runtime/ModelViewViewModel/Source/ModelViewViewModelEditor/Private/Widgets/SMVVMViewModelBindingListWidget.cpp
+++ b/Engine/Plugins/Runtime/ModelViewViewModel/Source/ModelViewViewModelEditor/Private/Widgets/SMVVMViewModelBindingListWidget.cpp
@@ -729,11 +729,12 @@ FMVVMBlueprintPropertyPath SSourceBindingList::CreateBlueprintPropertyPath(SProp
 							if (Source->Key.GetClass() == Component->GetClass())
 							{
 								FName ComponentPropertyName = UUIComponentContainer::GetPropertyNameForComponent(Component, Widget->GetFName());
-								FProperty* Property = FindFProperty<FObjectPropertyBase>(Cast<UWidgetBlueprintGeneratedClass>(WidgetBlueprint->GeneratedClass), ComponentPropertyName);
-								ensure(Property);
-								FMVVMConstFieldVariant Field = FMVVMConstFieldVariant(Property);
-								PropertyPath.AppendPropertyPath(WidgetBlueprintPtr, Field);
-								CompleteFieldPath.Add(FFieldVariant(Property));
+								if (FProperty* Property = FindFProperty<FObjectPropertyBase>(Cast<UWidgetBlueprintGeneratedClass>(WidgetBlueprint->GeneratedClass), ComponentPropertyName))
+								{
+									FMVVMConstFieldVariant Field = FMVVMConstFieldVariant(Property);
+									PropertyPath.AppendPropertyPath(WidgetBlueprintPtr, Field);
+									CompleteFieldPath.Add(FFieldVariant(Property));
+								}
 								break;
 							}
 						}
