[MVVM] Add utility for setting VM on a widget directly.

#tests Replaced instance of setting via extension with this to see it still working

diff --git a/Engine/Plugins/Runtime/ModelViewViewModel/Source/ModelViewViewModel/Private/MVVMBlueprintLibrary.cpp b/Engine/Plugins/Runtime/ModelViewViewModel/Source/ModelViewViewModel/Private/MVVMBlueprintLibrary.cpp
new file mode 100644
--- /dev/null
+++ b/Engine/Plugins/Runtime/ModelViewViewModel/Source/ModelViewViewModel/Private/MVVMBlueprintLibrary.cpp
@@ -0,0 +1,32 @@
+// Copyright Epic Games, Inc. All Rights Reserved.
+
+#include "MVVMBlueprintLibrary.h"
+
+#include "MVVMMessageLog.h"
+#include "View/MVVMView.h"
+
+#include UE_INLINE_GENERATED_CPP_BY_NAME(MVVMBlueprintLibrary)
+
+#define LOCTEXT_NAMESPACE "MVVMBlueprintLibrary"
+
+void UMVVMBlueprintLibrary::SetViewModelByClass(UUserWidget*& Widget, TScriptInterface<INotifyFieldValueChanged> ViewModel)
+{
+	if (UMVVMView* View = Widget ? Widget->GetExtension<UMVVMView>() : nullptr)
+	{
+		View->SetViewModelByClass(ViewModel);
+	}
+	else
+	{
+		UE::MVVM::FMessageLog Log(Widget); // Null widget is ok for log
+		if (Widget)
+		{
+			Log.Error(LOCTEXT("MissingView", "Widget does not have an MVVMView, cannot set any View Model"));
+		}
+		else
+		{
+			Log.Error(LOCTEXT("MissingWidget", "No widget passed into this function"));
+		}
+	}
+}
+
+#undef LOCTEXT_NAMESPACE
diff --git a/Engine/Plugins/Runtime/ModelViewViewModel/Source/ModelViewViewModel/Public/MVVMBlueprintLibrary.h b/Engine/Plugins/Runtime/ModelViewViewModel/Source/ModelViewViewModel/Public/MVVMBlueprintLibrary.h
new file mode 100644
--- /dev/null
+++ b/Engine/Plugins/Runtime/ModelViewViewModel/Source/ModelViewViewModel/Public/MVVMBlueprintLibrary.h
@@ -0,0 +1,16 @@
+// Copyright Epic Games, Inc. All Rights Reserved.
+
+#pragma once
+
+#include "Kismet/BlueprintFunctionLibrary.h"
+
+#include "MVVMBlueprintLibrary.generated.h"
+
+UCLASS()
+class UMVVMBlueprintLibrary : public UBlueprintFunctionLibrary
+{
+	GENERATED_BODY()
+public:
+	UFUNCTION(BlueprintCallable, Category="MVVM | Utilities")
+	static void SetViewModelByClass(UPARAM(ref) UUserWidget*& Widget, TScriptInterface<INotifyFieldValueChanged> ViewModel);
+};
\ No newline at end of file
