MVVM: Record conversion function parameter changes to transaction buffer

#rnx

diff --git a/Engine/Plugins/Runtime/ModelViewViewModel/Source/ModelViewViewModelEditor/Private/MVVMEditorSubsystem.cpp b/Engine/Plugins/Runtime/ModelViewViewModel/Source/ModelViewViewModelEditor/Private/MVVMEditorSubsystem.cpp
--- a/Engine/Plugins/Runtime/ModelViewViewModel/Source/ModelViewViewModelEditor/Private/MVVMEditorSubsystem.cpp
+++ b/Engine/Plugins/Runtime/ModelViewViewModel/Source/ModelViewViewModelEditor/Private/MVVMEditorSubsystem.cpp
@@ -1412,6 +1412,9 @@ void UMVVMEditorSubsystem::SetPathForConversionFunctionArgument(UWidgetBlueprint
 {
 	if (UMVVMBlueprintView* View = GetView(WidgetBlueprint))
 	{
+		const FText ParameterName = (!ParameterId.GetNames().IsEmpty()) ? FText::FromName(ParameterId.GetNames().Last()) : LOCTEXT("InvalidPin", "<Invalid Pin>");
+		const FScopedTransaction Transaction(FText::Format(LOCTEXT("SetConversionFunctionParameter", "Changed conversion function parameter {0}"), ParameterName));
+
 		FName NAME_Conversion = GET_MEMBER_NAME_CHECKED(FMVVMBlueprintViewBinding, Conversion);
 		UE::MVVM::Private::OnBindingPreEditChange(View, NAME_Conversion);
 
