- Added a 2 values filtering for entries in the view bindings panels:  view model or verse fields.
Binding with no view model/verse field reference, as well as binding referencing both, will be displayed whatever the filter value is.
- Fixed the text search filtering of the viewbinding view that could display empty results while there was something to display.

#tests Created many binding with different setup (binding, events, conditions), with or without conversion functions, using either a view model or verse fields. Validated in all those combinations that only relevant binding where displayed in each mode. Validated in Vanilla UE that this filter was not shown. Tested text filtering.
#rnx

diff --git a/Engine/Plugins/Runtime/ModelViewViewModel/Source/ModelViewViewModel/Public/View/MVVMView.h b/Engine/Plugins/Runtime/ModelViewViewModel/Source/ModelViewViewModel/Public/View/MVVMView.h
--- a/Engine/Plugins/Runtime/ModelViewViewModel/Source/ModelViewViewModel/Public/View/MVVMView.h
+++ b/Engine/Plugins/Runtime/ModelViewViewModel/Source/ModelViewViewModel/Public/View/MVVMView.h
@@ -62,7 +62,7 @@ struct FMVVMView_Source
 };
 
 /**
- * Instance UMVVMClassExtension_View for the UUserWdiget
+ * Instance UMVVMClassExtension_View for the UUserWidget
  */
 UCLASS(MinimalAPI, Transient, DisplayName="MVVM View")
 class UMVVMView : public UUserWidgetExtension
diff --git a/Engine/Plugins/Runtime/ModelViewViewModel/Source/ModelViewViewModelEditor/Private/Styling/MVVMEditorStyle.cpp b/Engine/Plugins/Runtime/ModelViewViewModel/Source/ModelViewViewModelEditor/Private/Styling/MVVMEditorStyle.cpp
--- a/Engine/Plugins/Runtime/ModelViewViewModel/Source/ModelViewViewModelEditor/Private/Styling/MVVMEditorStyle.cpp
+++ b/Engine/Plugins/Runtime/ModelViewViewModel/Source/ModelViewViewModelEditor/Private/Styling/MVVMEditorStyle.cpp
@@ -64,6 +64,7 @@ FMVVMEditorStyle::FMVVMEditorStyle()
 		.SetEvenRowBackgroundHoveredBrush(*FAppStyle::Get().GetBrush("Brushes.Header"))
 		.SetOddRowBackgroundHoveredBrush(*FAppStyle::Get().GetBrush("Brushes.Header"))
 	);
+	Set("BindingView.VerseBinding", new IMAGE_BRUSH("Slate/VerseBinding_16", Icon16x16));
 
 	FButtonStyle NoStyleComboButtonButtonStyle = FButtonStyle()
 		.SetNormal(FSlateBoxBrush(RootToCoreContentDir("Editor/Slate/Common/ButtonHoverHint.png"), FMargin(4.0f / 16.0f), FLinearColor(1.0f, 1.0f, 1.0f, 0.15f)))
diff --git a/Engine/Plugins/Runtime/ModelViewViewModel/Source/ModelViewViewModelEditor/Private/Types/MVVMBindingEntry.cpp b/Engine/Plugins/Runtime/ModelViewViewModel/Source/ModelViewViewModelEditor/Private/Types/MVVMBindingEntry.cpp
--- a/Engine/Plugins/Runtime/ModelViewViewModel/Source/ModelViewViewModelEditor/Private/Types/MVVMBindingEntry.cpp
+++ b/Engine/Plugins/Runtime/ModelViewViewModel/Source/ModelViewViewModelEditor/Private/Types/MVVMBindingEntry.cpp
@@ -122,6 +122,12 @@ void FBindingEntry::ResetChildren()
 	bUseFilteredChildren = false;
 }
 
+void FBindingEntry::ResetFilteredChildren()
+{
+	FilteredChildren.Reset();
+	bUseFilteredChildren = false;
+}
+
 void FBindingEntry::SetUseFilteredChildList()
 {
 	bUseFilteredChildren = true;
diff --git a/Engine/Plugins/Runtime/ModelViewViewModel/Source/ModelViewViewModelEditor/Private/Types/MVVMBindingEntry.h b/Engine/Plugins/Runtime/ModelViewViewModel/Source/ModelViewViewModelEditor/Private/Types/MVVMBindingEntry.h
--- a/Engine/Plugins/Runtime/ModelViewViewModel/Source/ModelViewViewModelEditor/Private/Types/MVVMBindingEntry.h
+++ b/Engine/Plugins/Runtime/ModelViewViewModel/Source/ModelViewViewModelEditor/Private/Types/MVVMBindingEntry.h
@@ -70,6 +70,11 @@ struct FBindingEntry
 
 	void SetBindingId(FGuid Id);
 
+	void SetHasPathToVerseVariable(bool bInHasPathToVerseVariable) { bHasPathToVerseVariable = bInHasPathToVerseVariable; }
+	void SetHasPathToViewModel(bool bInHasPathToViewModel) { bHasPathToViewModel = bInHasPathToViewModel; }
+	bool GetHasPathToVerseVariable() const { return bHasPathToVerseVariable; }
+	bool GetHasPathToViewModel() const { return bHasPathToViewModel; }
+
 	//~ binding parameter
 	const FMVVMBlueprintPinId& GetBindingParameterId() const
 	{
@@ -136,6 +141,7 @@ struct FBindingEntry
 	void AddFilteredChild(TSharedPtr<FBindingEntry> Child);
 
 	void ResetChildren();
+	void ResetFilteredChildren();
 
 	void SetUseFilteredChildList();
 
@@ -155,6 +161,8 @@ private:
 	TArray<TSharedPtr<FBindingEntry>> FilteredChildren;
 	bool bGroupIsWidget = false;
 	bool bUseFilteredChildren = false;
+	bool bHasPathToVerseVariable = false;
+	bool bHasPathToViewModel = false;
 };
 
 } // namespace
diff --git a/Engine/Plugins/Runtime/ModelViewViewModel/Source/ModelViewViewModelEditor/Private/Widgets/SMVVMViewBindingListView.cpp b/Engine/Plugins/Runtime/ModelViewViewModel/Source/ModelViewViewModelEditor/Private/Widgets/SMVVMViewBindingListView.cpp
--- a/Engine/Plugins/Runtime/ModelViewViewModel/Source/ModelViewViewModelEditor/Private/Widgets/SMVVMViewBindingListView.cpp
+++ b/Engine/Plugins/Runtime/ModelViewViewModel/Source/ModelViewViewModelEditor/Private/Widgets/SMVVMViewBindingListView.cpp
@@ -126,8 +126,48 @@ namespace Private
 		return TSharedPtr<FBindingEntry>();
 	}
 
-	void FilterEntryList(FString FilterString, const TArray<TSharedPtr<FBindingEntry>>& RootGroups, TArray<TSharedPtr<FBindingEntry>>& FilteredRootGroups, UMVVMBlueprintView* BlueprintView, UMVVMWidgetBlueprintExtension_View* MVVMExtensionPtr)
+	void FilterEntryList(FString FilterString, TOptional<SBindingsList::EFilterMode> FilterMode, const TArray<TSharedPtr<FBindingEntry>>& RootGroups, TArray<TSharedPtr<FBindingEntry>>& FilteredRootGroups, UMVVMBlueprintView* BlueprintView, UMVVMWidgetBlueprintExtension_View* MVVMExtensionPtr)
 	{
+		// First filtering pass: filter based on the filter mode.
+		if (FilterMode.IsSet())
+		{
+			for (const TSharedPtr<FBindingEntry>& GroupEntry : RootGroups)
+			{
+				for (TSharedPtr<FBindingEntry> Entry : GroupEntry->GetAllChildren())
+				{
+					FBindingEntry* EntryPtr = Entry.Get();
+
+					const bool bEntryHasLink = EntryPtr->GetHasPathToVerseVariable() || EntryPtr->GetHasPathToViewModel();
+					const bool bEntryHasMatchingLink = (FilterMode.GetValue() == SBindingsList::EFilterMode::Viewmodel && EntryPtr->GetHasPathToViewModel()) ||
+													   (FilterMode.GetValue() == SBindingsList::EFilterMode::Verse && EntryPtr->GetHasPathToVerseVariable());
+
+					bool bChildHasLink = false;
+					bool bChildHasMatchingLink = false;
+					for (TSharedPtr<FBindingEntry> ChildEntry : EntryPtr->GetAllChildren())
+					{
+						bChildHasLink |= ChildEntry->GetHasPathToVerseVariable() || ChildEntry->GetHasPathToViewModel();
+
+						bChildHasMatchingLink |= (FilterMode.GetValue() == SBindingsList::EFilterMode::Viewmodel && ChildEntry->GetHasPathToViewModel()) ||
+												 (FilterMode.GetValue() == SBindingsList::EFilterMode::Verse && ChildEntry->GetHasPathToVerseVariable());
+					}
+
+					// Keep entries that are not linked to any field, or if they are linked, keep any that have at least one link we want.
+					// An entry that is not tied to a viewmodel or a verse field will always be kept, and an entry that is linked to both will also always be kept, so they will be displayed whatever the filter mode is.
+					const bool bHasLink = bEntryHasLink || bChildHasLink;
+					if (!bHasLink || bEntryHasMatchingLink || bChildHasMatchingLink)
+					{
+						GroupEntry->AddFilteredChild(Entry);
+					}
+					else
+					{
+						GroupEntry->SetUseFilteredChildList();
+					}
+				}
+			}
+
+		}
+
+		// Second filtering pass: filter based on the filter string.
 		if (!FilterString.TrimStartAndEnd().IsEmpty())
 		{
 			TArray<FString> SearchKeywords;
@@ -156,22 +196,26 @@ namespace Private
 					, WidgetBlueprint(InWidgetBlueprint)
 					, BlueprintView(InBlueprintView)
 				{}
-				void operator()(TSharedPtr<FBindingEntry> ParentEntry)
+				bool operator()(TSharedPtr<FBindingEntry> ParentEntry)
 				{
-					for (TSharedPtr<FBindingEntry> Entry : ParentEntry->GetAllChildren())
+					bool bMatchFound = false;
+
+					//Work on a copy of the array, as following code will modify it and we want a second pass of filtering, that will remove children, not add ones.
+ 					TConstArrayView<TSharedPtr<FBindingEntry>> FilteredChildren = ParentEntry->GetFilteredChildren();
+ 					ParentEntry->ResetFilteredChildren();
+
+ 					for (TSharedPtr<FBindingEntry> Entry : FilteredChildren)
 					{
-						FString EntryString = Entry->GetSearchNameString(BlueprintView, WidgetBlueprint);
-						if (IsAllKeywordsInString(EntryString, SearchKeywords))
+						const FString EntryString = Entry->GetSearchNameString(BlueprintView, WidgetBlueprint);
+						const bool bEntryMatchFound = IsAllKeywordsInString(EntryString, SearchKeywords) || (*this)(Entry);
+						if (bEntryMatchFound)
 						{
-							// If the filter text is found in the group name, we keep the entire group.
 							ParentEntry->AddFilteredChild(Entry);
-						}
-						else
-						{
-							(*this)(Entry);
+							bMatchFound = true;
 						}
 					}
 					ParentEntry->SetUseFilteredChildList();
+					return bMatchFound;
 				}
 				FIsAllKeywordsInString& IsAllKeywordsInString;
 				TArray<FString>& SearchKeywords;
@@ -191,8 +235,7 @@ namespace Private
 				}
 				else
 				{
-					AddFilteredEntry(GroupEntry);
-					if (GroupEntry->GetFilteredChildren().Num() > 0)
+					if (AddFilteredEntry(GroupEntry))
 					{
 						FilteredRootGroups.Add(GroupEntry);
 					}
@@ -393,6 +436,45 @@ void SBindingsList::Refresh()
 			return GroupEntry;
 		};
 
+		auto IsVerseFieldVariable = [&WidgetBlueprint](const FMVVMBlueprintPropertyPath& Path) -> bool
+		{
+			if (Path.GetSource(WidgetBlueprint) != EMVVMBlueprintFieldPathSource::SelfContext)
+			{
+				return false;
+			}
+
+			TArray<UE::MVVM::FMVVMConstFieldVariant> Fields = Path.GetCompleteFields(WidgetBlueprint);
+			if (Fields.Num() != 1 || !Fields[0].IsProperty())
+			{
+				return false;
+			}
+
+			const struct FBPVariableDescription* FoundVariable = WidgetBlueprint->NewVariables.FindByPredicate(
+				[Field = Fields[0]](const struct FBPVariableDescription& NewVariable)
+				{
+					return NewVariable.VarName == Field.GetName();
+				}
+			);
+			if (!FoundVariable)
+			{
+				return false;
+			}
+
+			return FoundVariable->HasMetaData("VerseVariable");
+		};
+
+		auto UpdateEntryPathInfo = [&WidgetBlueprint, IsVerseFieldVariable](TSharedPtr<FBindingEntry> BindingEntry, const FMVVMBlueprintPropertyPath& Path)
+		{
+			BindingEntry->SetHasPathToVerseVariable(IsVerseFieldVariable(Path));
+			BindingEntry->SetHasPathToViewModel(Path.GetSource(WidgetBlueprint) == EMVVMBlueprintFieldPathSource::ViewModel);
+		};
+
+		auto UpdateEntryPathsInfo = [&WidgetBlueprint, IsVerseFieldVariable](TSharedPtr<FBindingEntry>  BindingEntry, const FMVVMBlueprintPropertyPath& Path1, const FMVVMBlueprintPropertyPath& Path2)
+		{
+			BindingEntry->SetHasPathToVerseVariable(IsVerseFieldVariable(Path1) || IsVerseFieldVariable(Path2));
+			BindingEntry->SetHasPathToViewModel(Path1.GetSource(WidgetBlueprint) == EMVVMBlueprintFieldPathSource::ViewModel || Path2.GetSource(WidgetBlueprint) == EMVVMBlueprintFieldPathSource::ViewModel);
+		};
+			
 		for (const FMVVMBlueprintViewBinding& Binding : BlueprintView->GetBindings())
 		{
 			// Make sure the graph for the bindings is generated
@@ -448,6 +530,8 @@ void SBindingsList::Refresh()
 
 					NewEntries.Add(BindingEntry);
 				}
+
+				UpdateEntryPathsInfo(BindingEntry, Binding.DestinationPath, Binding.SourcePath);
 				GroupEntry->AddChild(BindingEntry);
 			}
 
@@ -490,6 +574,8 @@ void SBindingsList::Refresh()
 
 						NewEntries.Add(ArgumentEntry);
 					}
+
+					UpdateEntryPathInfo(ArgumentEntry, Pin.GetPath());
 					BindingEntry->AddChild(ArgumentEntry);
 				}
 			}
@@ -544,6 +630,8 @@ void SBindingsList::Refresh()
 
 					NewEntries.Add(EventEntry);
 				}
+
+				UpdateEntryPathsInfo(EventEntry, Event->GetEventPath(), Event->GetDestinationPath());
 				GroupEntry->AddChild(EventEntry);
 			}
 
@@ -581,6 +669,8 @@ void SBindingsList::Refresh()
 
 					NewEntries.Add(ArgumentEntry);
 				}
+
+				UpdateEntryPathInfo(ArgumentEntry, Pin.GetPath());
 				EventEntry->AddChild(ArgumentEntry);
 			}
 		}
@@ -617,6 +707,8 @@ void SBindingsList::Refresh()
 
 					NewEntries.Add(ConditionEntry);
 				}
+
+				UpdateEntryPathsInfo(ConditionEntry, Condition->GetConditionPath(), Condition->GetDestinationPath());
 				GroupEntry->AddChild(ConditionEntry);
 			}
 
@@ -654,11 +746,21 @@ void SBindingsList::Refresh()
 
 					NewEntries.Add(ArgumentEntry);
 				}
+
+				UpdateEntryPathInfo(ArgumentEntry, Pin.GetPath());
 				ConditionEntry->AddChild(ArgumentEntry);
 			}
 		}
 
-		Private::FilterEntryList(FilterText.ToString(), AllRootGroups, FilteredRootGroups, BlueprintView, MVVMExtensionPtr);
+		TOptional<SBindingsList::EFilterMode> FilterModeToUse;
+		if (TSharedPtr<SBindingsPanel> BindingPanelPtr = BindingPanel.Pin())
+		{
+			if (BindingPanelPtr->IsFilterModeEnabled())
+			{
+				FilterModeToUse = FilterMode;
+			}
+		}
+		Private::FilterEntryList(FilterText.ToString(), FilterModeToUse, AllRootGroups, FilteredRootGroups, BlueprintView, MVVMExtensionPtr);
 	}
 
 	if (TreeView.IsValid())
@@ -744,6 +846,17 @@ TSharedRef<ITableRow> SBindingsList::GenerateEntryRow(TSharedPtr<FBindingEntry>
 	return SNew(STableRow<TSharedPtr<FBindingEntry>>, OwnerTable);
 }
 
+SBindingsList::EFilterMode SBindingsList::GetActiveFilterMode() const
+{
+	return FilterMode;
+}
+
+void SBindingsList::OnActiveFilterModeChanged(EFilterMode InFilterMode)
+{
+	FilterMode = InFilterMode;
+	Refresh();
+}
+
 void SBindingsList::OnFilterTextChanged(const FText& InFilterText)
 {
 	FilterText = InFilterText;
diff --git a/Engine/Plugins/Runtime/ModelViewViewModel/Source/ModelViewViewModelEditor/Private/Widgets/SMVVMViewBindingListView.h b/Engine/Plugins/Runtime/ModelViewViewModel/Source/ModelViewViewModelEditor/Private/Widgets/SMVVMViewBindingListView.h
--- a/Engine/Plugins/Runtime/ModelViewViewModel/Source/ModelViewViewModelEditor/Private/Widgets/SMVVMViewBindingListView.h
+++ b/Engine/Plugins/Runtime/ModelViewViewModel/Source/ModelViewViewModelEditor/Private/Widgets/SMVVMViewBindingListView.h
@@ -57,6 +57,15 @@ public:
 	void SelectCondition(const UMVVMBlueprintViewCondition* Condition, bool bShouldScrollIntoView);
 	void SelectParameter(const FMVVMBlueprintPinId& PinId, bool bShouldScrollIntoView);
 
+	enum class EFilterMode : uint8
+	{
+		Viewmodel = 0, // Display only bindings associated to a viewmodel
+		Verse = 1	// Display only bindings associated to a verse field
+	};
+
+	EFilterMode GetActiveFilterMode() const;
+	void OnActiveFilterModeChanged(EFilterMode InFilterMode);
+
 private:
 	TSharedRef<ITableRow> GenerateEntryRow(TSharedPtr<FBindingEntry> Entry, const TSharedRef<STableViewBase>& OwnerTable);
 	void OnSourceListSelectionChanged(TSharedPtr<FBindingEntry> Entry, ESelectInfo::Type SelectionType) const;
@@ -79,6 +88,7 @@ private:
 	TWeakPtr<FWidgetBlueprintEditor> WeakBlueprintEditor;
 	mutable bool bSelectionChangedGuard = false;
 	FText FilterText;
+	EFilterMode FilterMode = EFilterMode::Verse;
 };
 
 } // namespace UE::MVVM
diff --git a/Engine/Plugins/Runtime/ModelViewViewModel/Source/ModelViewViewModelEditor/Private/Widgets/SMVVMViewBindingPanel.cpp b/Engine/Plugins/Runtime/ModelViewViewModel/Source/ModelViewViewModelEditor/Private/Widgets/SMVVMViewBindingPanel.cpp
--- a/Engine/Plugins/Runtime/ModelViewViewModel/Source/ModelViewViewModelEditor/Private/Widgets/SMVVMViewBindingPanel.cpp
+++ b/Engine/Plugins/Runtime/ModelViewViewModel/Source/ModelViewViewModelEditor/Private/Widgets/SMVVMViewBindingPanel.cpp
@@ -35,6 +35,7 @@
 #include "Widgets/Input/SCheckBox.h"
 #include "Widgets/Input/SComboButton.h"
 #include "Widgets/Input/SSearchBox.h"
+#include "Widgets/Input/SSegmentedControl.h"
 #include "Widgets/Images/SImage.h"
 #include "Widgets/Layout/SSpacer.h"
 #include "Widgets/Layout/SSplitter.h"
@@ -46,6 +47,9 @@
 
 #define LOCTEXT_NAMESPACE "BindingPanel"
 
+bool CVar_MVVMViewBindingFilterMode_Enable = false;
+FAutoConsoleVariableRef CVar_MVVMViewBindingFilterMode(TEXT("MVVM.ViewBinding.FilterMode.Enable"), CVar_MVVMViewBindingFilterMode_Enable, TEXT(""));
+
 namespace UE::MVVM
 {
 
@@ -266,6 +270,11 @@ void SBindingsPanel::OnBindingListSelectionChanged(TConstArrayView<FBindingsSele
 	}
 }
 
+bool SBindingsPanel::IsFilterModeEnabled() const
+{
+	return CVar_MVVMViewBindingFilterMode_Enable;
+}
+
 void SBindingsPanel::SetDetailsView(TUnion<UObject*, TSharedPtr<FStructOnScope>> ObjectOrStruct, const FText& ErrorMessage)
 {
 	if (!ErrorMessage.IsEmpty())
@@ -944,6 +953,27 @@ TSharedRef<SWidget> SBindingsPanel::GenerateEditViewWidget()
 			SearchTextWidget
 		];
 
+	TSharedRef<SWidget> FilterWidget = SNullWidget::NullWidget;
+	if (BindingsList && IsFilterModeEnabled())
+	{
+		FilterWidget = SNew(SSegmentedControl<SBindingsList::EFilterMode>)
+						.Value(BindingsList.ToSharedRef(), &SBindingsList::GetActiveFilterMode)
+						.OnValueChanged(BindingsList.ToSharedRef(), &SBindingsList::OnActiveFilterModeChanged)
+						+ SSegmentedControl<SBindingsList::EFilterMode>::Slot(SBindingsList::EFilterMode::Viewmodel)
+						.Icon(FMVVMEditorStyle::Get().GetBrush("BlueprintView.TabIcon"))
+						.ToolTip(LOCTEXT("BindingPanelViewModelFilterModeToolTip", "Viewmodel bindings"))
+						+ SSegmentedControl<SBindingsList::EFilterMode>::Slot(SBindingsList::EFilterMode::Verse)
+						.Icon(FMVVMEditorStyle::Get().GetBrush("BindingView.VerseBinding"))
+						.ToolTip(LOCTEXT("BindingPanelVerseFilterModeToolTip", "Verse bindings"));
+	}
+
+	BindingPanelToolBar->AddSlot()
+		.AutoWidth()
+		.VAlign(VAlign_Center)
+		[
+			FilterWidget
+		];
+
 	// Reset the toolbar builder and insert widgets to the right of the search bar
 	ToolbarBuilderGlobal = FSlimHorizontalToolBarBuilder(TSharedPtr<const FUICommandList>(), FMultiBoxCustomization::None);
 	ToolbarBuilderGlobal.AddWidget(SNew(SSpacer), NAME_None, true, HAlign_Right);
diff --git a/Engine/Plugins/Runtime/ModelViewViewModel/Source/ModelViewViewModelEditor/Private/Widgets/SMVVMViewBindingPanel.h b/Engine/Plugins/Runtime/ModelViewViewModel/Source/ModelViewViewModelEditor/Private/Widgets/SMVVMViewBindingPanel.h
--- a/Engine/Plugins/Runtime/ModelViewViewModel/Source/ModelViewViewModelEditor/Private/Widgets/SMVVMViewBindingPanel.h
+++ b/Engine/Plugins/Runtime/ModelViewViewModel/Source/ModelViewViewModelEditor/Private/Widgets/SMVVMViewBindingPanel.h
@@ -47,6 +47,7 @@ public:
 	//~ End SWidget Interface
 
 	void OnBindingListSelectionChanged(TConstArrayView<FBindingsSelectionVariantType> Selection);
+	bool IsFilterModeEnabled() const;
 
 	static void RegisterSettingsMenu();
 
