Fix small code error from my last submit in SMVVMViewBindingListView.cpp (doesn't compile for some people with older toolchain).

#tests still compile, and works
#rnx

diff --git a/Engine/Plugins/Runtime/ModelViewViewModel/Source/ModelViewViewModelEditor/Private/Widgets/SMVVMViewBindingListView.cpp b/Engine/Plugins/Runtime/ModelViewViewModel/Source/ModelViewViewModelEditor/Private/Widgets/SMVVMViewBindingListView.cpp
--- a/Engine/Plugins/Runtime/ModelViewViewModel/Source/ModelViewViewModelEditor/Private/Widgets/SMVVMViewBindingListView.cpp
+++ b/Engine/Plugins/Runtime/ModelViewViewModel/Source/ModelViewViewModelEditor/Private/Widgets/SMVVMViewBindingListView.cpp
@@ -449,8 +449,8 @@ void SBindingsList::Refresh()
 				return false;
 			}
 
-			const struct FBPVariableDescription* FoundVariable = WidgetBlueprint->NewVariables.FindByPredicate(
-				[Field = Fields[0]](const struct FBPVariableDescription& NewVariable)
+			const FBPVariableDescription* FoundVariable = WidgetBlueprint->NewVariables.FindByPredicate(
+				[Field = Fields[0]](const FBPVariableDescription& NewVariable)
 				{
 					return NewVariable.VarName == Field.GetName();
 				}
