Viewbinding view:
-Add "All" filter. This is now the default.
-remove empty widget group entries when using filtering mode, to be consistent with existing behavior of the view, and to avoid accidentaly removing a group that has a child in another filter.

#rnx
#tests Played with the 3 states of the filter. Validated that a group entry without child entry is no longer displayed.

diff --git a/Engine/Plugins/Runtime/ModelViewViewModel/Source/ModelViewViewModelEditor/Private/Widgets/SMVVMViewBindingListView.cpp b/Engine/Plugins/Runtime/ModelViewViewModel/Source/ModelViewViewModelEditor/Private/Widgets/SMVVMViewBindingListView.cpp
--- a/Engine/Plugins/Runtime/ModelViewViewModel/Source/ModelViewViewModelEditor/Private/Widgets/SMVVMViewBindingListView.cpp
+++ b/Engine/Plugins/Runtime/ModelViewViewModel/Source/ModelViewViewModelEditor/Private/Widgets/SMVVMViewBindingListView.cpp
@@ -128,11 +128,14 @@ namespace Private
 
 	void FilterEntryList(FString FilterString, TOptional<SBindingsList::EFilterMode> FilterMode, const TArray<TSharedPtr<FBindingEntry>>& RootGroups, TArray<TSharedPtr<FBindingEntry>>& FilteredRootGroups, UMVVMBlueprintView* BlueprintView, UMVVMWidgetBlueprintExtension_View* MVVMExtensionPtr)
 	{
+		TArray<TSharedPtr<FBindingEntry>> FirstFilterRootGroups;
+
 		// First filtering pass: filter based on the filter mode.
-		if (FilterMode.IsSet())
+		if (FilterMode.IsSet() && FilterMode.GetValue() != SBindingsList::EFilterMode::All)
 		{
 			for (const TSharedPtr<FBindingEntry>& GroupEntry : RootGroups)
 			{
+				bool bNeedGroupEntry = false;
 				for (TSharedPtr<FBindingEntry> Entry : GroupEntry->GetAllChildren())
 				{
 					FBindingEntry* EntryPtr = Entry.Get();
@@ -157,14 +160,22 @@ namespace Private
 					if (!bHasLink || bEntryHasMatchingLink || bChildHasMatchingLink)
 					{
 						GroupEntry->AddFilteredChild(Entry);
+						bNeedGroupEntry = true;
 					}
 					else
 					{
 						GroupEntry->SetUseFilteredChildList();
 					}
 				}
+				if (bNeedGroupEntry)
+				{
+					FirstFilterRootGroups.Add(GroupEntry);
+				}
 			}
-
+		}
+		else
+		{
+			FirstFilterRootGroups = RootGroups;
 		}
 
 		// Second filtering pass: filter based on the filter string.
@@ -224,7 +235,7 @@ namespace Private
 			};
 			FAddFilteredEntry AddFilteredEntry = FAddFilteredEntry(IsAllKeywordsInString, SearchKeywords, BlueprintView, MVVMExtensionPtr->GetWidgetBlueprint());
 
-			for (const TSharedPtr<FBindingEntry>& GroupEntry : RootGroups)
+			for (const TSharedPtr<FBindingEntry>& GroupEntry : FirstFilterRootGroups)
 			{
 				FString EntryString = GroupEntry->GetSearchNameString(BlueprintView, MVVMExtensionPtr->GetWidgetBlueprint());
 
@@ -244,7 +255,7 @@ namespace Private
 		}
 		else
 		{
-			FilteredRootGroups = RootGroups;
+			FilteredRootGroups = FirstFilterRootGroups;
 		}
 	}
 
diff --git a/Engine/Plugins/Runtime/ModelViewViewModel/Source/ModelViewViewModelEditor/Private/Widgets/SMVVMViewBindingListView.h b/Engine/Plugins/Runtime/ModelViewViewModel/Source/ModelViewViewModelEditor/Private/Widgets/SMVVMViewBindingListView.h
--- a/Engine/Plugins/Runtime/ModelViewViewModel/Source/ModelViewViewModelEditor/Private/Widgets/SMVVMViewBindingListView.h
+++ b/Engine/Plugins/Runtime/ModelViewViewModel/Source/ModelViewViewModelEditor/Private/Widgets/SMVVMViewBindingListView.h
@@ -59,8 +59,9 @@ public:
 
 	enum class EFilterMode : uint8
 	{
-		Viewmodel = 0, // Display only bindings associated to a viewmodel
-		Verse = 1	// Display only bindings associated to a verse field
+		All = 0, // Display all bindings, no filtering
+		Viewmodel = 1, // Display only bindings associated to a viewmodel
+		Verse = 2	// Display only bindings associated to a verse field
 	};
 
 	EFilterMode GetActiveFilterMode() const;
@@ -88,7 +89,7 @@ private:
 	TWeakPtr<FWidgetBlueprintEditor> WeakBlueprintEditor;
 	mutable bool bSelectionChangedGuard = false;
 	FText FilterText;
-	EFilterMode FilterMode = EFilterMode::Verse;
+	EFilterMode FilterMode = EFilterMode::All;
 };
 
 } // namespace UE::MVVM
diff --git a/Engine/Plugins/Runtime/ModelViewViewModel/Source/ModelViewViewModelEditor/Private/Widgets/SMVVMViewBindingPanel.cpp b/Engine/Plugins/Runtime/ModelViewViewModel/Source/ModelViewViewModelEditor/Private/Widgets/SMVVMViewBindingPanel.cpp
--- a/Engine/Plugins/Runtime/ModelViewViewModel/Source/ModelViewViewModelEditor/Private/Widgets/SMVVMViewBindingPanel.cpp
+++ b/Engine/Plugins/Runtime/ModelViewViewModel/Source/ModelViewViewModelEditor/Private/Widgets/SMVVMViewBindingPanel.cpp
@@ -959,9 +959,12 @@ TSharedRef<SWidget> SBindingsPanel::GenerateEditViewWidget()
 		FilterWidget = SNew(SSegmentedControl<SBindingsList::EFilterMode>)
 						.Value(BindingsList.ToSharedRef(), &SBindingsList::GetActiveFilterMode)
 						.OnValueChanged(BindingsList.ToSharedRef(), &SBindingsList::OnActiveFilterModeChanged)
+						+ SSegmentedControl<SBindingsList::EFilterMode>::Slot(SBindingsList::EFilterMode::All)
+						.Icon(FAppStyle::GetBrush("Icons.Link"))
+						.ToolTip(LOCTEXT("BindingPanelAllFilterModeToolTip", "All bindings"))
 						+ SSegmentedControl<SBindingsList::EFilterMode>::Slot(SBindingsList::EFilterMode::Viewmodel)
 						.Icon(FMVVMEditorStyle::Get().GetBrush("BlueprintView.TabIcon"))
-						.ToolTip(LOCTEXT("BindingPanelViewModelFilterModeToolTip", "Viewmodel bindings"))
+						.ToolTip(LOCTEXT("BindingPanelViewModelFilterModeToolTip", "ViewModel bindings"))
 						+ SSegmentedControl<SBindingsList::EFilterMode>::Slot(SBindingsList::EFilterMode::Verse)
 						.Icon(FMVVMEditorStyle::Get().GetBrush("BindingView.VerseBinding"))
 						.ToolTip(LOCTEXT("BindingPanelVerseFilterModeToolTip", "Verse bindings"));
