- Fix display name of conversion functions in the details panel, in order to display a friendly name instead of the internal name.
- Small tweak on padding to avoid wasting space / keep content centered in the SPropertyBinding widget.

#tests Displayed various types of conversion function, validated that they were all displayed as expected. No regression on other types (fields)
#rnx

diff --git a/Engine/Plugins/Runtime/ModelViewViewModel/Source/ModelViewViewModelEditor/Private/Customizations/MVVMPropertyBindingExtension.cpp b/Engine/Plugins/Runtime/ModelViewViewModel/Source/ModelViewViewModelEditor/Private/Customizations/MVVMPropertyBindingExtension.cpp
--- a/Engine/Plugins/Runtime/ModelViewViewModel/Source/ModelViewViewModelEditor/Private/Customizations/MVVMPropertyBindingExtension.cpp
+++ b/Engine/Plugins/Runtime/ModelViewViewModel/Source/ModelViewViewModelEditor/Private/Customizations/MVVMPropertyBindingExtension.cpp
@@ -287,7 +287,21 @@ TOptional<FName> FMVVMPropertyBindingExtension::GetCurrentValue(const UWidgetBlu
 	}
 	if (UMVVMBlueprintViewConversionFunction* ConversionFunction = Binding->Conversion.GetConversionFunction(true))
 	{
-		return ConversionFunction->GetConversionFunction().GetName();
+		FMVVMBlueprintFunctionReference FunctionRef = ConversionFunction->GetConversionFunction();
+		if (FunctionRef.GetType() == EMVVMBlueprintFunctionReferenceType::Function)
+		{
+			const UFunction* FoundFunction = FunctionRef.GetFunction(WidgetBlueprint);
+			return FoundFunction ? FName(FoundFunction->GetDisplayNameText().ToString()) : TOptional<FName>();
+		}
+		else if (FunctionRef.GetType() == EMVVMBlueprintFunctionReferenceType::Node)
+		{
+			TSubclassOf<UK2Node> FoundFunction = FunctionRef.GetNode();
+			if (FoundFunction && FoundFunction->GetDefaultObject<UK2Node>())
+			{
+				return FName(FoundFunction->GetDefaultObject<UK2Node>()->GetNodeTitle(ENodeTitleType::ListView).ToString());
+			}
+			return FoundFunction.Get() ? FName(FoundFunction->GetDisplayNameText().ToString()) : TOptional<FName>();
+		}
 	}
 	return TOptional<FName>();
 }
diff --git a/Engine/Plugins/Runtime/PropertyAccess/Source/PropertyAccessEditor/Private/SPropertyBinding.cpp b/Engine/Plugins/Runtime/PropertyAccess/Source/PropertyAccessEditor/Private/SPropertyBinding.cpp
--- a/Engine/Plugins/Runtime/PropertyAccess/Source/PropertyAccessEditor/Private/SPropertyBinding.cpp
+++ b/Engine/Plugins/Runtime/PropertyAccess/Source/PropertyAccessEditor/Private/SPropertyBinding.cpp
@@ -61,7 +61,9 @@ void SPropertyBinding::Construct(const FArguments& InArgs, UBlueprint* InBluepri
 			BuildContextStructCategoryRecursive(Categories, Section->SubCategories, Index);
 		}
 	}
-	TSharedRef<SWidget> LinkIcon = Args.bUseLinkIconStyle && Args.bCanDisplayLinkIcon ?
+
+	const bool bDisplayLinkIcon = Args.bUseLinkIconStyle && Args.bCanDisplayLinkIcon;
+	TSharedRef<SWidget> LinkIcon = bDisplayLinkIcon ?
 		SNew(SBox)
 		.HeightOverride(16.0f)
 		[
@@ -71,6 +73,9 @@ void SPropertyBinding::Construct(const FArguments& InArgs, UBlueprint* InBluepri
 		:
 		SNullWidget::NullWidget;
 
+	//Adjust padding to avoid non centered content and wasted space.
+	const float TextLeftPadding = !bDisplayLinkIcon || Args.bLinkIconStyleCompactMode ? 0.0f : 4.0f;
+
 	ChildSlot
 	[
 		SNew(SHorizontalBox)
@@ -109,7 +114,7 @@ void SPropertyBinding::Construct(const FArguments& InArgs, UBlueprint* InBluepri
 				+ SHorizontalBox::Slot()
 				.AutoWidth()
 				.VAlign(VAlign_Center)
-				.Padding(4, 0, 0, 0)
+				.Padding(TextLeftPadding, 0, 0, 0)
 				[
 					SNew(STextBlock)
 					.Text(this, &SPropertyBinding::GetCurrentBindingText)
