MVVM: Move property permission validation out of compile

#rnx

diff --git a/Engine/Plugins/Runtime/ModelViewViewModel/Source/ModelViewViewModelBlueprint/Private/Bindings/MVVMCompiledBindingLibraryCompiler.cpp b/Engine/Plugins/Runtime/ModelViewViewModel/Source/ModelViewViewModelBlueprint/Private/Bindings/MVVMCompiledBindingLibraryCompiler.cpp
--- a/Engine/Plugins/Runtime/ModelViewViewModel/Source/ModelViewViewModelBlueprint/Private/Bindings/MVVMCompiledBindingLibraryCompiler.cpp
+++ b/Engine/Plugins/Runtime/ModelViewViewModel/Source/ModelViewViewModelBlueprint/Private/Bindings/MVVMCompiledBindingLibraryCompiler.cpp
@@ -357,7 +357,7 @@ TValueOrError<FCompiledBindingLibraryCompiler::FFieldPathHandle, FText> FCompile
 				return MakeError(ValidatedStr);
 			}
 
-			if (!GetDefault<UMVVMDeveloperProjectSettings>()->IsPropertyAllowed(Impl->BlueprintContext.Get(), CurrentContainer, FieldVariant.GetProperty()))
+			if (!GetDefault<UMVVMDeveloperProjectSettings>()->IsPropertyAllowed(Impl->BlueprintContext.Get(), CurrentContainer, FieldVariant.GetProperty(), false))
 			{
 				return MakeError(LOCTEXT("PropertyNotAllow", "A property is not allowed."));
 			}
@@ -384,7 +384,7 @@ TValueOrError<FCompiledBindingLibraryCompiler::FFieldPathHandle, FText> FCompile
 				return MakeError(LOCTEXT("InvalidContainer", "The path has an invalid container."));
 			}
 
-			if (!GetDefault<UMVVMDeveloperProjectSettings>()->IsFunctionAllowed(Impl->BlueprintContext.Get(), CurrentContainerAsClass, FieldVariant.GetFunction()))
+			if (!GetDefault<UMVVMDeveloperProjectSettings>()->IsFunctionAllowed(Impl->BlueprintContext.Get(), CurrentContainerAsClass, FieldVariant.GetFunction(), false))
 			{
 				return MakeError(LOCTEXT("FunctionNotAllow", "A function is not allowed."));
 			}
diff --git a/Engine/Plugins/Runtime/ModelViewViewModel/Source/ModelViewViewModelBlueprint/Private/MVVMDeveloperProjectSettings.cpp b/Engine/Plugins/Runtime/ModelViewViewModel/Source/ModelViewViewModelBlueprint/Private/MVVMDeveloperProjectSettings.cpp
--- a/Engine/Plugins/Runtime/ModelViewViewModel/Source/ModelViewViewModelBlueprint/Private/MVVMDeveloperProjectSettings.cpp
+++ b/Engine/Plugins/Runtime/ModelViewViewModel/Source/ModelViewViewModelBlueprint/Private/MVVMDeveloperProjectSettings.cpp
@@ -138,7 +138,7 @@ bool ShouldDoFieldEditorPermission(const UBlueprint* GeneratingFor, const UClass
 }
 }//namespace
 
-bool UMVVMDeveloperProjectSettings::IsPropertyAllowed(const UBlueprint* GeneratingFor, const UStruct* ObjectStruct, const FProperty* Property) const
+bool UMVVMDeveloperProjectSettings::IsPropertyAllowed(const UBlueprint* GeneratingFor, const UStruct* ObjectStruct, const FProperty* Property, bool bCheckEditorPermissions) const
 {
 	check(GeneratingFor);
 	check(ObjectStruct);
@@ -147,7 +147,7 @@ bool UMVVMDeveloperProjectSettings::IsPropertyAllowed(const UBlueprint* Generati
 	const UClass* AuthoritativeClass = Cast<const UClass>(ObjectStruct);
 	AuthoritativeClass = AuthoritativeClass ? AuthoritativeClass->GetAuthoritativeClass() : nullptr;
 
-	const bool bDoPropertyEditorPermission = UE::MVVM::Private::ShouldDoFieldEditorPermission(GeneratingFor, AuthoritativeClass, Property->GetOwnerClass());
+	const bool bDoPropertyEditorPermission = bCheckEditorPermissions && UE::MVVM::Private::ShouldDoFieldEditorPermission(GeneratingFor, AuthoritativeClass, Property->GetOwnerClass());
 	if (bDoPropertyEditorPermission)
 	{
 		if (!FPropertyEditorPermissionList::Get().DoesPropertyPassFilter(AuthoritativeClass, Property->GetFName()))
@@ -181,7 +181,7 @@ bool UMVVMDeveloperProjectSettings::IsPropertyAllowed(const UBlueprint* Generati
 	return true;
 }
 
-bool UMVVMDeveloperProjectSettings::IsFunctionAllowed(const UBlueprint* GeneratingFor, const UClass* ObjectClass, const UFunction* Function) const
+bool UMVVMDeveloperProjectSettings::IsFunctionAllowed(const UBlueprint* GeneratingFor, const UClass* ObjectClass, const UFunction* Function, bool bCheckEditorPermissions) const
 {
 	check(GeneratingFor);
 	check(ObjectClass);
@@ -196,7 +196,7 @@ bool UMVVMDeveloperProjectSettings::IsFunctionAllowed(const UBlueprint* Generati
 	const FPathPermissionList& FunctionPermissions = GetMutableDefault<UBlueprintEditorSettings>()->GetFunctionPermissions();
 	if (FunctionPermissions.HasFiltering())
 	{
-		const bool bDoPropertyEditorPermission = UE::MVVM::Private::ShouldDoFieldEditorPermission(GeneratingFor, AuthoritativeClass, Function->GetOwnerClass());
+		const bool bDoPropertyEditorPermission = bCheckEditorPermissions && UE::MVVM::Private::ShouldDoFieldEditorPermission(GeneratingFor, AuthoritativeClass, Function->GetOwnerClass());
 		if (bDoPropertyEditorPermission)
 		{
 			const UFunction* FunctionToTest = AuthoritativeClass->FindFunctionByName(Function->GetFName());
diff --git a/Engine/Plugins/Runtime/ModelViewViewModel/Source/ModelViewViewModelBlueprint/Private/MVVMViewBlueprintCompiler.cpp b/Engine/Plugins/Runtime/ModelViewViewModel/Source/ModelViewViewModelBlueprint/Private/MVVMViewBlueprintCompiler.cpp
--- a/Engine/Plugins/Runtime/ModelViewViewModel/Source/ModelViewViewModelBlueprint/Private/MVVMViewBlueprintCompiler.cpp
+++ b/Engine/Plugins/Runtime/ModelViewViewModel/Source/ModelViewViewModelBlueprint/Private/MVVMViewBlueprintCompiler.cpp
@@ -4687,7 +4687,7 @@ bool FMVVMViewBlueprintCompiler::IsPropertyPathValid(const UBlueprint* Context,
 				return false;
 			}
 
-			if (!GetDefault<UMVVMDeveloperProjectSettings>()->IsPropertyAllowed(Context, CurrentContainer, Field.GetProperty()))
+			if (!GetDefault<UMVVMDeveloperProjectSettings>()->IsPropertyAllowed(Context, CurrentContainer, Field.GetProperty(), false))
 			{
 				return false;
 			}
@@ -4703,7 +4703,7 @@ bool FMVVMViewBlueprintCompiler::IsPropertyPathValid(const UBlueprint* Context,
 			{
 				return false;
 			}
-			if (!GetDefault<UMVVMDeveloperProjectSettings>()->IsFunctionAllowed(Context, CurrentContainerAsClass, Field.GetFunction()))
+			if (!GetDefault<UMVVMDeveloperProjectSettings>()->IsFunctionAllowed(Context, CurrentContainerAsClass, Field.GetFunction(), false))
 			{
 				return false;
 			}
diff --git a/Engine/Plugins/Runtime/ModelViewViewModel/Source/ModelViewViewModelBlueprint/Public/MVVMDeveloperProjectSettings.h b/Engine/Plugins/Runtime/ModelViewViewModel/Source/ModelViewViewModelBlueprint/Public/MVVMDeveloperProjectSettings.h
--- a/Engine/Plugins/Runtime/ModelViewViewModel/Source/ModelViewViewModelBlueprint/Public/MVVMDeveloperProjectSettings.h
+++ b/Engine/Plugins/Runtime/ModelViewViewModel/Source/ModelViewViewModelBlueprint/Public/MVVMDeveloperProjectSettings.h
@@ -80,8 +80,8 @@ public:
 #endif
 
 	UE_API bool PropertyHasFiltering(const UStruct* ObjectStruct, const FProperty* Property) const;
-	UE_API bool IsPropertyAllowed(const UBlueprint* Context, const UStruct* ObjectStruct, const FProperty* Property) const;
-	UE_API bool IsFunctionAllowed(const UBlueprint* Context, const UClass* ObjectClass, const UFunction* Function) const;
+	UE_API bool IsPropertyAllowed(const UBlueprint* Context, const UStruct* ObjectStruct, const FProperty* Property, bool bCheckEditorPermissions = true) const;
+	UE_API bool IsFunctionAllowed(const UBlueprint* Context, const UClass* ObjectClass, const UFunction* Function, bool bCheckEditorPermissions = true) const;
 	UE_API bool IsConversionFunctionAllowed(const UBlueprint* Context, const UFunction* Function) const;
 	UE_API bool IsConversionFunctionAllowed(const UBlueprint* Context, const TSubclassOf<UK2Node> Function) const;
 
