MVVM: use the skeleton class in the binding picker so that we don't have to recompile to see new variables

diff --git a/Engine/Plugins/Runtime/ModelViewViewModel/Source/ModelViewViewModelEditor/Private/Types/MVVMBindingSource.cpp b/Engine/Plugins/Runtime/ModelViewViewModel/Source/ModelViewViewModelEditor/Private/Types/MVVMBindingSource.cpp
--- a/Engine/Plugins/Runtime/ModelViewViewModel/Source/ModelViewViewModelEditor/Private/Types/MVVMBindingSource.cpp
+++ b/Engine/Plugins/Runtime/ModelViewViewModel/Source/ModelViewViewModelEditor/Private/Types/MVVMBindingSource.cpp
@@ -74,7 +74,7 @@ FBindingSource FBindingSource::CreateForBlueprint(const UWidgetBlueprint* Widget
 
 	Source.Source = EMVVMBlueprintFieldPathSource::SelfContext;
 	Source.DisplayName = FText::FromString(WidgetBlueprint->GetName());
-	Source.Class = WidgetBlueprint->GeneratedClass;
+	Source.Class = WidgetBlueprint->SkeletonGeneratedClass;
 	
 	return Source;
 }
diff --git a/Engine/Plugins/Runtime/ModelViewViewModel/Source/ModelViewViewModelEditor/Private/Widgets/SMVVMViewModelBindingListWidget.cpp b/Engine/Plugins/Runtime/ModelViewViewModel/Source/ModelViewViewModelEditor/Private/Widgets/SMVVMViewModelBindingListWidget.cpp
--- a/Engine/Plugins/Runtime/ModelViewViewModel/Source/ModelViewViewModelEditor/Private/Widgets/SMVVMViewModelBindingListWidget.cpp
+++ b/Engine/Plugins/Runtime/ModelViewViewModel/Source/ModelViewViewModelEditor/Private/Widgets/SMVVMViewModelBindingListWidget.cpp
@@ -217,7 +217,7 @@ TArray<FFieldVariant> FFieldIterator_Bindable::GetFields(const UStruct* Struct,
 	if (const UClass* Class = Cast<const UClass>(Struct))
 	{
 		const UWidgetBlueprint* WidgetBlueprintPtr = WidgetBlueprint.Get();
-		TSubclassOf<UObject> AccessorClass = WidgetBlueprintPtr ? WidgetBlueprintPtr->GeneratedClass : nullptr;
+		TSubclassOf<UObject> AccessorClass = WidgetBlueprintPtr ? WidgetBlueprintPtr->SkeletonGeneratedClass : nullptr;
 
 		TArray<FMVVMAvailableBinding> Bindings = bIsBindingToEvent ? UMVVMSubsystem::GetAvailableBindingsForEvent(const_cast<UClass*>(Class), AccessorClass) : UMVVMSubsystem::GetAvailableBindings(const_cast<UClass*>(Class), AccessorClass);
 		AddResult(Bindings, Struct, false);
@@ -554,8 +554,9 @@ void SSourceBindingList::AddSources(TArrayView<const FBindingSource> InSources)
 
 bool IsPropertyPathComponent(TArrayView<const FFieldVariant> FieldPath, FMVVMBlueprintPropertyPath& PropertyPath, const UWidgetBlueprint* WidgetBlueprintPtr)
 {
-	if (!ensure(WidgetBlueprintPtr))
+	if (!WidgetBlueprintPtr)
 	{
+		ensure(false);
 		return false;
 	}
 
