MVVM: Use skeleton generated class for UI Component property paths.

#rnx

diff --git a/Engine/Plugins/Runtime/ModelViewViewModel/Source/ModelViewViewModelEditor/Private/Widgets/SMVVMViewModelBindingListWidget.cpp b/Engine/Plugins/Runtime/ModelViewViewModel/Source/ModelViewViewModelEditor/Private/Widgets/SMVVMViewModelBindingListWidget.cpp
--- a/Engine/Plugins/Runtime/ModelViewViewModel/Source/ModelViewViewModelEditor/Private/Widgets/SMVVMViewModelBindingListWidget.cpp
+++ b/Engine/Plugins/Runtime/ModelViewViewModel/Source/ModelViewViewModelEditor/Private/Widgets/SMVVMViewModelBindingListWidget.cpp
@@ -214,14 +214,15 @@ TArray<FFieldVariant> FFieldIterator_Bindable::GetFields(const UStruct* Struct,
 		}
 	};
 
+	const UWidgetBlueprint* WidgetBlueprintPtr = WidgetBlueprint.Get();
+	TSubclassOf<UObject> AccessorClass = nullptr;
+	if (WidgetBlueprintPtr)
+	{
+		AccessorClass = WidgetBlueprintPtr->SkeletonGeneratedClass ? WidgetBlueprintPtr->SkeletonGeneratedClass : WidgetBlueprintPtr->GeneratedClass;
+	}
+
 	if (const UClass* Class = Cast<const UClass>(Struct))
 	{
-		const UWidgetBlueprint* WidgetBlueprintPtr = WidgetBlueprint.Get();
-		TSubclassOf<UObject> AccessorClass = nullptr;
-		if (WidgetBlueprintPtr)
-		{
-			AccessorClass = WidgetBlueprintPtr->SkeletonGeneratedClass ? WidgetBlueprintPtr->SkeletonGeneratedClass : WidgetBlueprintPtr->GeneratedClass;
-		}
 		TArray<FMVVMAvailableBinding> Bindings = bIsBindingToEvent ? UMVVMSubsystem::GetAvailableBindingsForEvent(const_cast<UClass*>(Class), AccessorClass) : UMVVMSubsystem::GetAvailableBindings(const_cast<UClass*>(Class), AccessorClass);
 		AddResult(Bindings, Struct, false);
 	}
@@ -267,10 +268,6 @@ TArray<FFieldVariant> FFieldIterator_Bindable::GetFields(const UStruct* Struct,
 	{
 		if (const UWidgetBlueprint* ContainerWidgetBlueprintPtr = Cast<UWidgetBlueprint>(WidgetBPG->ClassGeneratedBy))
 		{
-			const UWidgetBlueprint* WidgetBlueprintPtr = WidgetBlueprint.Get();
-
-			TSubclassOf<UObject> AccessorClass = WidgetBlueprintPtr ? WidgetBlueprintPtr->GeneratedClass : nullptr;
-
 			if (const UClass* Class = Cast<const UClass>(Struct))
 			{
 				if (UWidget* Widget = ContainerWidgetBlueprintPtr->WidgetTree->FindWidget(FieldName))
@@ -362,7 +359,11 @@ TOptional<const UStruct*> FFieldExpander_Bindable::GetExpandedFunction(const UFu
 TSharedRef<SWidget> ConstructFieldPreSlot(const UWidgetBlueprint* WidgetBlueprint, UE::PropertyViewer::SPropertyViewer::FHandle Handle, const FFieldVariant FieldPath, const bool bIsForEvent)
 {
 	TSharedRef<SWidget> ImageWidget = SNullWidget::NullWidget;
-	TSubclassOf<UObject> AccessorClass = WidgetBlueprint ? WidgetBlueprint->SkeletonGeneratedClass : nullptr;
+	const UClass* AccessorClass = nullptr;
+	if (WidgetBlueprint != nullptr)
+	{
+		AccessorClass = WidgetBlueprint->SkeletonGeneratedClass ? WidgetBlueprint->SkeletonGeneratedClass : WidgetBlueprint->GeneratedClass;
+	}
 	FMVVMAvailableBinding Binding = bIsForEvent ? UMVVMSubsystem::GetAvailableBindingForEvent(FMVVMConstFieldVariant(FieldPath), AccessorClass) : UMVVMSubsystem::GetAvailableBindingForField(FMVVMConstFieldVariant(FieldPath), AccessorClass);
 	if (Binding.IsValid())
 	{
@@ -733,7 +734,7 @@ FMVVMBlueprintPropertyPath SSourceBindingList::CreateBlueprintPropertyPath(SProp
 							if (Source->Key.GetClass() == Component->GetClass())
 							{
 								FName ComponentPropertyName = UUIComponentContainer::GetPropertyNameForComponent(Component, Widget->GetFName());
-								if (FProperty* Property = FindFProperty<FObjectPropertyBase>(Cast<UWidgetBlueprintGeneratedClass>(WidgetBlueprint->GeneratedClass), ComponentPropertyName))
+								if (FProperty* Property = FindFProperty<FObjectPropertyBase>(Cast<UWidgetBlueprintGeneratedClass>(AccessorClass), ComponentPropertyName))
 								{
 									FMVVMConstFieldVariant Field = FMVVMConstFieldVariant(Property);
 									PropertyPath.AppendPropertyPath(WidgetBlueprintPtr, Field);
