MVVM: Add ability to add implicit converters

This is useful for automatically converting simple types in view bindings. For complex types, use conversion functions instead.

diff --git a/Engine/Plugins/Runtime/ModelViewViewModel/Source/ModelViewViewModel/Private/Bindings/ImplicitConverters/MVVMNumericImplicitConverter.cpp b/Engine/Plugins/Runtime/ModelViewViewModel/Source/ModelViewViewModel/Private/Bindings/ImplicitConverters/MVVMNumericImplicitConverter.cpp
new file mode 100644
--- /dev/null
+++ b/Engine/Plugins/Runtime/ModelViewViewModel/Source/ModelViewViewModel/Private/Bindings/ImplicitConverters/MVVMNumericImplicitConverter.cpp
@@ -0,0 +1,59 @@
+﻿// Copyright Epic Games, Inc. All Rights Reserved.
+
+
+#include "Bindings/ImplicitConverters/MVVMNumericImplicitConverter.h"
+
+#include "UObject/UnrealType.h"
+
+#ifndef UE_MVVM_ALLOW_AUTO_INTEGRAL_CONVERSION
+#define UE_MVVM_ALLOW_AUTO_INTEGRAL_CONVERSION 1
+#endif
+
+bool FMVVMNumericImplicitConverter::CanConvertImplicitly(const FProperty* Source, const FProperty* Destination)
+{
+	const FNumericProperty* SourceNumericProperty = CastField<const FNumericProperty>(Source);
+	const FNumericProperty* DestinationNumericProperty = CastField<const FNumericProperty>(Destination);
+	if (SourceNumericProperty && DestinationNumericProperty)
+	{
+		const bool bSameType = Destination->SameType(Source);
+		const bool bBothFloatingPoint = SourceNumericProperty->IsFloatingPoint() && DestinationNumericProperty->IsFloatingPoint();
+#if UE_MVVM_ALLOW_AUTO_INTEGRAL_CONVERSION
+		const bool bBothIntegral = SourceNumericProperty->IsInteger() && DestinationNumericProperty->IsInteger();
+		const bool bOneIsEnum = SourceNumericProperty->IsEnum() || DestinationNumericProperty->IsEnum();
+		return !bSameType && (bBothFloatingPoint || (bBothIntegral && !bOneIsEnum));
+#else
+		return !bSameType && bBothFloatingPoint;
+#endif
+	}
+
+	return false;
+}
+
+void FMVVMNumericImplicitConverter::ConvertImplicitly(const FProperty* Source, const FProperty* Destination, void* Data)
+{
+	const FNumericProperty* SourceNumericProperty = CastField<const FNumericProperty>(Source);
+	const FNumericProperty* DestinationNumericProperty = CastField<const FNumericProperty>(Destination);
+	if (SourceNumericProperty && DestinationNumericProperty)
+	{
+		if (SourceNumericProperty->IsFloatingPoint() && DestinationNumericProperty->IsFloatingPoint())
+		{
+			//floating to floating
+			const void* SrcElemValue = static_cast<const uint8*>(Data);
+			void* DestElemValue = static_cast<uint8*>(Data);
+
+			const double Value = SourceNumericProperty->GetFloatingPointPropertyValue(SrcElemValue);
+			DestinationNumericProperty->SetFloatingPointPropertyValue(DestElemValue, Value);
+		}
+#if UE_MVVM_ALLOW_AUTO_INTEGRAL_CONVERSION
+		else if (SourceNumericProperty->IsInteger() && DestinationNumericProperty->IsInteger())
+		{
+			//integral to integral
+			const void* SrcElemValue = static_cast<const uint8*>(Data);
+			void* DestElemValue = static_cast<uint8*>(Data);
+
+			const int64 Value = SourceNumericProperty->GetSignedIntPropertyValue(SrcElemValue);
+			DestinationNumericProperty->SetIntPropertyValue(DestElemValue, Value);
+		}
+#endif
+	}
+}
diff --git a/Engine/Plugins/Runtime/ModelViewViewModel/Source/ModelViewViewModel/Private/Bindings/MVVMBindingHelper.cpp b/Engine/Plugins/Runtime/ModelViewViewModel/Source/ModelViewViewModel/Private/Bindings/MVVMBindingHelper.cpp
--- a/Engine/Plugins/Runtime/ModelViewViewModel/Source/ModelViewViewModel/Private/Bindings/MVVMBindingHelper.cpp
+++ b/Engine/Plugins/Runtime/ModelViewViewModel/Source/ModelViewViewModel/Private/Bindings/MVVMBindingHelper.cpp
@@ -1,19 +1,19 @@
 // Copyright Epic Games, Inc. All Rights Reserved.
 
 #include "Bindings/MVVMBindingHelper.h"
+
+#include "Bindings/ImplicitConverters/MVVMImplicitConverter.h"
+#include "Blueprint/UserWidget.h"
+#include "ModelViewViewModelModule.h"
+#include "Modules/ModuleManager.h"
 #include "Templates/ValueOrError.h"
 #include "UObject/PropertyAccessUtil.h"
 #include "Types/MVVMBindingName.h"
 #include "Types/MVVMFieldContext.h"
 #include "Types/MVVMFunctionContext.h"
-#include "Blueprint/UserWidget.h"
 
 #define LOCTEXT_NAMESPACE "MVVMBindingHelper"
 
-#ifndef UE_MVVM_ALLOW_AUTO_INTEGRAL_CONVERSION
-#define UE_MVVM_ALLOW_AUTO_INTEGRAL_CONVERSION 1
-#endif
-
 namespace UE::MVVM::BindingHelper
 {
 	namespace Private
@@ -514,56 +514,35 @@ namespace UE::MVVM::BindingHelper
 		}
 
 
-		bool IsNumericConversionRequired(const FProperty* Source, const FProperty* Destination)
+		bool CanConvertImplicitly(const FProperty* Source, const FProperty* Destination)
 		{
-			const FNumericProperty* SourceNumericProperty = CastField<const FNumericProperty>(Source);
-			const FNumericProperty* DestinationNumericProperty = CastField<const FNumericProperty>(Destination);
-			if (SourceNumericProperty && DestinationNumericProperty)
+			const FModelViewViewModelModule& MVVMModule = FModuleManager::GetModuleChecked<FModelViewViewModelModule>("ModelViewViewModel");
+			for (const TSharedPtr<FMVVMImplicitConverter>& Converter : MVVMModule.GetImplicitConverters())
 			{
-				const bool bSameType = Destination->SameType(Source);
-				const bool bBothFloatingPoint = SourceNumericProperty->IsFloatingPoint() && DestinationNumericProperty->IsFloatingPoint();
-#if UE_MVVM_ALLOW_AUTO_INTEGRAL_CONVERSION
-				const bool bBothIntegral = SourceNumericProperty->IsInteger() && DestinationNumericProperty->IsInteger();
-				const bool bOneIsEnum = SourceNumericProperty->IsEnum() || DestinationNumericProperty->IsEnum();
-				return !bSameType && (bBothFloatingPoint || (bBothIntegral && !bOneIsEnum));
-#else
-				return !bSameType && bBothFloatingPoint;
-#endif
+				if (Converter.IsValid() && Converter->CanConvertImplicitly(Source, Destination))
+				{
+					return true;
+				}
 			}
+
 			return false;
 		}
 
 
-		void ConvertNumeric(const FProperty* Source, const FProperty* Destination, void* Data)
+		void ConvertImplicitly(const FProperty* Source, const FProperty* Destination, void* Data)
 		{
 			check(Source);
 			check(Destination);
 			check(Data);
 
-			const FNumericProperty* SourceNumericProperty = CastField<const FNumericProperty>(Source);
-			const FNumericProperty* DestinationNumericProperty = CastField<const FNumericProperty>(Destination);
-			if (SourceNumericProperty && DestinationNumericProperty)
+			const FModelViewViewModelModule& MVVMModule = FModuleManager::GetModuleChecked<FModelViewViewModelModule>("ModelViewViewModel");
+			for (const TSharedPtr<FMVVMImplicitConverter>& Converter : MVVMModule.GetImplicitConverters())
 			{
-				if (SourceNumericProperty->IsFloatingPoint() && DestinationNumericProperty->IsFloatingPoint())
+				if (Converter.IsValid() && Converter->CanConvertImplicitly(Source, Destination))
 				{
-					//floating to floating
-					const void* SrcElemValue = static_cast<const uint8*>(Data);
-					void* DestElemValue = static_cast<uint8*>(Data);
-
-					const double Value = SourceNumericProperty->GetFloatingPointPropertyValue(SrcElemValue);
-					DestinationNumericProperty->SetFloatingPointPropertyValue(DestElemValue, Value);
+					Converter->ConvertImplicitly(Source, Destination, Data);
+					break;
 				}
-#if UE_MVVM_ALLOW_AUTO_INTEGRAL_CONVERSION
-				else if (SourceNumericProperty->IsInteger() && DestinationNumericProperty->IsInteger())
-				{
-					//integral to integral
-					const void* SrcElemValue = static_cast<const uint8*>(Data);
-					void* DestElemValue = static_cast<uint8*>(Data);
-
-					const int64 Value = SourceNumericProperty->GetSignedIntPropertyValue(SrcElemValue);
-					DestinationNumericProperty->SetIntPropertyValue(DestElemValue, Value);
-				}
-#endif
 			}
 		}
 	} // namespace
@@ -576,7 +555,7 @@ namespace UE::MVVM::BindingHelper
 		}
 
 		return Source->SameType(Destination)
-				|| Private::IsNumericConversionRequired(Source, Destination)
+				|| Private::CanConvertImplicitly(Source, Destination)
 				|| Private::IsObjectPropertyCompatible(Source, Destination)
 				|| Private::IsObjectArrayPropertyCompatible(Source, Destination);
 	}
@@ -732,7 +711,7 @@ namespace UE::MVVM::BindingHelper
 			Source.GetObjectVariant().GetUObject()->ProcessEvent(Source.GetFieldVariant().GetFunction(), DataPtr);
 		}
 
-		Private::ConvertNumeric(GetterType, SetterType, DataPtr);
+		Private::ConvertImplicitly(GetterType, SetterType, DataPtr);
 
 		if (bIsDestinationBindingIsProperty)
 		{
@@ -771,7 +750,7 @@ namespace UE::MVVM::BindingHelper
 			check(SetterType);
 			check(ArePropertiesCompatible(ReturnConversionProperty, SetterType));
 
-			if (Private::IsNumericConversionRequired(ReturnConversionProperty, SetterType))
+			if (Private::CanConvertImplicitly(ReturnConversionProperty, SetterType))
 			{
 				// we need to do a copy because we may destroy the ReturnConversionProperty (imagine the return type is a TArray)
 				const int32 AllocationSize = FMath::Max(SetterType->GetSize(), ReturnConversionProperty->GetSize());
@@ -780,7 +759,7 @@ namespace UE::MVVM::BindingHelper
 				SetterType->InitializeValue(SetterDataPtr); // probably not needed since they are double/float
 
 				ReturnConversionProperty->CopyCompleteValue(SetterDataPtr, ReturnConversionProperty->ContainerPtrToValuePtr<void>(ConversionFunctionDataPtr));
-				Private::ConvertNumeric(ReturnConversionProperty, SetterType, SetterDataPtr);
+				Private::ConvertImplicitly(ReturnConversionProperty, SetterType, SetterDataPtr);
 
 				if (bIsDestinationBindingIsProperty)
 				{
@@ -841,7 +820,7 @@ namespace UE::MVVM::BindingHelper
 			check(GetterType);
 			check(ArePropertiesCompatible(GetterType, ArgumentConversionProperty));
 
-			if (Private::IsNumericConversionRequired(GetterType, ArgumentConversionProperty))
+			if (Private::CanConvertImplicitly(GetterType, ArgumentConversionProperty))
 			{
 				// we need to do a copy because we may destroy the ReturnConversionProperty (imagine the return type is a TArray)
 				const int32 AllocationSize = FMath::Max(GetterType->GetSize(), ArgumentConversionProperty->GetSize());
@@ -860,7 +839,7 @@ namespace UE::MVVM::BindingHelper
 					Source.GetObjectVariant().GetUObject()->ProcessEvent(Source.GetFieldVariant().GetFunction(), GetterDataPtr);
 				}
 
-				Private::ConvertNumeric(GetterType, ArgumentConversionProperty, GetterDataPtr);
+				Private::ConvertImplicitly(GetterType, ArgumentConversionProperty, GetterDataPtr);
 				ArgumentConversionProperty->CopyCompleteValue(ArgumentConversionProperty->ContainerPtrToValuePtr<void>(ConversionFunctionDataPtr), GetterDataPtr);
 
 				GetterType->DestroyValue(GetterDataPtr);
@@ -894,7 +873,7 @@ namespace UE::MVVM::BindingHelper
 			check(SetterType);
 			check(ArePropertiesCompatible(ReturnConversionProperty, SetterType));
 
-			if (Private::IsNumericConversionRequired(ReturnConversionProperty, SetterType))
+			if (Private::CanConvertImplicitly(ReturnConversionProperty, SetterType))
 			{
 				// we need to do a copy because we may destroy the ReturnConversionProperty (imagine the return type is a TArray)
 				const int32 AllocationSize = FMath::Max(SetterType->GetSize(), ReturnConversionProperty->GetSize());
@@ -903,7 +882,7 @@ namespace UE::MVVM::BindingHelper
 				SetterType->InitializeValue(SetterDataPtr); // probably not needed since they are double/float
 
 				ReturnConversionProperty->CopyCompleteValue(SetterDataPtr, ReturnConversionProperty->ContainerPtrToValuePtr<void>(ConversionFunctionDataPtr));
-				Private::ConvertNumeric(ReturnConversionProperty, SetterType, SetterDataPtr);
+				Private::ConvertImplicitly(ReturnConversionProperty, SetterType, SetterDataPtr);
 
 				if (bIsDestinationBindingIsProperty)
 				{
diff --git a/Engine/Plugins/Runtime/ModelViewViewModel/Source/ModelViewViewModel/Private/ModelViewViewModelModule.cpp b/Engine/Plugins/Runtime/ModelViewViewModel/Source/ModelViewViewModel/Private/ModelViewViewModelModule.cpp
--- a/Engine/Plugins/Runtime/ModelViewViewModel/Source/ModelViewViewModel/Private/ModelViewViewModelModule.cpp
+++ b/Engine/Plugins/Runtime/ModelViewViewModel/Source/ModelViewViewModel/Private/ModelViewViewModelModule.cpp
@@ -1,6 +1,8 @@
 // Copyright Epic Games, Inc. All Rights Reserved.
 
 #include "ModelViewViewModelModule.h"
+
+#include "Bindings/ImplicitConverters/MVVMNumericImplicitConverter.h"
 #include "HAL/IConsoleManager.h"
 #include "Modules/ModuleManager.h"
 #include "Types/MVVMExecutionMode.h"
@@ -20,6 +22,9 @@ void FModelViewViewModelModule::StartupModule()
 		HandleDefaultExecutionModeChanged(CVarDefaultExecutionMode);
 		CVarDefaultExecutionMode->OnChangedDelegate().AddRaw(this, &FModelViewViewModelModule::HandleDefaultExecutionModeChanged);
 	}
+	
+	NumericConverter = MakeShared<FMVVMNumericImplicitConverter>();
+	RegisterImplicitConverter(NumericConverter.ToSharedRef());
 
 #if WITH_EDITOR
 	FCoreDelegates::OnPostEngineInit.AddRaw(this, &FModelViewViewModelModule::OnPostEngineInit);
@@ -63,6 +68,11 @@ void FModelViewViewModelModule::ShutdownModule()
 		}
 	}
 
+	if (NumericConverter.IsValid())
+	{
+		UnregisterImplicitConverter(NumericConverter.ToSharedRef());
+	}
+
 #if WITH_EDITOR
 	FCoreDelegates::OnPostEngineInit.RemoveAll(this);
 	if (GEditor != nullptr)
@@ -79,6 +89,16 @@ void FModelViewViewModelModule::ShutdownModule()
 #endif
 }
 
+void FModelViewViewModelModule::RegisterImplicitConverter(const TSharedRef<FMVVMImplicitConverter>& Converter)
+{
+	ImplicitConverters.AddUnique(Converter);
+}
+
+void FModelViewViewModelModule::UnregisterImplicitConverter(const TSharedRef<FMVVMImplicitConverter>& Converter)
+{
+	ImplicitConverters.Remove(Converter);
+}
+
 void FModelViewViewModelModule::HandleDefaultExecutionModeChanged(IConsoleVariable* Variable)
 {
 	const int32 Value = Variable->GetInt();
diff --git a/Engine/Plugins/Runtime/ModelViewViewModel/Source/ModelViewViewModel/Private/ModelViewViewModelModule.h b/Engine/Plugins/Runtime/ModelViewViewModel/Source/ModelViewViewModel/Private/ModelViewViewModelModule.h
--- a/Engine/Plugins/Runtime/ModelViewViewModel/Source/ModelViewViewModel/Private/ModelViewViewModelModule.h
+++ b/Engine/Plugins/Runtime/ModelViewViewModel/Source/ModelViewViewModel/Private/ModelViewViewModelModule.h
@@ -2,9 +2,10 @@
 
 #pragma once
 
-#include "Modules/ModuleInterface.h"
+#include "IModelViewViewModelModule.h"
 #include "Stats/Stats.h"
 
+class FMVVMNumericImplicitConverter;
 class IConsoleVariable;
 class UBlueprint;
 
@@ -14,7 +15,7 @@ DECLARE_STATS_GROUP(TEXT("UMG Viewmodel"), STATGROUP_UMG_Viewmodel, STATCAT_Adva
 /**
  *
  */
-class FModelViewViewModelModule : public IModuleInterface
+class FModelViewViewModelModule : public IModelViewViewModelModule
 {
 public:
 	FModelViewViewModelModule() = default;
@@ -32,6 +33,10 @@ public:
 	FOnBlueprintPreCompile OnBlueprintPreCompile;
 #endif
 
+	virtual void RegisterImplicitConverter(const TSharedRef<FMVVMImplicitConverter>& Converter) override;
+	virtual void UnregisterImplicitConverter(const TSharedRef<FMVVMImplicitConverter>& Converter) override;
+	TConstArrayView<TSharedPtr<FMVVMImplicitConverter>> GetImplicitConverters() const { return ImplicitConverters; }
+
 private:
 	void HandleDefaultExecutionModeChanged(IConsoleVariable* Variable);
 	void OnPostEngineInit();
@@ -43,4 +48,6 @@ private:
 	FDelegateHandle OnBlueprintPreCompileHandle;
 #endif
 
+	TArray<TSharedPtr<FMVVMImplicitConverter>> ImplicitConverters;
+	TSharedPtr<FMVVMNumericImplicitConverter> NumericConverter;
 };
diff --git a/Engine/Plugins/Runtime/ModelViewViewModel/Source/ModelViewViewModel/Public/Bindings/ImplicitConverters/MVVMImplicitConverter.h b/Engine/Plugins/Runtime/ModelViewViewModel/Source/ModelViewViewModel/Public/Bindings/ImplicitConverters/MVVMImplicitConverter.h
new file mode 100644
--- /dev/null
+++ b/Engine/Plugins/Runtime/ModelViewViewModel/Source/ModelViewViewModel/Public/Bindings/ImplicitConverters/MVVMImplicitConverter.h
@@ -0,0 +1,18 @@
+﻿// Copyright Epic Games, Inc. All Rights Reserved.
+
+#pragma once
+
+class FProperty;
+
+/**
+ * Base class to support converting between simple types in MVVM bindings without the need for a conversion function
+ */
+class FMVVMImplicitConverter
+{
+public:
+	FMVVMImplicitConverter() = default;
+	virtual ~FMVVMImplicitConverter() = default;
+	
+	virtual bool CanConvertImplicitly(const FProperty* Source, const FProperty* Destination) = 0;
+	virtual void ConvertImplicitly(const FProperty* Source, const FProperty* Destination, void* Data) = 0;
+};
diff --git a/Engine/Plugins/Runtime/ModelViewViewModel/Source/ModelViewViewModel/Public/Bindings/ImplicitConverters/MVVMNumericImplicitConverter.h b/Engine/Plugins/Runtime/ModelViewViewModel/Source/ModelViewViewModel/Public/Bindings/ImplicitConverters/MVVMNumericImplicitConverter.h
new file mode 100644
--- /dev/null
+++ b/Engine/Plugins/Runtime/ModelViewViewModel/Source/ModelViewViewModel/Public/Bindings/ImplicitConverters/MVVMNumericImplicitConverter.h
@@ -0,0 +1,15 @@
+﻿// Copyright Epic Games, Inc. All Rights Reserved.
+
+#pragma once
+
+#include "Bindings/ImplicitConverters/MVVMImplicitConverter.h"
+
+/**
+ * For converting between numeric types
+ */
+class FMVVMNumericImplicitConverter : public FMVVMImplicitConverter
+{
+public:
+	virtual bool CanConvertImplicitly(const FProperty* Source, const FProperty* Destination) override;
+	virtual void ConvertImplicitly(const FProperty* Source, const FProperty* Destination, void* Data) override;
+};
diff --git a/Engine/Plugins/Runtime/ModelViewViewModel/Source/ModelViewViewModel/Public/IModelViewViewModelModule.h b/Engine/Plugins/Runtime/ModelViewViewModel/Source/ModelViewViewModel/Public/IModelViewViewModelModule.h
new file mode 100644
--- /dev/null
+++ b/Engine/Plugins/Runtime/ModelViewViewModel/Source/ModelViewViewModel/Public/IModelViewViewModelModule.h
@@ -0,0 +1,15 @@
+﻿// Copyright Epic Games, Inc. All Rights Reserved.
+
+#pragma once
+
+#include "Modules/ModuleInterface.h"
+#include "Templates/SharedPointer.h"
+
+class FMVVMImplicitConverter;
+
+class IModelViewViewModelModule : public IModuleInterface
+{
+public:
+	MODELVIEWVIEWMODEL_API virtual void RegisterImplicitConverter(const TSharedRef<FMVVMImplicitConverter>& Converter) = 0;
+	MODELVIEWVIEWMODEL_API virtual void UnregisterImplicitConverter(const TSharedRef<FMVVMImplicitConverter>& Converter) = 0;
+};
