MVVM: remove valid loaded property check to unblock build

diff --git a/Engine/Plugins/Runtime/ModelViewViewModel/Source/ModelViewViewModel/Private/Bindings/MVVMCompiledBindingLibrary.cpp b/Engine/Plugins/Runtime/ModelViewViewModel/Source/ModelViewViewModel/Private/Bindings/MVVMCompiledBindingLibrary.cpp
--- a/Engine/Plugins/Runtime/ModelViewViewModel/Source/ModelViewViewModel/Private/Bindings/MVVMCompiledBindingLibrary.cpp
+++ b/Engine/Plugins/Runtime/ModelViewViewModel/Source/ModelViewViewModel/Private/Bindings/MVVMCompiledBindingLibrary.cpp
@@ -339,7 +339,7 @@ bool FMVVMCompiledBindingLibrary::CanAccessLoadedProperiesAtIndex(int16 index) c
 		LoadedPropertiesString.Append(FString::Format(TEXT("{0}, "), { LoadedProperties[i]->GetNameCPP() }));
 	}
 	FString DebugLoadedProperties = FString::Format(TEXT("{0} on {1}. LoadedProperties : {2}"), { index, LoadedProperties.Num(), LoadedPropertiesString });
-	checkf(false, TEXT("LoadedProperties.IsValidIndex(PathIndex.Index) failed. '%s'"), *DebugLoadedProperties);
+	//checkf(false, TEXT("LoadedProperties.IsValidIndex(PathIndex.Index) failed. '%s'"), *DebugLoadedProperties);
 	UE_LOG(LogMVVM, Error, TEXT("LoadedProperties.IsValidIndex(PathIndex.Index) failed. '%s'"), *DebugLoadedProperties);
 	return false;
 }
