MVVM: Automated conversion binding setup from Details Panel

[RN] minor creator/ui

diff --git a/Engine/Plugins/Runtime/ModelViewViewModel/Source/ModelViewViewModelEditor/Private/Customizations/MVVMPropertyBindingExtension.cpp b/Engine/Plugins/Runtime/ModelViewViewModel/Source/ModelViewViewModelEditor/Private/Customizations/MVVMPropertyBindingExtension.cpp
--- a/Engine/Plugins/Runtime/ModelViewViewModel/Source/ModelViewViewModelEditor/Private/Customizations/MVVMPropertyBindingExtension.cpp
+++ b/Engine/Plugins/Runtime/ModelViewViewModel/Source/ModelViewViewModelEditor/Private/Customizations/MVVMPropertyBindingExtension.cpp
@@ -3,6 +3,7 @@
 #include "Customizations/MVVMPropertyBindingExtension.h"
 
 #include "Bindings/MVVMBindingHelper.h"
+#include "Bindings/MVVMConversionFunctionHelper.h"
 #include "BlueprintEditor.h"
 #include "Components/Widget.h"
 #include "DetailLayoutBuilder.h"
@@ -17,11 +18,14 @@
 #include "PropertyHandle.h"
 #include "PropertyPathHelpers.h"
 #include "ScopedTransaction.h"
+#include "SPrimaryButton.h"
 #include "Styling/CoreStyle.h"
 #include "Styling/StyleColors.h"
 #include "UMGEditorProjectSettings.h"
 #include "Widgets/ViewModelFieldDragDropOp.h"
 #include "Widgets/Images/SImage.h"
+#include "Widgets/Input/SButton.h"
+#include "Widgets/Views/SListView.h"
 #include "Widgets/SMVVMFieldSelectorMenu.h"
 #include "Widgets/Text/STextBlock.h"
 
@@ -29,6 +33,82 @@
 
 namespace UE::MVVM
 {
+namespace Private
+{
+
+struct FCreatePropertyBindingParams
+{
+	TWeakObjectPtr<const UWidgetBlueprint> WidgetBlueprint;
+	TWeakObjectPtr<UMVVMBlueprintView> MVVMBlueprintView;
+	TWeakObjectPtr<UWidget> Widget;
+	TSharedPtr<IPropertyHandle> PropertyHandle;
+};
+
+void CreateBindingFromProperty(const FCreatePropertyBindingParams& Params, FMVVMLinkedPinValue LinkedValue, const FMVVMFieldConversion& Conversion)
+{
+	UWidgetBlueprint* WidgetBlueprint = const_cast<UWidgetBlueprint*>(Params.WidgetBlueprint.Get());
+	UMVVMBlueprintView* MVVMBlueprintView = Params.MVVMBlueprintView.Get();
+	UMVVMEditorSubsystem* Subsystem = GEditor->GetEditorSubsystem<UMVVMEditorSubsystem>();
+	UWidget* Widget = Params.Widget.Get();
+	if (WidgetBlueprint == nullptr || MVVMBlueprintView == nullptr || Widget == nullptr || Subsystem == nullptr)
+	{
+		return;
+	}
+
+	//Search for existing binding by property target, and create it only if missing.
+	//Not perfect as if there are many bindings, only the first one will be used, but this is the limitation of the design.
+	FMVVMBlueprintViewBinding* Binding = MVVMBlueprintView->FindBinding(Widget, Params.PropertyHandle.Get()->GetProperty());
+	if (Binding == nullptr)
+	{
+		Binding = &MVVMBlueprintView->AddDefaultBinding();
+	}
+
+	if (LinkedValue.IsValid() && LinkedValue.IsPropertyPath())
+	{
+		const FMVVMBlueprintPropertyPath& SourcePath = LinkedValue.GetPropertyPath();
+		Binding->SourcePath = SourcePath;
+
+		// Generate the destination path from the widget property that we are dropping on.
+		FCachedPropertyPath CachedPropertyPath(Params.PropertyHandle->GeneratePathToProperty());
+		CachedPropertyPath.Resolve(Widget);
+
+		// Set the destination path.
+		FMVVMBlueprintPropertyPath DestinationPropertyPath;
+		DestinationPropertyPath.ResetPropertyPath();
+
+		for (int32 SegNum = 0; SegNum < CachedPropertyPath.GetNumSegments(); SegNum++)
+		{
+			FFieldVariant Field = CachedPropertyPath.GetSegment(SegNum).GetField();
+			DestinationPropertyPath.AppendPropertyPath(WidgetBlueprint, UE::MVVM::FMVVMConstFieldVariant(Field));
+		}
+
+		if (Widget->GetFName() == WidgetBlueprint->GetFName())
+		{
+			DestinationPropertyPath.SetSelfContext();
+		}
+		else
+		{
+			DestinationPropertyPath.SetWidgetName(Widget->GetFName());
+		}
+		Binding->DestinationPath = DestinationPropertyPath;
+		Binding->BindingType = EMVVMBindingMode::OneWayToDestination;
+
+		if (Conversion.IsValid())
+		{
+			Subsystem->SetSourceToDestinationConversionFunction(WidgetBlueprint, *Binding, Conversion.ToBlueprintFunctionReference(WidgetBlueprint));
+			Subsystem->SetPathForConversionFunctionArgument(WidgetBlueprint, *Binding, Conversion.ParameterId, SourcePath, true);
+		}
+	}
+	else
+	{
+		MVVMBlueprintView->RemoveBinding(Binding);
+	}
+
+	//Refresh after change.
+	MVVMBlueprintView->OnBindingsUpdated.Broadcast();
+}
+}
+
 void FMVVMPropertyBindingExtension::ExtendBindingsMenu(FMenuBuilder& MenuBuilder, TSharedRef<FMVVMPropertyBindingExtension> MVVMPropertyBindingExtension, const UWidgetBlueprint* WidgetBlueprint, UWidget* Widget, TSharedPtr<IPropertyHandle> WidgetPropertyHandle)
 {
 	const UUMGEditorProjectSettings* UMGEditorProjectSettings = GetMutableDefault<UUMGEditorProjectSettings>();
@@ -563,71 +643,33 @@ TSharedRef<SWidget> FMVVMPropertyBindingExtension::OnGenerateMenu(const UWidgetB
 		return SNullWidget::NullWidget;
 	}
 
-	TWeakObjectPtr<UMVVMBlueprintView> WeakMVVMBlueprintView = TWeakObjectPtr<UMVVMBlueprintView>(MVVMBlueprintView);
-	TWeakObjectPtr<const UWidgetBlueprint> WeakWidgetBlueprint = TWeakObjectPtr<const UWidgetBlueprint>(WidgetBlueprint);
-	TWeakObjectPtr<UWidget> WeakWidget = TWeakObjectPtr<UWidget>(Widget);
-	SFieldSelectorMenu::FOnLinkedValueSelected HandleFieldSelectionChanged = SFieldSelectorMenu::FOnLinkedValueSelected::CreateLambda([WeakWidgetBlueprint, WeakMVVMBlueprintView, PropertyHandle, WeakWidget, DetailsView](FMVVMLinkedPinValue LinkedValue, SFieldSelectorMenu::ESelectionType SelectionType)
-	{
-		const UWidgetBlueprint* WidgetBlueprint = WeakWidgetBlueprint.Get();
-		UMVVMBlueprintView* MVVMBlueprintView = WeakMVVMBlueprintView.Get();
-		UWidget* Widget = WeakWidget.Get();
-		if (WidgetBlueprint == nullptr || MVVMBlueprintView == nullptr || Widget == nullptr)
-		{
-			return;
-		}
+	Private::FCreatePropertyBindingParams CreateBindingParams;
+	CreateBindingParams.WidgetBlueprint = TWeakObjectPtr<const UWidgetBlueprint>(WidgetBlueprint);
+	CreateBindingParams.MVVMBlueprintView = TWeakObjectPtr<UMVVMBlueprintView>(MVVMBlueprintView);
+	CreateBindingParams.Widget = TWeakObjectPtr<UWidget>(Widget);
+	CreateBindingParams.PropertyHandle = PropertyHandle;
 
-		//Search for existing binding by property target, and create it only if missing.
-		//Not perfect as if there are many bindings, only the first one will be used, but this is the limitation of the design.
-		FMVVMBlueprintViewBinding* Binding = MVVMBlueprintView->FindBinding(Widget, PropertyHandle.Get()->GetProperty());
-		if (Binding == nullptr)
+	SFieldSelectorMenu::FOnLinkedValueConversionSelected HandleFieldConversionSelectionChanged = SFieldSelectorMenu::FOnLinkedValueConversionSelected::CreateLambda(
+		[CreateBindingParams, DetailsView]
+		(FMVVMLinkedPinValue LinkedValue, const FMVVMFieldConversion& Conversion, SFieldSelectorMenu::ESelectionType SelectionType)
 		{
-			Binding = &MVVMBlueprintView->AddDefaultBinding();
-		}
+			FSlateApplication::Get().DismissAllMenus();
 
-		//Setup the binding
-		if (LinkedValue.IsValid() && LinkedValue.IsPropertyPath())
-		{
-			Binding->SourcePath = LinkedValue.GetPropertyPath();
-
-			// Generate the destination path from the widget property that we are dropping on.
-			FCachedPropertyPath CachedPropertyPath(PropertyHandle->GeneratePathToProperty());
-			CachedPropertyPath.Resolve(Widget);
-
-			// Set the destination path.
-			FMVVMBlueprintPropertyPath DestinationPropertyPath;
-			DestinationPropertyPath.ResetPropertyPath();
-
-			for (int32 SegNum = 0; SegNum < CachedPropertyPath.GetNumSegments(); SegNum++)
-			{
-				FFieldVariant Field = CachedPropertyPath.GetSegment(SegNum).GetField();
-				DestinationPropertyPath.AppendPropertyPath(WidgetBlueprint, UE::MVVM::FMVVMConstFieldVariant(Field));
-			}
+			Private::CreateBindingFromProperty(CreateBindingParams, LinkedValue, Conversion);
 
-			if (Widget->GetFName() == WidgetBlueprint->GetFName())
-			{
-				DestinationPropertyPath.SetSelfContext();
-			}
-			else
+			IDetailsView* DetailsViewPtr = DetailsView.Pin().Get();
+			if (DetailsViewPtr != nullptr)
 			{
-				DestinationPropertyPath.SetWidgetName(Widget->GetFName());
+				DetailsViewPtr->ForceRefresh();
 			}
-			Binding->DestinationPath = DestinationPropertyPath;
-			Binding->BindingType = EMVVMBindingMode::OneWayToDestination;
-		}
-		else
-		{
-			MVVMBlueprintView->RemoveBinding(Binding);
-		}
+		});
 
-		//Refresh after change.
-		MVVMBlueprintView->OnBindingsUpdated.Broadcast();
-		FSlateApplication::Get().DismissAllMenus();
-		IDetailsView* DetailsViewPtr = DetailsView.Pin().Get();
-		if (DetailsViewPtr != nullptr)
+	SFieldSelectorMenu::FOnLinkedValueSelected HandleFieldSelectionChanged = SFieldSelectorMenu::FOnLinkedValueSelected::CreateLambda(
+		[HandleFieldConversionSelectionChanged]
+		(FMVVMLinkedPinValue LinkedValue, const SFieldSelectorMenu::ESelectionType SelectionType)
 		{
-			DetailsViewPtr->ForceRefresh();
-		}
-	});
+			HandleFieldConversionSelectionChanged.ExecuteIfBound(LinkedValue, FMVVMFieldConversion(), SelectionType);
+		});
 
 	FSimpleDelegate HandleMenuClosed = FSimpleDelegate::CreateLambda([]()
 	{
@@ -642,7 +684,6 @@ TSharedRef<SWidget> FMVVMPropertyBindingExtension::OnGenerateMenu(const UWidgetB
 	FFieldSelectionContext SelectionContext;
 	SelectionContext.BindingMode = EMVVMBindingMode::OneWayToDestination;
 	SelectionContext.AssignableTo = PropertyHandle.Get()->GetProperty();
-	SelectionContext.bAllowConversionFunctions = false;
 
 	//Right before displaying the menu (as the value can be changed anytime from UI), ensure the list of properties is filtered with compatible types only.
 	FMVVMViewBindingFilterSettings FilterSettings = MVVMExtensionPtr->GetFilterSettings();
@@ -654,12 +695,15 @@ TSharedRef<SWidget> FMVVMPropertyBindingExtension::OnGenerateMenu(const UWidgetB
 	TSharedRef<SFieldSelectorMenu> Menu = SNew(SFieldSelectorMenu, WidgetBlueprint)
 		.CurrentSelected(CurrentSelected)
 		.OnSelected(HandleFieldSelectionChanged)
+		.OnSelectedConversion(HandleFieldConversionSelectionChanged)
 		.OnMenuCloseRequested(HandleMenuClosed)
 		.SelectionContext(SelectionContext)
 		.IsBindingToEvent(false)
 		.CanCreateEvent(false)
 		.IsEnableContextToggleAvailable(false)
-		.IsClearEnableable(bHasBinding);
+		.IsClearEnableable(bHasBinding)
+		.CanConvertField(true)
+		.ShowTooltipText(true);
 
 	return Menu;
 }
diff --git a/Engine/Plugins/Runtime/ModelViewViewModel/Source/ModelViewViewModelEditor/Private/MVVMEditorSubsystem.cpp b/Engine/Plugins/Runtime/ModelViewViewModel/Source/ModelViewViewModelEditor/Private/MVVMEditorSubsystem.cpp
--- a/Engine/Plugins/Runtime/ModelViewViewModel/Source/ModelViewViewModelEditor/Private/MVVMEditorSubsystem.cpp
+++ b/Engine/Plugins/Runtime/ModelViewViewModel/Source/ModelViewViewModelEditor/Private/MVVMEditorSubsystem.cpp
@@ -1392,6 +1392,56 @@ TArray<UE::MVVM::FConversionFunctionValue> UMVVMEditorSubsystem::GetConversionFu
 	return TArray<UE::MVVM::FConversionFunctionValue>();
 }
 
+TArray<FMVVMBlueprintPinId> UMVVMEditorSubsystem::GetConversionFunctionPins(const UWidgetBlueprint* WidgetBlueprint, const UE::MVVM::FConversionFunctionValue& Function, const FProperty* ExpectedArgumentType) const
+{
+	TArray<FMVVMBlueprintPinId> OutPins;
+	if (Function.IsFunction())
+	{
+		TValueOrError<TArray<const FProperty*>, FText> ArgumentsResult = UE::MVVM::BindingHelper::TryGetArgumentsForConversionFunction(Function.GetFunction());
+		if (ArgumentsResult.HasError())
+		{
+			return OutPins;
+		}
+		for (const FProperty* ArgumentProperty : ArgumentsResult.GetValue())
+		{
+			if (UE::MVVM::BindingHelper::ArePropertiesCompatible(ExpectedArgumentType, ArgumentProperty))
+			{
+				TArray<FName> Names = { ArgumentProperty->GetFName() };
+				OutPins.Add(FMVVMBlueprintPinId(MoveTemp(Names)));
+			}
+		}
+
+		return OutPins;
+	}
+	else if (Function.IsNode())
+	{
+		FEdGraphPinType SourcePinType;
+		if (!GetDefault<UEdGraphSchema_K2>()->ConvertPropertyToPinType(ExpectedArgumentType, SourcePinType))
+		{
+			return OutPins;
+		}
+
+		TArray<UEdGraphPin*> InputPins = UE::MVVM::ConversionFunctionHelper::FindInputPins(Function.GetNode().GetDefaultObject());
+		if (InputPins.Num() == 0)
+		{
+			return OutPins;
+		}
+
+		UClass* CallingContext = UE::MVVM::Private::GetCallingContext(WidgetBlueprint);
+
+		bool bAnyCompatible = false;
+		for (const UEdGraphPin* ArgumentPin : InputPins)
+		{
+			bool bIgnoreArray = true;
+			if (GetDefault<UEdGraphSchema_K2>()->ArePinTypesCompatible(SourcePinType, ArgumentPin->PinType, CallingContext, bIgnoreArray))
+			{
+				OutPins.Add(FMVVMBlueprintPinId(UE::MVVM::ConversionFunctionHelper::FindPinId(ArgumentPin)));
+			}
+		}
+	}
+	return OutPins;
+}
+
 FMVVMBlueprintPropertyPath UMVVMEditorSubsystem::GetPathForConversionFunctionArgument(const UWidgetBlueprint* WidgetBlueprint, const FMVVMBlueprintViewBinding& Binding,  const FMVVMBlueprintPinId& ParameterId, bool bSourceToDestination) const
 {
 	UMVVMBlueprintViewConversionFunction* ConversionFunction = Binding.Conversion.GetConversionFunction(bSourceToDestination);
diff --git a/Engine/Plugins/Runtime/ModelViewViewModel/Source/ModelViewViewModelEditor/Private/Widgets/SMVVMFieldSelectorMenu.cpp b/Engine/Plugins/Runtime/ModelViewViewModel/Source/ModelViewViewModelEditor/Private/Widgets/SMVVMFieldSelectorMenu.cpp
--- a/Engine/Plugins/Runtime/ModelViewViewModel/Source/ModelViewViewModelEditor/Private/Widgets/SMVVMFieldSelectorMenu.cpp
+++ b/Engine/Plugins/Runtime/ModelViewViewModel/Source/ModelViewViewModelEditor/Private/Widgets/SMVVMFieldSelectorMenu.cpp
@@ -17,6 +17,7 @@
 #include "Styling/SlateIconFinder.h"
 #include "Types/MVVMBindingMode.h"
 #include "Widgets/SMVVMFieldEntry.h"
+#include "Widgets/SMVVMFieldIcon.h"
 #include "Widgets/SMVVMPropertyPath.h"
 #include "Widgets/SMVVMSourceEntry.h"
 #include "Widgets/Input/SCheckBox.h"
@@ -26,7 +27,6 @@
 #include "Widgets/Layout/SWidgetSwitcher.h"
 #include "Widgets/SMVVMCachedViewBindingPropertyPath.h"
 #include "Widgets/SMVVMCachedViewBindingConversionFunction.h"
-#include "Widgets/SMVVMViewModelBindingListWidget.h"
 #include "Widgets/Views/STreeView.h"
 
 #define LOCTEXT_NAMESPACE "MVVMFieldSelectorMenu"
@@ -65,9 +65,11 @@ void SFieldSelectorMenu::Construct(const FArguments& InArgs, const UWidgetBluepr
 	check(InWidgetBlueprint);
 
 	OnSelected = InArgs._OnSelected;
+	OnSelectedConversion = InArgs._OnSelectedConversion;
 	OnMenuCloseRequested = InArgs._OnMenuCloseRequested;
 	SelectionContext = InArgs._SelectionContext;
 	bCanCreateEvent = InArgs._CanCreateEvent;
+	bCanConvertField = InArgs._CanConvertField;
 
 	bIsMenuInitialized = false;
 	const bool bIsClearEnabled = InArgs._CurrentSelected.IsSet() && InArgs._CurrentSelected.GetValue().IsValid() && InArgs._IsClearEnableable;
@@ -77,42 +79,79 @@ void SFieldSelectorMenu::Construct(const FArguments& InArgs, const UWidgetBluepr
 
 	AssignableToProperty = SelectionContext.AssignableTo;
 
-	if (bShowConversionFunctions)
+	if (bShowConversionFunctions || bCanConvertField)
 	{
 		GenerateConversionFunctionItems();
 	}
 
 	TSharedRef<SVerticalBox> VBox = SNew(SVerticalBox)
-	+ SVerticalBox::Slot()
-	.HAlign(HAlign_Fill)
-	.VAlign(VAlign_Top)
-	.Padding(0, 4, 0, 4)
-	.AutoHeight()
-	[
-		SNew(SHorizontalBox)
-		+ SHorizontalBox::Slot()
-		.FillWidth(1.f)
-		.Padding(3.f, 0.f, 3.f, 0.f)	
-		[
-			SAssignNew(SearchBox, SSearchBox)
-			.OnTextChanged(this, &SFieldSelectorMenu::HandleSearchBoxTextChanged)
-		]
-		+ SHorizontalBox::Slot()
-		.HAlign(HAlign_Right)
-		.VAlign(VAlign_Center)
-		.AutoWidth()
+		+ SVerticalBox::Slot()
+		.HAlign(HAlign_Fill)
+		.VAlign(VAlign_Top)
+		.Padding(0, 4, 0, 4)
+		.AutoHeight()
 		[
-			SNew(SCheckBox)
-			.OnCheckStateChanged(this, &SFieldSelectorMenu::HandleEnabledContextToggleChanged)
-			.IsChecked(this, &SFieldSelectorMenu::ToggleEnabledContext)
-			.Visibility(InArgs._IsEnableContextToggleAvailable ? EVisibility::Visible : EVisibility::Collapsed)
-			.ToolTipText(LOCTEXT("ContextFlagToolTip", "Should the list be filtered to only properties and functions that make sense in the current context?"))
+			SNew(SHorizontalBox)
+			+ SHorizontalBox::Slot()
+			.FillWidth(1.f)
+			.Padding(3.f, 0.f, 3.f, 0.f)
 			[
-				SNew(STextBlock)
-				.Text(LOCTEXT("MVVMContextSensitiveToggle", "Context Sensitive"))
+				SAssignNew(SearchBox, SSearchBox)
+				.OnTextChanged(this, &SFieldSelectorMenu::HandleSearchBoxTextChanged)
+			]
+			+ SHorizontalBox::Slot()
+			.HAlign(HAlign_Right)
+			.VAlign(VAlign_Center)
+			.AutoWidth()
+			[
+				SNew(SCheckBox)
+				.OnCheckStateChanged(this, &SFieldSelectorMenu::HandleEnabledContextToggleChanged)
+				.IsChecked(this, &SFieldSelectorMenu::ToggleEnabledContext)
+				.Visibility(InArgs._IsEnableContextToggleAvailable ? EVisibility::Visible : EVisibility::Collapsed)
+				.ToolTipText(LOCTEXT("ContextFlagToolTip", "Should the list be filtered to only properties and functions that make sense in the current context?"))
+				[
+					SNew(STextBlock)
+					.Text(LOCTEXT("MVVMContextSensitiveToggle", "Context Sensitive"))
+				]
 			]
-		]
-	];
+		];
+
+	if (InArgs._ShowTooltipText)
+	{
+		VBox->AddSlot()
+			.AutoHeight()
+			.Padding(3.0, 5.0, 3.0, 5.0)
+			[
+				SNew(SHorizontalBox)
+				+ SHorizontalBox::Slot()
+				.HAlign(HAlign_Left)
+				.VAlign(VAlign_Center)
+				.AutoWidth()
+				[
+					SNew(STextBlock)
+					.Text(this, &SFieldSelectorMenu::GetTooltipText)
+				]
+				+ SHorizontalBox::Slot()
+				.HAlign(HAlign_Left)
+				.VAlign(VAlign_Center)
+				.AutoWidth()
+				[
+					SNew(SCachedViewBindingPropertyPath, WidgetBlueprint.Get())
+					.ShowContext(false)
+					.ShowFieldNotify(false)
+					.ShowOnlyLastPath(true)
+					.OnGetPropertyPath(this, &SFieldSelectorMenu::GetTooltipPropertyPath)
+				]
+				+ SHorizontalBox::Slot()
+				.HAlign(HAlign_Left)
+				.VAlign(VAlign_Center)
+				.AutoWidth()
+				[
+					SNew(STextBlock)
+					.Text(this, &SFieldSelectorMenu::GetAdditionalTooltipText)
+				]
+			];
+	}
 
 	VBox->AddSlot()
 		.FillHeight(1.0f)
@@ -151,48 +190,48 @@ void SFieldSelectorMenu::Construct(const FArguments& InArgs, const UWidgetBluepr
 		.AutoHeight()
 		[
 			SNew(SHorizontalBox)
+				+ SHorizontalBox::Slot()
+				.AutoWidth()
+				[
+					SAssignNew(EventSelectButton, SPrimaryButton)
+						.OnClicked(this, &SFieldSelectorMenu::HandleEventSelectClicked)
+						.Text(LOCTEXT("EventSelect", "Event"))
+						.ToolTipText(LOCTEXT("SelectorEventTooltip", "Create an event."))
+				]
+				+ SHorizontalBox::Slot()
+				.AutoWidth()
+				[
+					SNew(SPrimaryButton)
+						.OnClicked(this, &SFieldSelectorMenu::HandleSelectClicked)
+						.IsEnabled(this, &SFieldSelectorMenu::IsSelectEnabled)
+						.Text(LOCTEXT("Select", "Select"))
+						.ToolTipText(LOCTEXT("SelectorBindingTooltip", "Create a binding with the field."))
+				]
 			+ SHorizontalBox::Slot()
-			.AutoWidth()
-			[
-				SAssignNew(EventSelectButton, SPrimaryButton)
-				.OnClicked(this, &SFieldSelectorMenu::HandleEventSelectClicked)
-				.Text(LOCTEXT("EventSelect", "Event"))
-				.ToolTipText(LOCTEXT("SelectorEventTooltip", "Create an event.") )
-			]
-			+ SHorizontalBox::Slot()
-			.AutoWidth()
-			[
-				SNew(SPrimaryButton)
-				.OnClicked(this, &SFieldSelectorMenu::HandleSelectClicked)
-				.IsEnabled(this, &SFieldSelectorMenu::IsSelectEnabled)
-				.Text(LOCTEXT("Select", "Select"))
-				.ToolTipText(LOCTEXT("SelectorBindingTooltip", "Create a binding with the field."))
-			]
-			+ SHorizontalBox::Slot()
-			.Padding(FAppStyle::Get().GetMargin("StandardDialog.SlotPadding"))
-			.AutoWidth()
-			[
-				SNew(SButton)
-				.OnClicked(this, &SFieldSelectorMenu::HandleClearClicked)
-				.IsEnabled(bIsClearEnabled)
-				.HAlign(HAlign_Center)
+				.Padding(FAppStyle::Get().GetMargin("StandardDialog.SlotPadding"))
+				.AutoWidth()
 				[
-					SNew(STextBlock)
-					.Text(LOCTEXT("Clear", "Clear"))
+					SNew(SButton)
+						.OnClicked(this, &SFieldSelectorMenu::HandleClearClicked)
+						.IsEnabled(bIsClearEnabled)
+						.HAlign(HAlign_Center)
+						[
+							SNew(STextBlock)
+								.Text(LOCTEXT("Clear", "Clear"))
+						]
 				]
-			]
 			+ SHorizontalBox::Slot()
-			.Padding(FAppStyle::Get().GetMargin("StandardDialog.SlotPadding"))
-			.AutoWidth()
-			[
-				SNew(SButton)
-				.OnClicked(this, &SFieldSelectorMenu::HandleCancelClicked)
-				.HAlign(HAlign_Center)
+				.Padding(FAppStyle::Get().GetMargin("StandardDialog.SlotPadding"))
+				.AutoWidth()
 				[
-					SNew(STextBlock)
-					.Text(LOCTEXT("Cancel", "Cancel"))
+					SNew(SButton)
+						.OnClicked(this, &SFieldSelectorMenu::HandleCancelClicked)
+						.HAlign(HAlign_Center)
+						[
+							SNew(STextBlock)
+								.Text(LOCTEXT("Cancel", "Cancel"))
+						]
 				]
-			]
 		];
 
 	TSharedRef<SWidget> MenuWidget = SNew(SBox)
@@ -200,17 +239,17 @@ void SFieldSelectorMenu::Construct(const FArguments& InArgs, const UWidgetBluepr
 		.MinDesiredHeight(200.0f)
 		[
 			SNew(SBorder)
-			.BorderImage(FAppStyle::Get().GetBrush("ToolPanel.GroupBorder"))
-			.Padding(FMargin(8.0f, 2.0f, 8.0f, 3.0f))
-			[
-				VBox
-			]
+				.BorderImage(FAppStyle::Get().GetBrush("ToolPanel.GroupBorder"))
+				.Padding(FMargin(8.0f, 2.0f, 8.0f, 3.0f))
+				[
+					VBox
+				]
 		];
 
 	ChildSlot
-	[
-		MenuWidget
-	];
+		[
+			MenuWidget
+		];
 
 	bIsMenuInitialized = true;
 	UpdateSelection();
@@ -222,7 +261,7 @@ void SFieldSelectorMenu::Tick(const FGeometry& AllottedGeometry, const double In
 	TOptional<FMVVMLinkedPinValue> NewSelection = GetCurrentSelection();
 	if (NewSelection != CurrentSelectedValue)
 	{
-		CurrentSelectedValue = NewSelection;
+		SetSelectedValue(NewSelection);
 		UpdateSelection();
 	}
 }
@@ -230,7 +269,90 @@ void SFieldSelectorMenu::Tick(const FGeometry& AllottedGeometry, const double In
 void SFieldSelectorMenu::UpdateSelection()
 {
 	check(EventSelectButton);
-	EventSelectButton->SetVisibility(IsEventSelectEnabled()? EVisibility::Visible : EVisibility::Collapsed);
+	EventSelectButton->SetVisibility(IsEventSelectEnabled() ? EVisibility::Visible : EVisibility::Collapsed);
+
+	if(!bIsSelectingConversion && ConversionFunctionCategoryTree.IsValid())
+	{
+		// If users enabled auto-conversion, do not show the conversion function selection tree. Users will select conversion function when auto-converting
+		ConversionFunctionCategoryTree->SetVisibility(BindingList->IsConversionEnabled() ? EVisibility::Collapsed : EVisibility::Visible);
+	}
+
+	if (bIsSelectingConversion)
+	{
+		if (BindingList.IsValid())
+		{
+			BindingList->SetEnabled(false);
+			BindingList->ClearSources();
+		}
+
+		if (ViewModelList.IsValid())
+		{
+			ViewModelList->SetEnabled(false);
+			ViewModelList->ClearSelection();
+		}
+
+		if (WidgetList.IsValid())
+		{
+			WidgetList->SetEnabled(false);
+			WidgetList->ClearSelection();
+		}
+	}
+
+	Invalidate(EInvalidateWidgetReason::Layout);
+}
+
+void SFieldSelectorMenu::SetSelectedValue(const TOptional<FMVVMLinkedPinValue>& NewSelection)
+{
+	if (bIsSelectingConversion)
+	{
+		if (NewSelection.IsSet())
+		{
+			// Find the available conversion
+			for (const FMVVMFieldConversion& AvailableConversion : AvailableFieldConversions)
+			{
+				if (NewSelection->IsConversionFunction() && AvailableConversion.Function.IsFunction()
+					&& AvailableConversion.Function.GetFunction()->GetFName() == NewSelection->GetConversionFunction()->GetFName())
+				{
+					CurrentSelectedConversionValue = AvailableConversion;
+					break;
+				}
+				if (NewSelection->IsConversionNode() && AvailableConversion.Function.IsNode()
+					&& AvailableConversion.Function.GetNode()->GetFName() == NewSelection->GetConversionNode()->GetFName())
+				{
+					CurrentSelectedConversionValue = AvailableConversion;
+					break;
+				}
+			}
+		}
+	}
+	else
+	{
+		CurrentSelectedValue = NewSelection;
+	}
+}
+
+
+FText SFieldSelectorMenu::GetTooltipText() const
+{
+	if (CurrentSelectedValue.IsSet())
+	{
+		return LOCTEXT("FieldSelector_TooltipBindTo", "Bind to ");
+	}
+	return LOCTEXT("FieldSelector_Tooltip", "Select a property to bind to");
+}
+
+FText SFieldSelectorMenu::GetAdditionalTooltipText() const
+{
+	if (bIsSelectingConversion)
+	{
+		return LOCTEXT("FieldSelector_TooltipConversion", " with conversion function");
+	}
+	return FText::GetEmpty();
+}
+
+FMVVMBlueprintPropertyPath SFieldSelectorMenu::GetTooltipPropertyPath() const
+{
+	return CurrentSelectedValue.Get(FMVVMLinkedPinValue()).GetPropertyPath();
 }
 
 TSharedRef<SWidget> SFieldSelectorMenu::GetWidgetToFocus() const
@@ -268,6 +390,11 @@ TOptional<FMVVMLinkedPinValue> SFieldSelectorMenu::GetCurrentSelection() const
 
 bool SFieldSelectorMenu::IsSelectEnabled() const
 {
+	if (bIsSelectingConversion)
+	{
+		return CurrentSelectedConversionValue.IsSet();
+	}
+
 	return CurrentSelectedValue.IsSet();
 }
 
@@ -282,20 +409,19 @@ bool SFieldSelectorMenu::IsEventSelectEnabled() const
 
 FReply SFieldSelectorMenu::HandleSelectClicked()
 {
-	if (OnSelected.IsBound())
+	if (bCanConvertField)
 	{
-		OnSelected.Execute(CurrentSelectedValue.Get(FMVVMLinkedPinValue()), ESelectionType::Binding);
+		return PromptConversionSelection();
 	}
 
+	OnSelected.ExecuteIfBound(CurrentSelectedValue.Get(FMVVMLinkedPinValue()), ESelectionType::Binding);
+
 	return FReply::Handled();
 }
 
 FReply SFieldSelectorMenu::HandleEventSelectClicked()
 {
-	if (OnSelected.IsBound())
-	{
-		OnSelected.Execute(CurrentSelectedValue.Get(FMVVMLinkedPinValue()), ESelectionType::Event);
-	}
+	OnSelected.ExecuteIfBound(CurrentSelectedValue.Get(FMVVMLinkedPinValue()), ESelectionType::Event);
 
 	return FReply::Handled();
 }
@@ -322,9 +448,53 @@ FReply SFieldSelectorMenu::HandleCancelClicked()
 	return FReply::Handled();
 }
 
+
+FReply SFieldSelectorMenu::PromptConversionSelection()
+{
+	if (!bIsSelectingConversion)
+	{
+		check(BindingList);
+		AvailableFieldConversions = BindingList->GetSelectedConversions();
+
+		if (AvailableFieldConversions.Num() == 0)
+		{
+			OnSelectedConversion.ExecuteIfBound(CurrentSelectedValue.Get(FMVVMLinkedPinValue()), FMVVMFieldConversion(), ESelectionType::Binding);
+		}
+		else if (AvailableFieldConversions.Num() == 1)
+		{
+			CurrentSelectedConversionValue = AvailableFieldConversions[0];
+			OnSelectedConversion.ExecuteIfBound(CurrentSelectedValue.Get(FMVVMLinkedPinValue()), CurrentSelectedConversionValue.Get(FMVVMFieldConversion()), ESelectionType::Binding);
+		}
+		else if (AvailableFieldConversions.Num() > 1)
+		{
+			// If multiple are available, display conversion function list and wait for users to select.
+			bIsSelectingConversion = true;
+
+			UpdateSelection();
+			Algo::SortBy(ConversionFunctions, &FConversionFunctionValue::GetDisplayName, FText::FSortPredicate());
+			FilterConversionFunctions();
+
+			if (ConversionFunctionList.IsValid())
+			{
+				ConversionFunctionList->SetVisibility(EVisibility::Visible);
+				ConversionFunctionList->RequestListRefresh();
+			}
+		}
+
+		return FReply::Handled();
+	}
+
+	// If we are selecting conversion, the select button is disabled until a valid conversion is selected
+	bIsSelectingConversion = false;
+	AvailableFieldConversions.Empty();
+	OnSelectedConversion.ExecuteIfBound(CurrentSelectedValue.Get(FMVVMLinkedPinValue()), CurrentSelectedConversionValue.Get(FMVVMFieldConversion()), ESelectionType::Binding);
+
+	return FReply::Handled();
+}
+
 void SFieldSelectorMenu::SetPropertyPathSelection(const FMVVMBlueprintPropertyPath& SelectedPath)
 {
-	CurrentSelectedValue = FMVVMLinkedPinValue(SelectedPath);
+	SetSelectedValue(FMVVMLinkedPinValue(SelectedPath));
 	if (IsEventSelectEnabled())
 	{
 		HandleEventSelectClicked();
@@ -337,7 +507,7 @@ void SFieldSelectorMenu::SetPropertyPathSelection(const FMVVMBlueprintPropertyPa
 
 void SFieldSelectorMenu::SetConversionFunctionSelection(const FConversionFunctionValue SelectedFunction)
 {
-	CurrentSelectedValue = FMVVMLinkedPinValue(SelectedFunction);
+	SetSelectedValue(FMVVMLinkedPinValue(SelectedFunction));
 	if (IsEventSelectEnabled())
 	{
 		HandleEventSelectClicked();
@@ -922,6 +1092,12 @@ TSharedRef<SWidget> SFieldSelectorMenu::CreateBindingContextPanel(const FArgumen
 			}
 		}
 
+		// bCanConvertField handles visibility of this
+		if (!SelectionContext.bAllowConversionFunctions)
+		{
+			ConversionFunctionCategoryTree->SetVisibility(EVisibility::Collapsed);
+		}
+
 		StackedSourcePicker->AddSlot()
 			.AutoHeight()
 			[
@@ -939,7 +1115,8 @@ TSharedRef<SWidget> SFieldSelectorMenu::CreateBindingListPanel(const FArguments&
 		.OnDoubleClicked(this, &SFieldSelectorMenu::SetPropertyPathSelection)
 		.FieldVisibilityFlags(Private::GetFieldVisibilityFlags(SelectionContext.BindingMode, SelectionContext.bReadable, SelectionContext.bWritable))
 		.AssignableTo(AssignableToProperty)
-		.IsBindingToEvent(InArgs._IsBindingToEvent);
+		.IsBindingToEvent(InArgs._IsBindingToEvent)
+		.bEnableConversion(bCanConvertField);
 
 	const bool bValidSelectedProperty = InArgs._CurrentSelected.IsSet() && InArgs._CurrentSelected.GetValue().IsPropertyPath();
 
@@ -966,7 +1143,7 @@ TSharedRef<SWidget> SFieldSelectorMenu::CreateBindingListPanel(const FArguments&
 			BindingList.ToSharedRef()
 		];
 
-	if (SelectionContext.bAllowConversionFunctions)
+	if (SelectionContext.bAllowConversionFunctions || bCanConvertField)
 	{
 		const bool bValidSelectedFunction = InArgs._CurrentSelected.IsSet() && (InArgs._CurrentSelected.GetValue().IsConversionFunction() || InArgs._CurrentSelected.GetValue().IsConversionNode());
 
@@ -987,6 +1164,11 @@ TSharedRef<SWidget> SFieldSelectorMenu::CreateBindingListPanel(const FArguments&
 				: FConversionFunctionValue(InArgs._CurrentSelected.GetValue().GetConversionNode());
 			ConversionFunctionList->SetItemSelection(ConversionFunctionValue, true);
 		}
+
+		if (bCanConvertField)
+		{
+			ConversionFunctionList->SetVisibility(EVisibility::Collapsed);
+		}
 	}
 
 	return BindingListVBox;
@@ -1001,6 +1183,8 @@ void SFieldSelectorMenu::HandleEnabledContextToggleChanged(ECheckBoxState CheckS
 	FilterSettings.FilterFlags = CheckState == ECheckBoxState::Checked ? EFilterFlag::All : EFilterFlag::None;
 	ExtensionView->SetFilterSettings(FilterSettings);
 
+	BindingList->EnableConversion(bCanConvertField && FilterSettings.FilterFlags == EFilterFlag::All);
+
 	if (WidgetList.IsValid() && WidgetList->GetSelectedWidgets().Num() > 0)
 	{
 		HandleWidgetSelected(WidgetBlueprint->GetFName(), ESelectInfo::Direct);
@@ -1142,6 +1326,19 @@ void SFieldSelectorMenu::AddConversionFunctionChildrenRecursive(const TSharedPtr
 
 void SFieldSelectorMenu::FilterConversionFunctions()
 {
+	if (bIsSelectingConversion)
+	{
+		FilteredConversionFunctions.Reset();
+		for (const FMVVMFieldConversion& AvailableConversion : AvailableFieldConversions)
+		{
+			if (AvailableConversion.IsValid())
+			{
+				FilteredConversionFunctions.Add(AvailableConversion.Function);
+			}
+		}
+		return;
+	}
+
 	TArray<FString> FilterStrings;
 	if (SearchBox.IsValid())
 	{
diff --git a/Engine/Plugins/Runtime/ModelViewViewModel/Source/ModelViewViewModelEditor/Private/Widgets/SMVVMFieldSelectorMenu.h b/Engine/Plugins/Runtime/ModelViewViewModel/Source/ModelViewViewModelEditor/Private/Widgets/SMVVMFieldSelectorMenu.h
--- a/Engine/Plugins/Runtime/ModelViewViewModel/Source/ModelViewViewModelEditor/Private/Widgets/SMVVMFieldSelectorMenu.h
+++ b/Engine/Plugins/Runtime/ModelViewViewModel/Source/ModelViewViewModelEditor/Private/Widgets/SMVVMFieldSelectorMenu.h
@@ -9,6 +9,7 @@
 #include "Styling/CoreStyle.h"
 #include "Styling/SlateTypes.h"
 #include "Widgets/SCompoundWidget.h"
+#include "Widgets/SMVVMViewModelBindingListWidget.h"
 
 class SPrimaryButton;
 
@@ -17,9 +18,10 @@ class ITableRow;
 class SSearchBox;
 class SReadOnlyHierarchyView;
 class STableViewBase;
+
+class SPropertyPath;
 template <typename ItemType> class SListView;
 template <typename ItemType> class STreeView;
-namespace UE::MVVM {class SSourceBindingList; }
 
 class UWidgetBlueprint;
 
@@ -101,16 +103,26 @@ public:
 
 	DECLARE_DELEGATE_RetVal(FFieldSelectionContext, FOnGetFieldSelectionContext);
 	DECLARE_DELEGATE_TwoParams(FOnLinkedValueSelected, FMVVMLinkedPinValue, ESelectionType);
+	DECLARE_DELEGATE_ThreeParams(FOnLinkedValueConversionSelected, FMVVMLinkedPinValue, const FMVVMFieldConversion&, ESelectionType);
 
 	SLATE_BEGIN_ARGS(SFieldSelectorMenu){}
 		SLATE_ARGUMENT(TOptional<FMVVMLinkedPinValue>, CurrentSelected)
 		SLATE_EVENT(FOnLinkedValueSelected, OnSelected)
+		SLATE_EVENT(FOnLinkedValueConversionSelected, OnSelectedConversion)
 		SLATE_EVENT(FSimpleDelegate, OnMenuCloseRequested)
 		SLATE_ARGUMENT(FFieldSelectionContext, SelectionContext)
 		SLATE_ARGUMENT_DEFAULT(bool, IsBindingToEvent) { false };
 		SLATE_ARGUMENT_DEFAULT(bool, CanCreateEvent) { false };
 		SLATE_ARGUMENT_DEFAULT(bool, IsEnableContextToggleAvailable) { true };
 		SLATE_ARGUMENT_DEFAULT(bool, IsClearEnableable) { true };
+		SLATE_ARGUMENT_DEFAULT(bool, ShowTooltipText) { false };
+
+		/**
+		* If CanConvertField is set to true
+		* - Menu will display Fields that can be converted to the Selected Property.
+		* - Selections will be return via OnSelectedConversion
+		*/
+		SLATE_ARGUMENT_DEFAULT(bool, CanConvertField) { false };
 	SLATE_END_ARGS()
 
 	void Construct(const FArguments& InArgs, const UWidgetBlueprint* InWidgetBlueprint);
@@ -166,6 +178,11 @@ private:
 
 	TOptional<FMVVMLinkedPinValue> GetCurrentSelection() const;
 	void UpdateSelection();
+	void SetSelectedValue(const TOptional<FMVVMLinkedPinValue>& NewSelection);
+
+	FText GetTooltipText() const;
+	FText GetAdditionalTooltipText() const;
+	FMVVMBlueprintPropertyPath GetTooltipPropertyPath() const;
 
 	bool IsSelectEnabled() const;
 	bool IsEventSelectEnabled() const;
@@ -175,9 +192,13 @@ private:
 	FReply HandleClearClicked();
 	FReply HandleCancelClicked();
 
+	FReply PromptConversionSelection();
+
 private:
 	TWeakObjectPtr<const UWidgetBlueprint> WidgetBlueprint;
 	FOnLinkedValueSelected OnSelected;
+	FOnLinkedValueConversionSelected OnSelectedConversion;
+
 	FSimpleDelegate OnMenuCloseRequested;
 	FFieldSelectionContext SelectionContext;
 
@@ -204,11 +225,15 @@ private:
 	TSharedPtr<SListView<UE::MVVM::FConversionFunctionValue>> ConversionFunctionList;
 	TArray<UE::MVVM::FConversionFunctionValue> ConversionFunctions;
 	TArray<UE::MVVM::FConversionFunctionValue> FilteredConversionFunctions;
+	TArray<FMVVMFieldConversion> AvailableFieldConversions;
 
 	TOptional<FMVVMLinkedPinValue> CurrentSelectedValue;
+	TOptional<FMVVMFieldConversion> CurrentSelectedConversionValue;
 
 	bool bIsMenuInitialized = false;
-	bool bCanCreateEvent = false;
-}; 
+	bool bCanCreateEvent = false;	
+	bool bCanConvertField = false;
+	bool bIsSelectingConversion = false;
+};
 
 } // namespace UE::MVVM
diff --git a/Engine/Plugins/Runtime/ModelViewViewModel/Source/ModelViewViewModelEditor/Private/Widgets/SMVVMViewModelBindingListWidget.cpp b/Engine/Plugins/Runtime/ModelViewViewModel/Source/ModelViewViewModelEditor/Private/Widgets/SMVVMViewModelBindingListWidget.cpp
--- a/Engine/Plugins/Runtime/ModelViewViewModel/Source/ModelViewViewModelEditor/Private/Widgets/SMVVMViewModelBindingListWidget.cpp
+++ b/Engine/Plugins/Runtime/ModelViewViewModel/Source/ModelViewViewModelEditor/Private/Widgets/SMVVMViewModelBindingListWidget.cpp
@@ -7,9 +7,12 @@
 #include "Blueprint/WidgetBlueprintGeneratedClass.h"
 #include "Blueprint/WidgetTree.h"
 #include "BlueprintEditorSettings.h"
+#include "Editor/EditorEngine.h"
 #include "Extensions/UIComponentContainer.h"
 #include "GraphEditorSettings.h"
+#include "MVVMBlueprintPin.h"
 #include "MVVMBlueprintViewModelContext.h"
+#include "MVVMEditorSubsystem.h"
 #include "MVVMDeveloperProjectSettings.h"
 #include "MVVMSubsystem.h"
 #include "MVVMWidgetBlueprintExtension_View.h"
@@ -17,13 +20,17 @@
 #include "Styling/MVVMEditorStyle.h"
 #include "Types/MVVMAvailableBinding.h"
 #include "Types/MVVMBindingSource.h"
+#include "Types/MVVMConversionFunctionValue.h"
 #include "UIComponentWidgetBlueprintExtension.h"
 #include "WidgetBlueprint.h"
 #include "Widgets/Images/SImage.h"
 #include "Widgets/Layout/SBox.h"
 
+
 #define LOCTEXT_NAMESPACE "SSourceBindingList"
 
+extern UNREALED_API UEditorEngine* GEditor;
+
 using UE::PropertyViewer::SPropertyViewer;
 
 namespace UE::MVVM
@@ -31,26 +38,60 @@ namespace UE::MVVM
 
 namespace Private
 {
-	TOptional<FFieldVariant> PassFilter(const UBlueprint* Blueprint, const FMVVMAvailableBinding& Binding, const UStruct* Struct, const FMVVMFieldVariant& FieldVariant, EFieldVisibility FieldVisibilityFlags, const FProperty* AssignableTo, bool bDoObjectProperty, bool bAllowComponentProperty = false);
-	TOptional<FFieldVariant> PassFilter(const UBlueprint* Blueprint, const FMVVMAvailableBinding& Binding, const UStruct* Struct, EFieldVisibility FieldVisibilityFlags, const FProperty* AssignableTo, bool bDoObjectProperty, bool bAllowComponentProperty = false);
+	struct FPassFilterInputs
+	{
+		FPassFilterInputs(
+			const UWidgetBlueprint* WidgetBlueprint,
+			const FMVVMAvailableBinding Binding,
+			const UStruct* Struct,
+			EFieldVisibility FieldVisibilityFlags,
+			const FProperty* AssignableTo,
+			bool bDoObjectProperty,
+			bool bAllowComponentProperty,
+			bool bCanConvert)
+			: WidgetBlueprint(WidgetBlueprint)
+			, Binding(Binding)
+			, Struct(Struct)
+			, FieldVisibilityFlags(FieldVisibilityFlags)
+			, AssignableTo(AssignableTo)
+			, bDoObjectProperty(bDoObjectProperty)
+			, bAllowComponentProperty(bAllowComponentProperty)
+			, bCanConvert(bCanConvert)
+		{}
+
+		const UWidgetBlueprint* WidgetBlueprint;
+		const FMVVMAvailableBinding Binding;
+		const UStruct* Struct;
+		EFieldVisibility FieldVisibilityFlags;
+		const FProperty* AssignableTo;
+		bool bDoObjectProperty;
+		bool bAllowComponentProperty;
+		bool bCanConvert;
+	};
 
-	TOptional<FFieldVariant> PassFilter(const UBlueprint* Blueprint, const FMVVMAvailableBinding& Binding, const UStruct* Struct, EFieldVisibility FieldVisibilityFlags, const FProperty* AssignableTo, bool bDoObjectProperty, bool bAllowComponentProperty)
+	TOptional<FFieldVariant> PassFilter(const FPassFilterInputs& Inputs, const FMVVMFieldVariant& FieldVariant);
+	TOptional<FFieldVariant> PassFilter(const FPassFilterInputs& Inputs);
+	bool PassFilter_ConversionFunction(const FPassFilterInputs& Inputs, const FProperty* Source);
+
+	TArray<FMVVMFieldConversion> GetCompatibleSourceConversions(const UWidgetBlueprint* WidgetBlueprint, const FProperty* Source, const FProperty* Destination, bool bReturnFirst);
+
+	TOptional<FFieldVariant> PassFilter(const FPassFilterInputs& Inputs)
 	{
-		if (EnumHasAllFlags(FieldVisibilityFlags, EFieldVisibility::Readable) && !Binding.IsReadable())
+		if (EnumHasAllFlags(Inputs.FieldVisibilityFlags, EFieldVisibility::Readable) && !Inputs.Binding.IsReadable())
 		{
 			return TOptional<FFieldVariant>();
 		}
 
-		if (EnumHasAllFlags(FieldVisibilityFlags, EFieldVisibility::Notify | EFieldVisibility::Readable) && !Binding.HasNotify())
+		if (EnumHasAllFlags(Inputs.FieldVisibilityFlags, EFieldVisibility::Notify | EFieldVisibility::Readable) && !Inputs.Binding.HasNotify())
 		{
 			return TOptional<FFieldVariant>();
 		}
 
-		FMVVMFieldVariant FieldVariant = BindingHelper::FindFieldByName(Struct, Binding.GetBindingName());
-		return PassFilter(Blueprint, Binding, Struct, FieldVariant, FieldVisibilityFlags, AssignableTo, bDoObjectProperty, bAllowComponentProperty);
+		FMVVMFieldVariant FieldVariant = BindingHelper::FindFieldByName(Inputs.Struct, Inputs.Binding.GetBindingName());
+		return PassFilter(Inputs, FieldVariant);
 	}
 
-	TOptional<FFieldVariant> PassFilter(const UBlueprint* Blueprint, const FMVVMAvailableBinding& Binding, const UStruct* Struct, const FMVVMFieldVariant& FieldVariant, EFieldVisibility FieldVisibilityFlags, const FProperty* AssignableTo, bool bDoObjectProperty, bool bAllowComponentProperty)
+	TOptional<FFieldVariant> PassFilter(const FPassFilterInputs& Inputs, const FMVVMFieldVariant& FieldVariant)
 	{
 		static FName NAME_Hidden("Hidden");
 		if (ensure(!FieldVariant.IsEmpty()))
@@ -63,19 +104,24 @@ namespace Private
 					return TOptional<FFieldVariant>();
 				}
 
+				if (!GetDefault<UMVVMDeveloperProjectSettings>()->IsFunctionAllowed(Inputs.WidgetBlueprint, Cast<const UClass>(Inputs.Struct), Function))
+				{
+					return TOptional<FFieldVariant>();
+				}
+
 				const FProperty* ReturnProperty = BindingHelper::GetReturnProperty(Function);
-				bool bDoCompatibleTest = ReturnProperty == nullptr || AssignableTo != nullptr;
+				bool bDoCompatibleTest = ReturnProperty == nullptr || Inputs.AssignableTo != nullptr;
 				bool bDoWritableTest = true;
 				// Do we allow walking up the tree
 				if (CastField<FObjectPropertyBase>(ReturnProperty))
 				{
-					bDoCompatibleTest = bDoCompatibleTest && bDoObjectProperty;
-					bDoWritableTest = bDoObjectProperty;
+					bDoCompatibleTest = bDoCompatibleTest && Inputs.bDoObjectProperty;
+					bDoWritableTest = Inputs.bDoObjectProperty;
 				}
 
 				if (bDoWritableTest)
 				{
-					if (EnumHasAllFlags(FieldVisibilityFlags, EFieldVisibility::Writable) && !Binding.IsWritable())
+					if (EnumHasAllFlags(Inputs.FieldVisibilityFlags, EFieldVisibility::Writable) && !Inputs.Binding.IsWritable())
 					{
 						return TOptional<FFieldVariant>();
 					}
@@ -83,27 +129,25 @@ namespace Private
 
 				if (bDoCompatibleTest)
 				{
-					if (EnumHasAllFlags(FieldVisibilityFlags, EFieldVisibility::Writable))
+					if (EnumHasAllFlags(Inputs.FieldVisibilityFlags, EFieldVisibility::Writable))
 					{
-						if (AssignableTo && !BindingHelper::ArePropertiesCompatible(AssignableTo, BindingHelper::GetFirstArgumentProperty(Function)))
+						if (Inputs.AssignableTo && !BindingHelper::ArePropertiesCompatible(Inputs.AssignableTo, BindingHelper::GetFirstArgumentProperty(Function)))
 						{
-							return TOptional<FFieldVariant>();
+							if (Inputs.bCanConvert && PassFilter_ConversionFunction(Inputs, ReturnProperty))
+							{
+								return FFieldVariant(ReturnProperty);
+							}
 						}
 					}
-					if (EnumHasAllFlags(FieldVisibilityFlags, EFieldVisibility::Readable))
+					if (EnumHasAllFlags(Inputs.FieldVisibilityFlags, EFieldVisibility::Readable))
 					{
-						if (AssignableTo && !BindingHelper::ArePropertiesCompatible(ReturnProperty, AssignableTo))
+						if (Inputs.AssignableTo && !BindingHelper::ArePropertiesCompatible(ReturnProperty, Inputs.AssignableTo))
 						{
 							return TOptional<FFieldVariant>();
 						}
 					}
 				}
 
-				if (!GetDefault<UMVVMDeveloperProjectSettings>()->IsFunctionAllowed(Blueprint, Cast<const UClass>(Struct), Function))
-				{
-					return TOptional<FFieldVariant>();
-				}
-
 				return FFieldVariant(Function);
 			}
 			else if (FieldVariant.IsProperty())
@@ -114,19 +158,24 @@ namespace Private
 					return TOptional<FFieldVariant>();
 				}
 
-				bool bDoCompatibleTest = AssignableTo != nullptr;
+				if (!GetDefault<UMVVMDeveloperProjectSettings>()->IsPropertyAllowed(Inputs.WidgetBlueprint, Inputs.Struct, Property))
+				{
+					return TOptional<FFieldVariant>();
+				}
+
+				bool bDoCompatibleTest = Inputs.AssignableTo != nullptr;
 				bool bDoWritableTest = true;
 				// Do we allow walking up the tree
 				if (CastField<FObjectPropertyBase>(Property))
 				{
-					bDoCompatibleTest = bDoCompatibleTest && bDoObjectProperty;
-					bDoWritableTest = bDoObjectProperty;
+					bDoCompatibleTest = bDoCompatibleTest && Inputs.bDoObjectProperty;
+					bDoWritableTest = Inputs.bDoObjectProperty;
 				}
 
 				// If the path ends with the object property, then it needs to follow the writable rule
 				if (bDoWritableTest)
 				{
-					if (EnumHasAllFlags(FieldVisibilityFlags, EFieldVisibility::Writable) && !Binding.IsWritable())
+					if (EnumHasAllFlags(Inputs.FieldVisibilityFlags, EFieldVisibility::Writable) && !Inputs.Binding.IsWritable())
 					{
 						return TOptional<FFieldVariant>();
 					}
@@ -134,16 +183,19 @@ namespace Private
 
 				if (bDoCompatibleTest)
 				{
-					if (EnumHasAllFlags(FieldVisibilityFlags, EFieldVisibility::Writable))
+					if (EnumHasAllFlags(Inputs.FieldVisibilityFlags, EFieldVisibility::Writable))
 					{
-						if (!BindingHelper::ArePropertiesCompatible(AssignableTo, Property))
+						if (!BindingHelper::ArePropertiesCompatible(Inputs.AssignableTo, Property))
 						{
-							return TOptional<FFieldVariant>();
+							if (Inputs.bCanConvert && PassFilter_ConversionFunction(Inputs, Property))
+							{
+								return FFieldVariant(Property);
+							}
 						}
 					}
-					if (EnumHasAllFlags(FieldVisibilityFlags, EFieldVisibility::Readable))
+					if (EnumHasAllFlags(Inputs.FieldVisibilityFlags, EFieldVisibility::Readable))
 					{
-						if (!BindingHelper::ArePropertiesCompatible(Property, AssignableTo))
+						if (!BindingHelper::ArePropertiesCompatible(Property, Inputs.AssignableTo))
 						{
 							return TOptional<FFieldVariant>();
 						}
@@ -152,12 +204,7 @@ namespace Private
 
 				// Hide the component property on the userwidget. It should only be visible from the owning widget.
 				// The metadata helps us determine whether this is a component property generated by the compiler or not. We want to filter only those ones when bAllowComponentProperty is false.
-				if (Property->HasMetaData(UUIComponentWidgetBlueprintExtension::MD_ComponentVariable) && !bAllowComponentProperty)
-				{
-					return TOptional<FFieldVariant>();
-				}
-
-				if (!GetDefault<UMVVMDeveloperProjectSettings>()->IsPropertyAllowed(Blueprint, Struct, Property))
+				if (Property->HasMetaData(UUIComponentWidgetBlueprintExtension::MD_ComponentVariable) && !Inputs.bAllowComponentProperty)
 				{
 					return TOptional<FFieldVariant>();
 				}
@@ -168,13 +215,51 @@ namespace Private
 
 		return TOptional<FFieldVariant>();
 	}
+
+	bool PassFilter_ConversionFunction(const FPassFilterInputs& Inputs, const FProperty* Source)
+	{
+		const TArray<FMVVMFieldConversion>& Conversions = GetCompatibleSourceConversions(Inputs.WidgetBlueprint, Source, Inputs.AssignableTo, true);
+		return Conversions.Num() > 0;
+	}
+
+	TArray<FMVVMFieldConversion> GetCompatibleSourceConversions(const UWidgetBlueprint* WidgetBlueprint, const FProperty* Source, const FProperty* Destination, bool bReturnFirst)
+	{
+		TArray<FMVVMFieldConversion> OutConversions;
+
+		// No conversion necessary
+		if (BindingHelper::ArePropertiesCompatible(Source, Destination))
+		{
+			return OutConversions;
+		}
+
+		if (UMVVMEditorSubsystem* Subsystem = GEditor->GetEditorSubsystem<UMVVMEditorSubsystem>())
+		{
+			TArray<FConversionFunctionValue> ConversionFunctions = Subsystem->GetConversionFunctions(WidgetBlueprint, Source, Destination);
+			for (const FConversionFunctionValue& ConversionFunction : ConversionFunctions)
+			{
+				TArray<FMVVMBlueprintPinId> ArgumentPins = Subsystem->GetConversionFunctionPins(WidgetBlueprint, ConversionFunction, Source);
+					// TODO: Find first argument that matches type
+				if (ArgumentPins.Num() == 1) // TODO: Should we pick first available parameter that matched type? Or skip?
+				{
+					FMVVMFieldConversion Conversion(ConversionFunction, ArgumentPins[0]);
+					OutConversions.Add(Conversion);
+					if (bReturnFirst)
+					{
+						break;
+					}
+				}
+			}
+		}
+		return OutConversions;
+	}
 }
 
-FFieldIterator_Bindable::FFieldIterator_Bindable(const UWidgetBlueprint* InWidgetBlueprint, EFieldVisibility InVisibilityFlags, const FProperty* InAssignableTo, const bool InIsBindingToEvent) :
+FFieldIterator_Bindable::FFieldIterator_Bindable(const UWidgetBlueprint* InWidgetBlueprint, EFieldVisibility InVisibilityFlags, const FProperty* InAssignableTo, const bool InIsBindingToEvent, const bool bEnableConversion) :
 	WidgetBlueprint(InWidgetBlueprint),
 	FieldVisibilityFlags(InVisibilityFlags),
 	AssignableTo(InAssignableTo),
-	bIsBindingToEvent(InIsBindingToEvent)
+	bIsBindingToEvent(InIsBindingToEvent),
+	bEnableConversion(bEnableConversion)
 {
 }
 
@@ -201,11 +286,13 @@ TArray<FFieldVariant> FFieldIterator_Bindable::GetFields(const UStruct* Struct,
 			TOptional<FFieldVariant> PassResult;
 			if (FilterFlags == EFilterFlag::All)
 			{
-				PassResult = Private::PassFilter(WidgetBlueprint.Get(), Value, Struct, FieldVisibilityFlags, AssignableTo, false, bAllowComponentProperty);
+				const Private::FPassFilterInputs Params(WidgetBlueprint.Get(), Value, Struct, FieldVisibilityFlags, AssignableTo, false, bAllowComponentProperty, bEnableConversion);
+				PassResult = Private::PassFilter(Params);
 			}
 			else
 			{
-				PassResult = Private::PassFilter(WidgetBlueprint.Get(), Value, Struct, EFieldVisibility::None, nullptr, false, bAllowComponentProperty);
+				const Private::FPassFilterInputs Params(WidgetBlueprint.Get(), Value, Struct, EFieldVisibility::None, nullptr, false, bAllowComponentProperty, bEnableConversion);
+				PassResult = Private::PassFilter(Params);
 			}
 			if (PassResult.IsSet())
 			{
@@ -416,11 +503,41 @@ TSharedRef<SWidget> ConstructFieldPreSlot(const UWidgetBlueprint* WidgetBlueprin
 		];
 }
 
+/** */
+
+FMVVMFieldConversion::FMVVMFieldConversion(FConversionFunctionValue Function, FMVVMBlueprintPinId ParameterId)
+	: Function(Function), ParameterId(ParameterId)
+{
+}
+
+bool FMVVMFieldConversion::IsValid() const
+{
+	return Function.IsValid() && ParameterId.IsValid(); 
+}
+
+FName FMVVMFieldConversion::GetFName() const
+{ 
+	return Function.GetFName(); 
+}
+
+FMVVMBlueprintFunctionReference FMVVMFieldConversion::ToBlueprintFunctionReference(UBlueprint* Blueprint) const
+{
+	if (Function.IsFunction())
+	{
+		return FMVVMBlueprintFunctionReference(Blueprint, Function.GetFunction());
+	}
+	else if (Function.IsNode())
+	{
+		return FMVVMBlueprintFunctionReference(Function.GetNode());
+	}
+	return FMVVMBlueprintFunctionReference();
+}
+
 /** */
 void SSourceBindingList::Construct(const FArguments& InArgs, const UWidgetBlueprint* InWidgetBlueprint)
 {
 	WidgetBlueprint = InWidgetBlueprint;
-	FieldIterator = MakeUnique<FFieldIterator_Bindable>(InWidgetBlueprint, InArgs._FieldVisibilityFlags, InArgs._AssignableTo, InArgs._IsBindingToEvent);
+	FieldIterator = MakeUnique<FFieldIterator_Bindable>(InWidgetBlueprint, InArgs._FieldVisibilityFlags, InArgs._AssignableTo, InArgs._IsBindingToEvent, InArgs._bEnableConversion);
 	FieldExpander = MakeUnique<FFieldExpander_Bindable>();
 
 	bIsBindingToEvent = InArgs._IsBindingToEvent;
@@ -700,11 +817,13 @@ FMVVMBlueprintPropertyPath SSourceBindingList::CreateBlueprintPropertyPath(SProp
 
 					if (FilterFlags == EFilterFlag::All)
 					{
-						bPassFilter = Private::PassFilter(WidgetBlueprint.Get(), Binding, OwnerClass, FieldIterator->GetFieldVisibilityFlags(), FieldIterator->GetAssignableTo(), true).IsSet();
+						const Private::FPassFilterInputs Params(WidgetBlueprint.Get(), Binding, OwnerClass, FieldIterator->GetFieldVisibilityFlags(), FieldIterator->GetAssignableTo(), true, false, IsConversionEnabled());
+						bPassFilter = Private::PassFilter(Params).IsSet();
 					}
 					else
 					{
-						bPassFilter = Private::PassFilter(WidgetBlueprint.Get(), Binding, OwnerClass, EFieldVisibility::None, nullptr, true).IsSet();
+						const Private::FPassFilterInputs Params(WidgetBlueprint.Get(), Binding, OwnerClass, EFieldVisibility::None, nullptr, true, false, IsConversionEnabled());
+						bPassFilter = Private::PassFilter(Params).IsSet();
 					}
 				}
 				break;
@@ -776,11 +895,13 @@ FMVVMBlueprintPropertyPath SSourceBindingList::CreateBlueprintPropertyPath(SProp
 
 			if (FilterFlags == EFilterFlag::All)
 			{
-				bPassFilter = Private::PassFilter(WidgetBlueprint.Get(), Binding, AccessorClass, FieldIterator->GetFieldVisibilityFlags(), FieldIterator->GetAssignableTo(), true).IsSet();
+				const Private::FPassFilterInputs Params(WidgetBlueprint.Get(), Binding, AccessorClass, FieldIterator->GetFieldVisibilityFlags(), FieldIterator->GetAssignableTo(), true, false, IsConversionEnabled());
+				bPassFilter = Private::PassFilter(Params).IsSet();
 			}
 			else
 			{
-				bPassFilter = Private::PassFilter(WidgetBlueprint.Get(), Binding, AccessorClass, EFieldVisibility::None, nullptr, true).IsSet();
+				const Private::FPassFilterInputs Params(WidgetBlueprint.Get(), Binding, AccessorClass, EFieldVisibility::None, nullptr, true, false, IsConversionEnabled());
+				bPassFilter = Private::PassFilter(Params).IsSet();
 			}
 			if (bPassFilter)
 			{
@@ -823,6 +944,26 @@ void SSourceBindingList::SetRawFilterText(const FText& InFilterText)
 	}
 }
 
+void SSourceBindingList::EnableConversion(bool InEnabled)
+{
+	if (FieldIterator)
+	{
+		FieldIterator->EnableConversion(InEnabled);
+	}
+
+	Invalidate(EInvalidateWidgetReason::Layout);
+}
+
+bool SSourceBindingList::IsConversionEnabled() const
+{
+	if (FieldIterator)
+	{
+		return FieldIterator->IsConversionEnabled();
+	}
+
+	return false;
+}
+
 FMVVMBlueprintPropertyPath SSourceBindingList::GetSelectedProperty() const
 {
 	return SelectedPath;
@@ -875,6 +1016,46 @@ void SSourceBindingList::SetSelectedProperty(const FMVVMBlueprintPropertyPath& P
 	PropertyViewer->SetSelection(SelectedHandle, FieldPath);
 }
 
+TArray<FMVVMFieldConversion> SSourceBindingList::GetSelectedConversions() const
+{
+	TArray<FMVVMFieldConversion> SelectedConversions;
+
+	if (!PropertyViewer.IsValid())
+	{
+		return SelectedConversions;
+	}
+
+	const UWidgetBlueprint* WidgetBlueprintPtr = WidgetBlueprint.Get();
+	if (!WidgetBlueprintPtr)
+	{
+		return SelectedConversions;
+	}
+
+	TArray<SPropertyViewer::FSelectedItem> SelectedItems = PropertyViewer->GetSelectedItems();
+	for (const SPropertyViewer::FSelectedItem& SelectedItem : SelectedItems)
+	{
+		if (SelectedItem.Fields.IsEmpty())
+		{
+			continue;
+		}
+
+		for(const TArray<FFieldVariant>& Fields : SelectedItem.Fields)
+		{
+			if (Fields.IsEmpty())
+			{
+				continue;
+			}
+
+			if (FProperty* Property = CastField<FProperty>(Fields.Last().ToField()))
+			{
+				TArray<FMVVMFieldConversion> Conversions = Private::GetCompatibleSourceConversions(WidgetBlueprintPtr, Property, FieldIterator->GetAssignableTo(), false);
+				SelectedConversions.Append(MoveTemp(Conversions));
+			}
+		}
+	}
+	return SelectedConversions;
+}
+
 } // namespace UE::MVVM
 
 #undef LOCTEXT_NAMESPACE
diff --git a/Engine/Plugins/Runtime/ModelViewViewModel/Source/ModelViewViewModelEditor/Private/Widgets/SMVVMViewModelBindingListWidget.h b/Engine/Plugins/Runtime/ModelViewViewModel/Source/ModelViewViewModelEditor/Private/Widgets/SMVVMViewModelBindingListWidget.h
--- a/Engine/Plugins/Runtime/ModelViewViewModel/Source/ModelViewViewModelEditor/Private/Widgets/SMVVMViewModelBindingListWidget.h
+++ b/Engine/Plugins/Runtime/ModelViewViewModel/Source/ModelViewViewModelEditor/Private/Widgets/SMVVMViewModelBindingListWidget.h
@@ -4,7 +4,10 @@
 
 #include "Framework/PropertyViewer/IFieldExpander.h"
 #include "Framework/PropertyViewer/IFieldIterator.h"
+#include "MVVMBlueprintPin.h"
+#include "MVVMBlueprintFunctionReference.h"
 #include "MVVMPropertyPath.h"
+#include "Types/MVVMConversionFunctionValue.h"
 #include "Widgets/PropertyViewer/SPropertyViewer.h"
 
 namespace UE::MVVM { struct FBindingSource; }
@@ -33,7 +36,7 @@ namespace UE::MVVM
 	class FFieldIterator_Bindable : public UE::PropertyViewer::FFieldIterator_BlueprintVisible
 	{
 	public:
-		FFieldIterator_Bindable(const UWidgetBlueprint* WidgetBlueprint, EFieldVisibility InVisibilityFlags, const FProperty* InAssignableTo = nullptr, const bool InIsBindingToEvent = false);
+		FFieldIterator_Bindable(const UWidgetBlueprint* WidgetBlueprint, EFieldVisibility InVisibilityFlags, const FProperty* InAssignableTo = nullptr, const bool InIsBindingToEvent = false, const bool bEnableConversion = false);
 		virtual TArray<FFieldVariant> GetFields(const UStruct* Struct, const FName FieldName, const UStruct* Container) const override;
 		EFieldVisibility GetFieldVisibilityFlags() const
 		{
@@ -43,12 +46,20 @@ namespace UE::MVVM
 		{
 			return AssignableTo;
 		}
-
+		bool IsConversionEnabled() const
+		{
+			return bEnableConversion;
+		}
+		void EnableConversion(bool InEnabled)
+		{
+			bEnableConversion = InEnabled;
+		}
 	private:
 		TWeakObjectPtr<const UWidgetBlueprint> WidgetBlueprint; 
 		EFieldVisibility FieldVisibilityFlags = EFieldVisibility::All;
 		const FProperty* AssignableTo = nullptr;
 		bool bIsBindingToEvent = false;
+		bool bEnableConversion = false;
 	};
 
 	/**
@@ -66,6 +77,19 @@ namespace UE::MVVM
 	/** Genereate the icon for the binding type. */
 	TSharedRef<SWidget> ConstructFieldPreSlot(const UWidgetBlueprint* WidgetBlueprint, UE::PropertyViewer::SPropertyViewer::FHandle Handle, const FFieldVariant FieldPath, const bool bIsForEvent = false);
 
+	struct FMVVMFieldConversion : TSharedFromThis<FMVVMFieldConversion>
+	{
+		FMVVMFieldConversion() = default;
+		FMVVMFieldConversion(FConversionFunctionValue Function, FMVVMBlueprintPinId ParameterId);
+
+		FConversionFunctionValue Function;
+		FMVVMBlueprintPinId ParameterId;
+
+		bool IsValid() const; 
+		FName GetFName() const;
+		FMVVMBlueprintFunctionReference ToBlueprintFunctionReference(UBlueprint* Blueprint) const;
+	};
+
 	/** 
 	 * 
 	 */
@@ -87,6 +111,12 @@ namespace UE::MVVM
 			 */
 			SLATE_ARGUMENT_DEFAULT(const FProperty*, AssignableTo) = nullptr;
 
+			/**
+			* If enabled, will show properties that can be converted to AssignableTo. 
+			* The compatible conversion functions will be accessible via GetSelectedConversions
+			*/
+			SLATE_ARGUMENT_DEFAULT(bool, bEnableConversion) = false;
+
 			SLATE_EVENT(FOnDoubleClicked, OnDoubleClicked)
 		SLATE_END_ARGS()
 
@@ -104,8 +134,13 @@ namespace UE::MVVM
 		FMVVMBlueprintPropertyPath GetSelectedProperty() const; 
 		void SetSelectedProperty(const FMVVMBlueprintPropertyPath& Property);
 
+		TArray<FMVVMFieldConversion> GetSelectedConversions() const;
+
 		void SetRawFilterText(const FText& InFilterText);
 
+		void EnableConversion(bool InEnabled);
+		bool IsConversionEnabled() const;
+
 	private:
 		using SPropertyViewer = UE::PropertyViewer::SPropertyViewer;
 
diff --git a/Engine/Plugins/Runtime/ModelViewViewModel/Source/ModelViewViewModelEditor/Public/MVVMEditorSubsystem.h b/Engine/Plugins/Runtime/ModelViewViewModel/Source/ModelViewViewModelEditor/Public/MVVMEditorSubsystem.h
--- a/Engine/Plugins/Runtime/ModelViewViewModel/Source/ModelViewViewModelEditor/Public/MVVMEditorSubsystem.h
+++ b/Engine/Plugins/Runtime/ModelViewViewModel/Source/ModelViewViewModelEditor/Public/MVVMEditorSubsystem.h
@@ -139,6 +139,7 @@ public:
 	UE_API TArray<UFunction*> GetAvailableConversionFunctions(const UWidgetBlueprint* WidgetBlueprint, const FMVVMBlueprintPropertyPath& Source, const FMVVMBlueprintPropertyPath& Destination) const;
 
 	UE_API TArray<UE::MVVM::FConversionFunctionValue> GetConversionFunctions(const UWidgetBlueprint* WidgetBlueprint, const FProperty* ExpectedArgumentType, const FProperty* ExptectedReturnType) const;
+	UE_API TArray<FMVVMBlueprintPinId> GetConversionFunctionPins(const UWidgetBlueprint* WidgetBlueprint, const UE::MVVM::FConversionFunctionValue& Function, const FProperty* ExpectedArgumentType) const;
 
 	UE_API FMVVMBlueprintPropertyPath GetPathForConversionFunctionArgument(const UWidgetBlueprint* WidgetBlueprint, const FMVVMBlueprintViewBinding& Binding, const FMVVMBlueprintPinId& PinId, bool bSourceToDestination) const;
 	UE_API void SetPathForConversionFunctionArgument(UWidgetBlueprint* WidgetBlueprint, FMVVMBlueprintViewBinding& Binding, const FMVVMBlueprintPinId& PinId, const FMVVMBlueprintPropertyPath& Path, bool bSourceToDestination) const;
