MVVM: Auto Conversion, resolve missing input pins on function nodes.

diff --git a/Engine/Plugins/Runtime/ModelViewViewModel/Source/ModelViewViewModelEditor/Private/MVVMEditorSubsystem.cpp b/Engine/Plugins/Runtime/ModelViewViewModel/Source/ModelViewViewModelEditor/Private/MVVMEditorSubsystem.cpp
--- a/Engine/Plugins/Runtime/ModelViewViewModel/Source/ModelViewViewModelEditor/Private/MVVMEditorSubsystem.cpp
+++ b/Engine/Plugins/Runtime/ModelViewViewModel/Source/ModelViewViewModelEditor/Private/MVVMEditorSubsystem.cpp
@@ -1421,7 +1421,16 @@ TArray<FMVVMBlueprintPinId> UMVVMEditorSubsystem::GetConversionFunctionPins(cons
 			return OutPins;
 		}
 
-		TArray<UEdGraphPin*> InputPins = UE::MVVM::ConversionFunctionHelper::FindInputPins(Function.GetNode().GetDefaultObject());
+		UEdGraph* TempGraph = NewObject<UEdGraph>(GetTransientPackage(), UEdGraph::StaticClass());
+		TempGraph->Schema = UEdGraphSchema_K2::StaticClass();
+
+		// Setup input pins
+		UK2Node* Node = NewObject<UK2Node>(TempGraph, Function.GetNode());
+		Node->CreateNewGuid();
+		Node->PostPlacedNewNode();
+		Node->AllocateDefaultPins();
+
+		TArray<UEdGraphPin*> InputPins = UE::MVVM::ConversionFunctionHelper::FindInputPins(Node);
 		if (InputPins.Num() == 0)
 		{
 			return OutPins;
@@ -1435,7 +1444,8 @@ TArray<FMVVMBlueprintPinId> UMVVMEditorSubsystem::GetConversionFunctionPins(cons
 			bool bIgnoreArray = true;
 			if (GetDefault<UEdGraphSchema_K2>()->ArePinTypesCompatible(SourcePinType, ArgumentPin->PinType, CallingContext, bIgnoreArray))
 			{
-				OutPins.Add(FMVVMBlueprintPinId(UE::MVVM::ConversionFunctionHelper::FindPinId(ArgumentPin)));
+				TArray<FName> Names = { ArgumentPin->GetFName() };
+				OutPins.Add(FMVVMBlueprintPinId(Names));
 			}
 		}
 	}
diff --git a/Engine/Plugins/Runtime/ModelViewViewModel/Source/ModelViewViewModelEditor/Private/Widgets/SMVVMViewModelBindingListWidget.cpp b/Engine/Plugins/Runtime/ModelViewViewModel/Source/ModelViewViewModelEditor/Private/Widgets/SMVVMViewModelBindingListWidget.cpp
--- a/Engine/Plugins/Runtime/ModelViewViewModel/Source/ModelViewViewModelEditor/Private/Widgets/SMVVMViewModelBindingListWidget.cpp
+++ b/Engine/Plugins/Runtime/ModelViewViewModel/Source/ModelViewViewModelEditor/Private/Widgets/SMVVMViewModelBindingListWidget.cpp
@@ -286,7 +286,8 @@ TArray<FFieldVariant> FFieldIterator_Bindable::GetFields(const UStruct* Struct,
 			TOptional<FFieldVariant> PassResult;
 			if (FilterFlags == EFilterFlag::All)
 			{
-				const Private::FPassFilterInputs Params(WidgetBlueprint.Get(), Value, Struct, FieldVisibilityFlags, AssignableTo, false, bAllowComponentProperty, bEnableConversion);
+				const bool bDoObjectProperty = AssignableTo != nullptr;
+				const Private::FPassFilterInputs Params(WidgetBlueprint.Get(), Value, Struct, FieldVisibilityFlags, AssignableTo, bDoObjectProperty, bAllowComponentProperty, bEnableConversion);
 				PassResult = Private::PassFilter(Params);
 			}
 			else
