Null check ViewEvents

#rnx
#Synced-CL 47766440

diff --git a/Engine/Plugins/Runtime/ModelViewViewModel/Source/ModelViewViewModelBlueprint/Private/MVVMWidgetBlueprintExtension_View.cpp b/Engine/Plugins/Runtime/ModelViewViewModel/Source/ModelViewViewModelBlueprint/Private/MVVMWidgetBlueprintExtension_View.cpp
--- a/Engine/Plugins/Runtime/ModelViewViewModel/Source/ModelViewViewModelBlueprint/Private/MVVMWidgetBlueprintExtension_View.cpp
+++ b/Engine/Plugins/Runtime/ModelViewViewModel/Source/ModelViewViewModelBlueprint/Private/MVVMWidgetBlueprintExtension_View.cpp
@@ -281,7 +281,10 @@ void UMVVMWidgetBlueprintExtension_View::HandleBeginCompilation(FWidgetBlueprint
 		BlueprintView->ResetBindingMessages();
 		for (UMVVMBlueprintViewEvent* ViewEvent : BlueprintView->GetEvents())
 		{
-			ViewEvent->ResetCompilationMessages();
+			if(ViewEvent)
+			{
+				ViewEvent->ResetCompilationMessages();
+			}
 		}
 		CurrentCompilerContext = MakePimpl<UE::MVVM::Private::FMVVMViewBlueprintCompiler>(InCreationContext, GetBlueprintView());
 	}
