MVVM: Quick binding, support multiparameter conversion functions.

#rnx
#Synced-CL 48472971

diff --git a/Engine/Plugins/Runtime/ModelViewViewModel/Source/ModelViewViewModelEditor/Private/Widgets/SMVVMViewModelBindingListWidget.cpp b/Engine/Plugins/Runtime/ModelViewViewModel/Source/ModelViewViewModelEditor/Private/Widgets/SMVVMViewModelBindingListWidget.cpp
--- a/Engine/Plugins/Runtime/ModelViewViewModel/Source/ModelViewViewModelEditor/Private/Widgets/SMVVMViewModelBindingListWidget.cpp
+++ b/Engine/Plugins/Runtime/ModelViewViewModel/Source/ModelViewViewModelEditor/Private/Widgets/SMVVMViewModelBindingListWidget.cpp
@@ -237,9 +237,10 @@ namespace Private
 			TArray<FConversionFunctionValue> ConversionFunctions = Subsystem->GetConversionFunctions(WidgetBlueprint, Source, Destination);
 			for (const FConversionFunctionValue& ConversionFunction : ConversionFunctions)
 			{
+				// Returns function pins that support the Source property. If there are multiple, just pick the first. 
+				// Currently no method to determine which parameter is the conversion value
 				TArray<FMVVMBlueprintPinId> ArgumentPins = Subsystem->GetConversionFunctionPins(WidgetBlueprint, ConversionFunction, Source);
-					// TODO: Find first argument that matches type
-				if (ArgumentPins.Num() == 1) // TODO: Should we pick first available parameter that matched type? Or skip?
+				if (!ArgumentPins.IsEmpty()) 
 				{
 					FMVVMFieldConversion Conversion(ConversionFunction, ArgumentPins[0]);
 					OutConversions.Add(Conversion);
