Quick Binding: Minor text label changes

#rnx

#Synced-CL 48633006

diff --git a/Engine/Plugins/Runtime/ModelViewViewModel/Source/ModelViewViewModelEditor/Private/Widgets/SMVVMFieldSelectorMenu.cpp b/Engine/Plugins/Runtime/ModelViewViewModel/Source/ModelViewViewModelEditor/Private/Widgets/SMVVMFieldSelectorMenu.cpp
--- a/Engine/Plugins/Runtime/ModelViewViewModel/Source/ModelViewViewModelEditor/Private/Widgets/SMVVMFieldSelectorMenu.cpp
+++ b/Engine/Plugins/Runtime/ModelViewViewModel/Source/ModelViewViewModelEditor/Private/Widgets/SMVVMFieldSelectorMenu.cpp
@@ -338,14 +338,14 @@ FText SFieldSelectorMenu::GetTooltipText() const
 	{
 		return LOCTEXT("FieldSelector_TooltipBindTo", "Bind to ");
 	}
-	return LOCTEXT("FieldSelector_Tooltip", "Select a property to bind to");
+	return LOCTEXT("FieldSelector_Tooltip", "Select a property to bind to:");
 }
 
 FText SFieldSelectorMenu::GetAdditionalTooltipText() const
 {
 	if (bIsSelectingConversion)
 	{
-		return LOCTEXT("FieldSelector_TooltipConversion", " with conversion function");
+		return LOCTEXT("FieldSelector_TooltipConversion", " with conversion function:");
 	}
 	return FText::GetEmpty();
 }
