PR #13959: Add ClassThumbnail for UMVVMViewModelBase
#13959

#tests unshelved and saw the class icon thumbnail change
#Synced-CL 49386971

diff --git a/Engine/Plugins/Runtime/ModelViewViewModel/Source/ModelViewViewModelEditor/Private/Styling/MVVMEditorStyle.cpp b/Engine/Plugins/Runtime/ModelViewViewModel/Source/ModelViewViewModelEditor/Private/Styling/MVVMEditorStyle.cpp
--- a/Engine/Plugins/Runtime/ModelViewViewModel/Source/ModelViewViewModelEditor/Private/Styling/MVVMEditorStyle.cpp
+++ b/Engine/Plugins/Runtime/ModelViewViewModel/Source/ModelViewViewModelEditor/Private/Styling/MVVMEditorStyle.cpp
@@ -18,13 +18,8 @@ TUniquePtr<FMVVMEditorStyle, FMVVMEditorStyle::FCustomDeleter> FMVVMEditorStyle:
 FMVVMEditorStyle::FMVVMEditorStyle()
 	: FSlateStyleSet("MVVMEditorStyle")
 {
-	const FVector2D Icon10x10(10.0f, 10.0f);
-	const FVector2D Icon14x14(14.0f, 14.0f);
-	const FVector2D Icon16x16(16.0f, 16.0f);
-	const FVector2D Icon20x20(20.0f, 20.0f);
-	const FVector2D Icon24x24(24.0f, 24.0f);
-	const FVector2D Icon32x32(32.0f, 32.0f);
-	const FVector2D Icon40x40(40.0f, 40.0f);
+	using CoreStyleConstants::Icon16x16;
+	using CoreStyleConstants::Icon64x64;
 
 	TSharedPtr<IPlugin> MVVMPlugin = IPluginManager::Get().FindPlugin("ModelViewViewModel");
 	if (ensure(MVVMPlugin))
@@ -37,6 +32,7 @@ FMVVMEditorStyle::FMVVMEditorStyle()
 	{
 		Set("ClassIcon.MVVMBlueprintView", new IMAGE_BRUSH_SVG("Slate/ViewModel", Icon16x16));
 		Set("ClassIcon.MVVMBlueprintViewModel", new IMAGE_BRUSH_SVG("Slate/ViewModel", Icon16x16));
+		Set("ClassThumbnail.MVVMViewModelBase", new IMAGE_BRUSH_SVG("Slate/ViewModel", Icon64x64));
 	}
 
 	Set("BlueprintView.TabIcon", new IMAGE_BRUSH_SVG("Slate/ViewModel", Icon16x16));
