MVVM: Cleanup empty conditions and events to prevent Cook and Compile errors

#rnx
#Synced-CL 49564432

diff --git a/Engine/Plugins/Runtime/ModelViewViewModel/Source/ModelViewViewModelBlueprint/Private/MVVMBlueprintView.cpp b/Engine/Plugins/Runtime/ModelViewViewModel/Source/ModelViewViewModelBlueprint/Private/MVVMBlueprintView.cpp
--- a/Engine/Plugins/Runtime/ModelViewViewModel/Source/ModelViewViewModelBlueprint/Private/MVVMBlueprintView.cpp
+++ b/Engine/Plugins/Runtime/ModelViewViewModel/Source/ModelViewViewModelBlueprint/Private/MVVMBlueprintView.cpp
@@ -676,6 +676,25 @@ void UMVVMBlueprintView::PostLoad()
 
 void UMVVMBlueprintView::PreSave(FObjectPreSaveContext Context)
 {
+	// Cleanup Event and Conditions to prevent BP compilation issues
+	for (int32 EventIndex = Conditions.Num() - 1; EventIndex >= 0; --EventIndex)
+	{
+		UMVVMBlueprintViewCondition* Condition = Conditions[EventIndex];
+		if (Condition == nullptr)
+		{
+			Conditions.RemoveAt(EventIndex);
+		}
+	}
+
+	for (int32 ConditionIndex = Events.Num() - 1; ConditionIndex >= 0; --ConditionIndex)
+	{
+		UMVVMBlueprintViewEvent* Event = Events[ConditionIndex];
+		if (Event == nullptr)
+		{
+			Events.RemoveAt(ConditionIndex);
+		}
+	}
+
 	UWidgetBlueprint* WidgetBlueprint = GetOuterUMVVMWidgetBlueprintExtension_View()->GetWidgetBlueprint();
 	for (FMVVMBlueprintViewBinding& Binding : Bindings)
 	{
