[MVVM] Make VM passed in UMVVMViewModelContextResolver::DestroyInstance() non-const, to allow resolvers to operate on it for any pre-destroy cleanup (no longer referenced but not yet garbage collected).

#tests ran in conjunction with required project changes in PIE and verified various VM cleanup scenarios

#Synced-CL 50201124

diff --git a/Engine/Plugins/Runtime/ModelViewViewModel/Source/ModelViewViewModel/Private/View/MVVMViewClass.cpp b/Engine/Plugins/Runtime/ModelViewViewModel/Source/ModelViewViewModel/Private/View/MVVMViewClass.cpp
--- a/Engine/Plugins/Runtime/ModelViewViewModel/Source/ModelViewViewModel/Private/View/MVVMViewClass.cpp
+++ b/Engine/Plugins/Runtime/ModelViewViewModel/Source/ModelViewViewModel/Private/View/MVVMViewClass.cpp
@@ -166,7 +166,7 @@ UObject* FMVVMViewClass_Source::GetOrCreateInstance(const UMVVMViewClass* InView
 	return Result;
 }
 
-void FMVVMViewClass_Source::ReleaseInstance(const UObject* ViewModel, const UMVVMView* View) const
+void FMVVMViewClass_Source::ReleaseInstance(UObject* ViewModel, const UMVVMView* View) const
 {
 	if (Resolver)
 	{
diff --git a/Engine/Plugins/Runtime/ModelViewViewModel/Source/ModelViewViewModel/Public/View/MVVMViewClass.h b/Engine/Plugins/Runtime/ModelViewViewModel/Source/ModelViewViewModel/Public/View/MVVMViewClass.h
--- a/Engine/Plugins/Runtime/ModelViewViewModel/Source/ModelViewViewModel/Public/View/MVVMViewClass.h
+++ b/Engine/Plugins/Runtime/ModelViewViewModel/Source/ModelViewViewModel/Public/View/MVVMViewClass.h
@@ -274,7 +274,7 @@ public:
 	MODELVIEWVIEWMODEL_API UObject* GetOrCreateInstance(const UMVVMViewClass* ViewClass, UMVVMView* View, UUserWidget* UserWidget) const;
 
 	/** The source is not needed anymore. */
-	MODELVIEWVIEWMODEL_API void ReleaseInstance(const UObject* ViewModel, const UMVVMView* View) const;
+	MODELVIEWVIEWMODEL_API void ReleaseInstance(UObject* ViewModel, const UMVVMView* View) const;
 
 	/** Returns the viewmodel inside the global viewmodel collection. */
 	MODELVIEWVIEWMODEL_API TScriptInterface<INotifyFieldValueChanged> GetGlobalCollectionViewModel(UUserWidget* UserWidget) const;
diff --git a/Engine/Plugins/Runtime/ModelViewViewModel/Source/ModelViewViewModel/Public/View/MVVMViewModelContextResolver.h b/Engine/Plugins/Runtime/ModelViewViewModel/Source/ModelViewViewModel/Public/View/MVVMViewModelContextResolver.h
--- a/Engine/Plugins/Runtime/ModelViewViewModel/Source/ModelViewViewModel/Public/View/MVVMViewModelContextResolver.h
+++ b/Engine/Plugins/Runtime/ModelViewViewModel/Source/ModelViewViewModel/Public/View/MVVMViewModelContextResolver.h
@@ -30,13 +30,13 @@ public:
 	UFUNCTION(BlueprintImplementableEvent, Category = "Viewmodel", meta = (DisplayName = "Create Instance"))
 	UE_API TScriptInterface<INotifyFieldValueChanged> K2_CreateInstance(const UClass* ExpectedType, const UUserWidget* UserWidget) const;
 
-	virtual void DestroyInstance(const UObject* ViewModel, const UMVVMView* View) const
+	virtual void DestroyInstance(UObject* ViewModel, const UMVVMView* View) const
 	{
 		K2_DestroyInstance(ViewModel, View);
 	}
 
 	UFUNCTION(BlueprintImplementableEvent, Category = "Viewmodel", meta = (DisplayName = "Destroy Instance"))
-	UE_API void K2_DestroyInstance(const UObject* ViewModel, const UMVVMView* View) const;
+	UE_API void K2_DestroyInstance(UObject* ViewModel, const UMVVMView* View) const;
 
 public:
 #if WITH_EDITOR
