[MVVM] Add new creation more Parent that will look for a vm in the parent widget hierarchy

#tests Modified crew selection vm to allow ancestor type and used it in a few screens to make sure it got initialized properly

#Synced-CL 50384806

diff --git a/Engine/Plugins/Runtime/ModelViewViewModel/Source/ModelViewViewModel/Private/View/MVVMView.cpp b/Engine/Plugins/Runtime/ModelViewViewModel/Source/ModelViewViewModel/Private/View/MVVMView.cpp
--- a/Engine/Plugins/Runtime/ModelViewViewModel/Source/ModelViewViewModel/Private/View/MVVMView.cpp
+++ b/Engine/Plugins/Runtime/ModelViewViewModel/Source/ModelViewViewModel/Private/View/MVVMView.cpp
@@ -189,6 +189,8 @@ void UMVVMView::InitializeSources()
 			//note: Run the bindings. The binding from Source A can influence how the Source B is initialized/evaluated.
 			InitializeSourceBindings(FMVVMView_SourceKey(Index), false);
 		}
+
+		InitializeSourceParentBinding(FMVVMView_SourceKey(Index));
 	}
 
 	if (GeneratedViewClass->DoesInitializeBindingsOnConstruct())
@@ -250,6 +252,7 @@ void UMVVMView::UninitializeSources()
 	for (int32 Index = 0; Index < NumberOfSources; ++Index)
 	{
 		UninitializeSource(FMVVMView_SourceKey(Index));
+		UninitializeSourceParentBinding(FMVVMView_SourceKey(Index));
 	}
 }
 
@@ -303,6 +306,9 @@ void UMVVMView::InitializeSourceInternal(UObject* NewSource, FMVVMViewClass_Sour
 	}
 
 	ViewSource.Source = NewSource;
+
+	// See comment in UMVVMView::InitializeSourceParentBinding for why we do this
+	GetUserWidget()->BroadcastFieldValueChanged(UE::FieldNotification::FFieldId(ClassSource.GetName(), 0));
 }
 
 
@@ -340,6 +346,64 @@ void UMVVMView::UninitializeSource(FMVVMView_SourceKey SourceKey)
 	}
 }
 
+void UMVVMView::InitializeSourceParentBinding(FMVVMView_SourceKey SourceKey)
+{
+	FMVVMView_Source& ViewSource = Sources[SourceKey.GetIndex()];
+	const FMVVMViewClass_Source& ClassSource = GeneratedViewClass->GetSource(ViewSource.ClassKey);
+
+	if (ClassSource.IsParentContextReference())
+	{
+		const TPair<const UMVVMView*, FMVVMView_SourceKey> ParentInfo = ClassSource.GetParentSupplyingSourceClass(GetUserWidget());
+		if (const UMVVMView* ParentView = ParentInfo.Key)
+		{
+			if (UUserWidget* ParentWidget = ParentView->GetOuterUUserWidget())
+			{
+				const FMVVMView_SourceKey ParentSourcKey = ParentInfo.Value;
+				const FMVVMViewClass_Source& ParentClassSource = ParentView->GetViewClass()->GetSource(ParentSourcKey.GetIndex());
+
+				TWeakObjectPtr<const UMVVMView> WeakView = ParentView;
+				UE::FieldNotification::FFieldMulticastDelegate::FDelegate Delegate = UE::FieldNotification::FFieldMulticastDelegate::FDelegate::CreateWeakLambda(this, [this, SourceKey, ParentSourcKey, WeakView](UObject*, UE::FieldNotification::FFieldId)
+					{
+						if (const UMVVMView* ParentView = WeakView.Get())
+						{
+							SetSourceInternal(FMVVMViewClass_SourceKey(SourceKey.GetIndex()), ParentView->GetSource(ParentSourcKey).Source, true);
+						}
+					});
+
+				ParentWidget->GetFieldNotificationDescriptor();
+
+				// This is a big of a hack. The VM variable is not marked as field notify and we don't want it to be, but we also don't want to add bespoke
+				// notification logic. We always know that there will be at least one field notify in a user widget (UWidget has a few) and field notifies
+				// use the index to enable callbacks then the names to filter them. Adding this will enable all calbacks bound to index 0 (which may be multiple)
+				// and get called for any callbacks that share the name, regardless of index. If there happens to be a duplicate name in the delegates, then
+				// it likely just means that some callbacks are firing off more than they should, i.e. being notified of a change when there wasn't any,
+				// so there should be very low risk to this.
+				//
+				// If any of this changes make sure to also update UMVVMView::InitializeSourceInternal and UMVVMView::UninitializeSourceParentBinding
+				ParentWidget->AddFieldValueChangedDelegate(UE::FieldNotification::FFieldId(ParentClassSource.GetName(), 0), MoveTemp(Delegate));
+			}
+		}
+	}
+}
+
+void UMVVMView::UninitializeSourceParentBinding(FMVVMView_SourceKey SourceKey)
+{
+	FMVVMView_Source& ViewSource = Sources[SourceKey.GetIndex()];
+	const FMVVMViewClass_Source& ClassSource = GeneratedViewClass->GetSource(ViewSource.ClassKey);
+
+	if (ClassSource.IsParentContextReference())
+	{
+		// See comment in UMVVMView::InitializeSourceParentBinding for why we do this
+		const TPair<const UMVVMView*, FMVVMView_SourceKey> ParentInfo = ClassSource.GetParentSupplyingSourceClass(GetUserWidget());
+		if (const UMVVMView* ParentView = ParentInfo.Key)
+		{
+			if (UUserWidget* ParentWidget = ParentView->GetOuterUUserWidget())
+			{
+				ParentWidget->RemoveAllFieldValueChangedDelegates(this);
+			}
+		}
+	}
+}
 
 void UMVVMView::InitializeBindings()
 {
diff --git a/Engine/Plugins/Runtime/ModelViewViewModel/Source/ModelViewViewModel/Private/View/MVVMViewClass.cpp b/Engine/Plugins/Runtime/ModelViewViewModel/Source/ModelViewViewModel/Private/View/MVVMViewClass.cpp
--- a/Engine/Plugins/Runtime/ModelViewViewModel/Source/ModelViewViewModel/Private/View/MVVMViewClass.cpp
+++ b/Engine/Plugins/Runtime/ModelViewViewModel/Source/ModelViewViewModel/Private/View/MVVMViewClass.cpp
@@ -111,6 +111,29 @@ UObject* FMVVMViewClass_Source::GetOrCreateInstance(const UMVVMViewClass* InView
 			Log.Error(FText::Format(LOCTEXT("CreateInstanceCreateInstance", "The source '{0}' could not be created. The class is not loaded."), FText::FromName(PropertyName)));
 		}
 	}
+	else if (IsParentContextReference())
+	{
+		UE::MVVM::FMessageLog Log(InUserWidget);
+
+		if (ExpectedSourceType.Get() != nullptr)
+		{
+			const TPair<const UMVVMView*, FMVVMView_SourceKey> ParentInfo = GetParentSupplyingSourceClass(InUserWidget);
+			if (const UMVVMView* ParentView = ParentInfo.Key)
+			{
+				Result = ParentView->GetSource(ParentInfo.Value).Source;
+
+				Log.Info(FText::Format(LOCTEXT("CreateInstanceContextParentFound", "The source '{0}' was found on parent widget '{1}'"), FText::FromName(PropertyName), FText::FromString(GetNameSafe(ParentView->GetOuterUUserWidget()))));
+			}
+			else
+			{
+				Log.Error(FText::Format(LOCTEXT("CreateInstanceContextParentMissing", "The source '{0}' could not be created. No parent widget found that provides an instance of '{1}'"), FText::FromName(PropertyName), FText::FromString(GetNameSafe(ExpectedSourceType))));
+			}		
+		}
+		else if (!bOptional)
+		{
+			Log.Error(FText::Format(LOCTEXT("CreateInstanceContextClassMissing", "The source '{0}' could not be created. The class is not loaded."), FText::FromName(PropertyName)));
+		}
+	}
 	else if (IsUserWidget())
 	{ 
 		Result = InUserWidget;
@@ -181,6 +204,29 @@ TScriptInterface<INotifyFieldValueChanged> FMVVMViewClass_Source::GetGlobalColle
 	return FoundViewModelInstance ? TScriptInterface<INotifyFieldValueChanged>(FoundViewModelInstance) : TScriptInterface<INotifyFieldValueChanged>();
 }
 
+TPair<const UMVVMView*, FMVVMView_SourceKey> FMVVMViewClass_Source::GetParentSupplyingSourceClass(UUserWidget* UserWidget) const
+{
+	// Find an existing instance from a parent source and/or bind to when that paren't source changes
+	for (const UUserWidget* Parent = UserWidget->GetTypedOuter<UUserWidget>(); Parent; Parent = Parent->GetTypedOuter<UUserWidget>())
+	{
+		if (const UMVVMView* ParentView = UMVVMSubsystem::GetViewFromUserWidget(Parent); ParentView && ParentView->GetViewClass())
+		{
+			const int32 FoundIndex = ParentView->GetViewClass()->GetSources().IndexOfByPredicate([this](const FMVVMViewClass_Source& Other)
+				{
+					return Other.IsViewModel() && Other.GetSourceClass()->IsChildOf(ExpectedSourceType) &&
+						(GetName().IsNone() ? Other.GetName().IsNone() : GetName() == Other.GetName());
+				});
+
+			if (ParentView->GetSources().IsValidIndex(FoundIndex))
+			{
+				return { ParentView, FoundIndex };
+			}
+		}
+	}
+
+	return { nullptr, INDEX_NONE };
+}
+
 #if UE_WITH_MVVM_DEBUGGING
 FMVVMViewClass_Source::FToStringArgs FMVVMViewClass_Source::FToStringArgs::Short()
 {
diff --git a/Engine/Plugins/Runtime/ModelViewViewModel/Source/ModelViewViewModel/Public/View/MVVMView.h b/Engine/Plugins/Runtime/ModelViewViewModel/Source/ModelViewViewModel/Public/View/MVVMView.h
--- a/Engine/Plugins/Runtime/ModelViewViewModel/Source/ModelViewViewModel/Public/View/MVVMView.h
+++ b/Engine/Plugins/Runtime/ModelViewViewModel/Source/ModelViewViewModel/Public/View/MVVMView.h
@@ -198,6 +198,9 @@ private:
 	UE_API void UninitializeSource(FMVVMView_SourceKey SourceKey);
 	UE_API bool SetSourceInternal(FMVVMViewClass_SourceKey SourceKey, TScriptInterface<INotifyFieldValueChanged> ViewModel, bool bForDynamicSource);
 
+	void InitializeSourceParentBinding(FMVVMView_SourceKey SourceKey);
+	void UninitializeSourceParentBinding(FMVVMView_SourceKey SourceKey);
+
 	//~ Bindings
 	UE_API void InitializeSourceBindings(FMVVMView_SourceKey SourceKey, bool bRunAllBindings);
 	UE_API void InitializeSourceBindingsCommon();
diff --git a/Engine/Plugins/Runtime/ModelViewViewModel/Source/ModelViewViewModel/Public/View/MVVMViewClass.h b/Engine/Plugins/Runtime/ModelViewViewModel/Source/ModelViewViewModel/Public/View/MVVMViewClass.h
--- a/Engine/Plugins/Runtime/ModelViewViewModel/Source/ModelViewViewModel/Public/View/MVVMViewClass.h
+++ b/Engine/Plugins/Runtime/ModelViewViewModel/Source/ModelViewViewModel/Public/View/MVVMViewClass.h
@@ -279,6 +279,9 @@ public:
 	/** Returns the viewmodel inside the global viewmodel collection. */
 	MODELVIEWVIEWMODEL_API TScriptInterface<INotifyFieldValueChanged> GetGlobalCollectionViewModel(UUserWidget* UserWidget) const;
 
+	/** Get the parent vm view that has a matching vm class source */
+	TPair<const UMVVMView*, FMVVMView_SourceKey> GetParentSupplyingSourceClass(UUserWidget* UserWidget) const;
+
 	/** The expected class of the source. */
 	UClass* GetSourceClass() const
 	{
@@ -339,7 +342,13 @@ public:
 	/** When a source is manually set, then always execute that source the bindings. */
 	bool AlwaysExecuteBindingsOnSetSource() const
 	{
-		return (Flags & (uint16)EFlags::AlwaysExecuteBindingsOnSetSource) != 0;;
+		return (Flags & (uint16)EFlags::AlwaysExecuteBindingsOnSetSource) != 0;
+	}
+
+	/** If we're going to lookup the vm from a parent view */
+	bool IsParentContextReference() const
+	{
+		return (Flags & (uint16)EFlags::ParentContextReference) != 0;
 	}
 	
 	/**
@@ -462,6 +471,7 @@ private:
 		IsViewModelInstanceExposed = 1 << 10,
 		GlobalViewModelCollectionUpdate = 1 << 11,
 		AlwaysExecuteBindingsOnSetSource = 1 << 12,
+		ParentContextReference = 1 << 13,
 	};
 
 	UPROPERTY(VisibleAnywhere, Category = "View")
diff --git a/Engine/Plugins/Runtime/ModelViewViewModel/Source/ModelViewViewModelBlueprint/Private/MVVMDeveloperProjectSettings.cpp b/Engine/Plugins/Runtime/ModelViewViewModel/Source/ModelViewViewModelBlueprint/Private/MVVMDeveloperProjectSettings.cpp
--- a/Engine/Plugins/Runtime/ModelViewViewModel/Source/ModelViewViewModelBlueprint/Private/MVVMDeveloperProjectSettings.cpp
+++ b/Engine/Plugins/Runtime/ModelViewViewModel/Source/ModelViewViewModelBlueprint/Private/MVVMDeveloperProjectSettings.cpp
@@ -37,6 +37,7 @@ UMVVMDeveloperProjectSettings::UMVVMDeveloperProjectSettings()
 	AllowedContextCreationType.Add(EMVVMBlueprintViewModelContextCreationType::PropertyPath);
 	AllowedContextCreationType.Add(EMVVMBlueprintViewModelContextCreationType::PropertyPath);
 	AllowedContextCreationType.Add(EMVVMBlueprintViewModelContextCreationType::Resolver);
+	AllowedContextCreationType.Add(EMVVMBlueprintViewModelContextCreationType::Ancestors);
 
 	FTopLevelAssetPath BlueprintFunctionLibrary = FTopLevelAssetPath("/Script/Engine", "BlueprintFunctionLibrary");
 	FTopLevelAssetPath FormatText = FTopLevelAssetPath("/Script/BlueprintGraph", "K2Node_FormatText");
diff --git a/Engine/Plugins/Runtime/ModelViewViewModel/Source/ModelViewViewModelBlueprint/Private/MVVMViewBlueprintCompiler.cpp b/Engine/Plugins/Runtime/ModelViewViewModel/Source/ModelViewViewModelBlueprint/Private/MVVMViewBlueprintCompiler.cpp
--- a/Engine/Plugins/Runtime/ModelViewViewModel/Source/ModelViewViewModelBlueprint/Private/MVVMViewBlueprintCompiler.cpp
+++ b/Engine/Plugins/Runtime/ModelViewViewModel/Source/ModelViewViewModelBlueprint/Private/MVVMViewBlueprintCompiler.cpp
@@ -963,7 +963,7 @@ void FMVVMViewBlueprintCompiler::CreateRequiredProperties(const FWidgetBlueprint
 			SourceContext->Type = FCompilerBindingSource::EType::ViewModel;
 			SourceContext->bIsOptional = ViewModelContext.bOptional;
 
-			if (ViewModelContext.CreationType == EMVVMBlueprintViewModelContextCreationType::Manual)
+			if (ViewModelContext.CreationType == EMVVMBlueprintViewModelContextCreationType::Manual || ViewModelContext.CreationType == EMVVMBlueprintViewModelContextCreationType::Ancestors)
 			{
 				SourceContext->bIsOptional = true;
 			}
@@ -2570,7 +2570,7 @@ void FMVVMViewBlueprintCompiler::PreCompileViewModelCreatorContexts(UWidgetBluep
 			}
 		}
 
-		if (ViewModelContext.CreationType == EMVVMBlueprintViewModelContextCreationType::Manual)
+		if (ViewModelContext.CreationType == EMVVMBlueprintViewModelContextCreationType::Manual || ViewModelContext.CreationType == EMVVMBlueprintViewModelContextCreationType::Ancestors)
 		{
 		}
 		else if (ViewModelContext.CreationType == EMVVMBlueprintViewModelContextCreationType::CreateInstance)
@@ -2777,6 +2777,7 @@ void FMVVMViewBlueprintCompiler::CompileSources(const FCompiledBindingLibraryCom
 		bool bExposeInstanceInEditor = false;
 		bool bGlobalViewModelCollectionUpdate = false;
 		bool bAlwaysExecuteBindingsOnSetSource =  GetDefault<UMVVMDeveloperProjectSettings>()->bForceExecuteBindingsOnSetSource;
+		bool bParentContextReference = false;
 
 		if (ViewModelContext.CreationType == EMVVMBlueprintViewModelContextCreationType::Manual)
 		{
@@ -2856,6 +2857,12 @@ void FMVVMViewBlueprintCompiler::CompileSources(const FCompiledBindingLibraryCom
 			CompiledSourceCreator.Resolver = Resolver;
 			ensure(bIsOptional == ViewModelContext.bOptional);
 		}
+		else if (ViewModelContext.CreationType == EMVVMBlueprintViewModelContextCreationType::Ancestors)
+		{
+			ensure(bCanBeSet == false);
+			ensure(bIsOptional == true);
+			bParentContextReference = true;
+		}
 		else
 		{
 			AddMessageForViewModel(ViewModelContext
@@ -2896,6 +2903,7 @@ void FMVVMViewBlueprintCompiler::CompileSources(const FCompiledBindingLibraryCom
 		CompiledSourceCreator.Flags |= bExposeInstanceInEditor ? (uint16)FMVVMViewClass_Source::EFlags::IsViewModelInstanceExposed : 0;
 		CompiledSourceCreator.Flags |= bGlobalViewModelCollectionUpdate ? (uint16)FMVVMViewClass_Source::EFlags::GlobalViewModelCollectionUpdate : 0;
 		CompiledSourceCreator.Flags |= bAlwaysExecuteBindingsOnSetSource ? (uint16)FMVVMViewClass_Source::EFlags::AlwaysExecuteBindingsOnSetSource : 0;
+		CompiledSourceCreator.Flags |= bParentContextReference ? (uint16)FMVVMViewClass_Source::EFlags::ParentContextReference : 0;
 
 		{
 			FSortData SortData;
diff --git a/Engine/Plugins/Runtime/ModelViewViewModel/Source/ModelViewViewModelBlueprint/Public/MVVMBlueprintViewModelContext.h b/Engine/Plugins/Runtime/ModelViewViewModel/Source/ModelViewViewModelBlueprint/Public/MVVMBlueprintViewModelContext.h
--- a/Engine/Plugins/Runtime/ModelViewViewModel/Source/ModelViewViewModelBlueprint/Public/MVVMBlueprintViewModelContext.h
+++ b/Engine/Plugins/Runtime/ModelViewViewModel/Source/ModelViewViewModelBlueprint/Public/MVVMBlueprintViewModelContext.h
@@ -28,6 +28,8 @@ enum class EMVVMBlueprintViewModelContextCreationType : uint8
 	PropertyPath,
 	// The viewmodel will be fetched by evaluating the resolver object.
 	Resolver,
+	// The viewmodel will be fetched from any parent view that contains the same class and name
+	Ancestors,
 };
 
 namespace UE::MVVM
diff --git a/Engine/Plugins/Runtime/ModelViewViewModel/Source/ModelViewViewModelEditor/Private/Customizations/MVVMBlueprintViewModelContextCustomization.cpp b/Engine/Plugins/Runtime/ModelViewViewModel/Source/ModelViewViewModelEditor/Private/Customizations/MVVMBlueprintViewModelContextCustomization.cpp
--- a/Engine/Plugins/Runtime/ModelViewViewModel/Source/ModelViewViewModelEditor/Private/Customizations/MVVMBlueprintViewModelContextCustomization.cpp
+++ b/Engine/Plugins/Runtime/ModelViewViewModel/Source/ModelViewViewModelEditor/Private/Customizations/MVVMBlueprintViewModelContextCustomization.cpp
@@ -390,7 +390,7 @@ void FBlueprintViewModelContextDetailCustomization::CustomizeChildren(TSharedRef
 				ChildBuilder.AddProperty(CreateSetterFunctionHandle.ToSharedRef())
 					.IsEnabled(MakeAttributeLambda([bCanEdit, ContextPtr]()
 						{
-							bool bResult = ContextPtr->CreationType != EMVVMBlueprintViewModelContextCreationType::Manual;
+							bool bResult = ContextPtr->CreationType != EMVVMBlueprintViewModelContextCreationType::Manual && ContextPtr->CreationType != EMVVMBlueprintViewModelContextCreationType::Ancestors;
 							return bResult && bCanEdit;
 						}));
 			}
@@ -526,9 +526,10 @@ void FBlueprintViewModelContextDetailCustomization::HandleCreationTypeChanged()
 		{
 			const EMVVMBlueprintViewModelContextCreationType CreationType = (EMVVMBlueprintViewModelContextCreationType)NewValue;
 			const bool bIsManual = CreationType == EMVVMBlueprintViewModelContextCreationType::Manual;
-			if (ContextPtr->bOptional != bIsManual)
+			const bool bIsAlwaysOptional = CreationType == EMVVMBlueprintViewModelContextCreationType::Manual || CreationType == EMVVMBlueprintViewModelContextCreationType::Ancestors;
+			if (ContextPtr->bOptional != bIsAlwaysOptional)
 			{
-				OptionalHandle->SetValue(bIsManual);
+				OptionalHandle->SetValue(bIsAlwaysOptional);
 			}
 			if (ContextPtr->bCreateSetterFunction != bIsManual)
 			{
