Check for valid CompiledBindingLibraryIds before failing MVVM bindings due to mismatch
#rnx
#Synced-CL 50859483

diff --git a/Engine/Plugins/Runtime/ModelViewViewModel/Source/ModelViewViewModel/Private/Bindings/MVVMCompiledBindingLibrary.cpp b/Engine/Plugins/Runtime/ModelViewViewModel/Source/ModelViewViewModel/Private/Bindings/MVVMCompiledBindingLibrary.cpp
--- a/Engine/Plugins/Runtime/ModelViewViewModel/Source/ModelViewViewModel/Private/Bindings/MVVMCompiledBindingLibrary.cpp
+++ b/Engine/Plugins/Runtime/ModelViewViewModel/Source/ModelViewViewModel/Private/Bindings/MVVMCompiledBindingLibrary.cpp
@@ -188,7 +188,7 @@ TValueOrError<void, FMVVMCompiledBindingLibrary::EExecutionFailingReason> FMVVMC
 {
 #if WITH_EDITORONLY_DATA
 	const bool bIsValidBinding = InBinding.CompiledBindingLibraryId == CompiledBindingLibraryId;
-	if (!bIsValidBinding)
+	if (InBinding.CompiledBindingLibraryId.IsValid() && !bIsValidBinding)
 	{
 		ensureAlwaysMsgf(false, TEXT("The binding is from a different library."));
 		return MakeError(EExecutionFailingReason::IncompatibleLibrary);
@@ -353,7 +353,7 @@ TValueOrError<UE::MVVM::FFieldContext, void> FMVVMCompiledBindingLibrary::Evalua
 
 #if WITH_EDITORONLY_DATA
 	const bool bIsValidBinding = InFieldPath.CompiledBindingLibraryId == CompiledBindingLibraryId;
-	if (!bIsValidBinding)
+	if (InFieldPath.CompiledBindingLibraryId.IsValid() && !bIsValidBinding)
 	{
 		ensureAlwaysMsgf(false, TEXT("The binding is from a different library."));
 		return MakeError();
@@ -501,7 +501,7 @@ TValueOrError<FString, FString> FMVVMCompiledBindingLibrary::FieldPathToString(F
 {
 #if WITH_EDITORONLY_DATA
 	const bool bIsValidBinding = FieldPath.CompiledBindingLibraryId == CompiledBindingLibraryId;
-	if (!bIsValidBinding)
+	if (FieldPath.CompiledBindingLibraryId.IsValid() && !bIsValidBinding)
 	{
 		ensureAlwaysMsgf(false, TEXT("The binding is from a different library."));
 		return MakeError<FString>(TEXT("The binding is from a different library."));
