MVVM: make default resolver object public to fix serialization issues
#Synced-CL 51250083

diff --git a/Engine/Plugins/Runtime/ModelViewViewModel/Source/ModelViewViewModelBlueprint/Private/MVVMBlueprintView.cpp b/Engine/Plugins/Runtime/ModelViewViewModel/Source/ModelViewViewModelBlueprint/Private/MVVMBlueprintView.cpp
--- a/Engine/Plugins/Runtime/ModelViewViewModel/Source/ModelViewViewModelBlueprint/Private/MVVMBlueprintView.cpp
+++ b/Engine/Plugins/Runtime/ModelViewViewModel/Source/ModelViewViewModelBlueprint/Private/MVVMBlueprintView.cpp
@@ -12,6 +12,7 @@
 #include "WidgetBlueprint.h"
 #include "UObject/AssetRegistryTagsContext.h"
 #include "UObject/FortniteMainBranchObjectVersion.h"
+#include "View/MVVMViewModelContextResolver.h"
 
 #include UE_INLINE_GENERATED_CPP_BY_NAME(MVVMBlueprintView)
 
@@ -672,6 +673,14 @@ void UMVVMBlueprintView::PostLoad()
 		};
 		UE::MVVM::Private::ForEachPropertyPath_Update(this, DeprecateSelfPath, false);
 	}
+
+	for (FMVVMBlueprintViewModelContext& ViewModelContext : AvailableViewModels)
+	{
+		if (ViewModelContext.Resolver && !ViewModelContext.Resolver->HasAnyFlags(RF_Public))
+		{
+			ViewModelContext.Resolver->SetFlags(RF_Public);
+		}
+	}
 }
 
 void UMVVMBlueprintView::PreSave(FObjectPreSaveContext Context)
diff --git a/Engine/Plugins/Runtime/ModelViewViewModel/Source/ModelViewViewModelBlueprint/Private/MVVMBlueprintViewModelContext.cpp b/Engine/Plugins/Runtime/ModelViewViewModel/Source/ModelViewViewModelBlueprint/Private/MVVMBlueprintViewModelContext.cpp
--- a/Engine/Plugins/Runtime/ModelViewViewModel/Source/ModelViewViewModelBlueprint/Private/MVVMBlueprintViewModelContext.cpp
+++ b/Engine/Plugins/Runtime/ModelViewViewModel/Source/ModelViewViewModelBlueprint/Private/MVVMBlueprintViewModelContext.cpp
@@ -44,7 +44,7 @@ TObjectPtr<UMVVMViewModelContextResolver> FMVVMBlueprintViewModelContext::Create
 			{
 				if (DefaultResolverClass->DoesSupportViewModelClass(GetViewModelClass()))
 				{
-					DefaultResolver = NewObject<UMVVMViewModelContextResolver>(Package, DefaultClass);
+					DefaultResolver = NewObject<UMVVMViewModelContextResolver>(Package, DefaultClass, NAME_None, RF_Public);
 				}
 			}
 		}
