Widget Reflector -> UMG Editor Debug Path:
- Added a path, similar to the widget selection path, for selecting debug assets from the Widget Reflector, through Slate Application, to the Engine Asset Editor Subsystem.
- Clicking the 'source' link in a UMG widget entry in the Widget Reflector will now also select that widget for debugging after opening the widget class in the UMG editor.
- Clicking the 'source' link in a UMG widget entry in the Widget Reflector will now also attempt to open the actual BP to debug the instance, rather than opening the asset into which the instance was placed.
- Clicking the 'source' link in a UMG widget entry in the Widget Reflector will now pin the instance on the debug list.
- When debugging an instance of a widget in the UMG editor the viewmodel tab will now show the editable properties of the viewmodels (where available) and show live values of the widget itself.
- In the debug object selector in blueprint editors, the items are now pinnable.

#Synced-CL 51353830

diff --git a/Engine/Plugins/Runtime/ModelViewViewModel/Source/ModelViewViewModelEditor/Private/Widgets/SMVVMViewModelPanel.cpp b/Engine/Plugins/Runtime/ModelViewViewModel/Source/ModelViewViewModelEditor/Private/Widgets/SMVVMViewModelPanel.cpp
--- a/Engine/Plugins/Runtime/ModelViewViewModel/Source/ModelViewViewModelEditor/Private/Widgets/SMVVMViewModelPanel.cpp
+++ b/Engine/Plugins/Runtime/ModelViewViewModel/Source/ModelViewViewModelEditor/Private/Widgets/SMVVMViewModelPanel.cpp
@@ -9,10 +9,12 @@
 #include "MVVMBlueprintInstancedViewModel.h"
 #include "MVVMDeveloperProjectSettings.h"
 #include "MVVMEditorSubsystem.h"
+#include "MVVMSubsystem.h"
 #include "MVVMWidgetBlueprintExtension_View.h"
 #include "PropertyEditorModule.h"
 #include "ScopedTransaction.h"
 #include "ToolMenus.h"
+#include "View/MVVMView.h"
 #include "View/MVVMViewModelContextResolver.h"
 #include "ViewModelFieldDragDropOp.h"
 #include "WidgetBlueprint.h"
@@ -216,6 +218,8 @@
 	FieldExpander = MakeUnique<FFieldExpander_Bindable>();
 
 	WidgetBlueprintEditor->OnSelectedWidgetsChanging.AddSP(this, &SMVVMViewModelPanel::HandleEditorSelectionChanged);
+	
+	WidgetBlueprint->OnSetObjectBeingDebugged().AddSP(this, &SMVVMViewModelPanel::HandleSetObjectBeingDebugged);
 
 	if (CurrentBlueprintView)
 	{
@@ -229,56 +233,8 @@
 
 	WidgetBlueprintTransactionHandle = WidgetBlueprintEditor->GetOnWidgetBlueprintTransaction().AddSP(this, &SMVVMViewModelPanel::FillViewModel);
 
-
 	CreateCommandList();
 
-	FToolMenuContext GenerateWidgetContext;
-	{
-		UWidgetBlueprintToolMenuContext* WidgetBlueprintMenuContext = NewObject<UWidgetBlueprintToolMenuContext>();
-		WidgetBlueprintMenuContext->WidgetBlueprintEditor = WeakBlueprintEditor;
-		GenerateWidgetContext.AddObject(WidgetBlueprintMenuContext);
-
-		UMVVMViewModelPanelToolMenuContext* ViewModelPanelToolMenuContext = NewObject<UMVVMViewModelPanelToolMenuContext>();
-		ViewModelPanelToolMenuContext->ViewModelPanel = SharedThis(this);
-		GenerateWidgetContext.AddObject(ViewModelPanelToolMenuContext);
-	}
-
-	ViewModelTreeView = SNew(UE::PropertyViewer::SPropertyViewer)
-		.PropertyVisibility(UE::PropertyViewer::SPropertyViewer::EPropertyVisibility::Visible)
-		.bShowSearchBox(true)
-		.bShowFieldIcon(true)
-		.bSanitizeName(true)
-		.FieldIterator(FieldIterator.Get())
-		.FieldExpander(FieldExpander.Get())
-		.OnGetPreSlot(this, &SMVVMViewModelPanel::HandleGetPreSlot)
-		.OnContextMenuOpening(this, &SMVVMViewModelPanel::HandleContextMenuOpening)
-		.OnSelectionChanged(this, &SMVVMViewModelPanel::HandleSelectionChanged)
-		.OnGenerateContainer(this, &SMVVMViewModelPanel::HandleGenerateContainer)
-		.OnDragDetected(this, &SMVVMViewModelPanel::HandleDragDetected)
-		.SearchBoxPreSlot()
-		[
-			UToolMenus::Get()->GenerateWidget("MVVM.Viewmodels.Toolbar", GenerateWidgetContext)
-		]
-		.SearchBoxPostSlot()
-		[
-			SNew(SComboButton)
-			.HasDownArrow(false)
-			.ContentPadding(0)
-			.ForegroundColor(FSlateColor::UseForeground())
-			.ButtonStyle(FAppStyle::Get(), "SimpleButton")
-			.MenuContent()
-			[
-				UToolMenus::Get()->GenerateWidget("MVVM.Viewmodels.Settings", GenerateWidgetContext)
-			]
-			.ButtonContent()
-			[
-				SNew(SImage)
-				.Image(FAppStyle::GetBrush("DetailsView.ViewOptions"))
-			]
-		];
-
-	FillViewModel();
-
 	ModelContextWrapper.Reset(NewObject<UMVVMBlueprintViewModelContextWrapper>());
 
 	ChildSlot
@@ -291,7 +247,7 @@
 			+ SVerticalBox::Slot()
 			.FillHeight(1.f)
 			[
-				ViewModelTreeView.ToSharedRef()
+				SAssignNew(ViewModelTreeViewContainer, SBox)
 			]
 			+ SVerticalBox::Slot()
 			.AutoHeight()
@@ -310,20 +266,25 @@
 			]
 		]
 	];
+
+	CreateViewModelTreeView(UE::PropertyViewer::SPropertyViewer::EPropertyVisibility::Visible);
+
+	FillViewModel();
 }
 
-SMVVMViewModelPanel::SMVVMViewModelPanel() = default;
+SMVVMViewModelPanel::SMVVMViewModelPanel()
+	: PropertyVisibility(UE::PropertyViewer::SPropertyViewer::EPropertyVisibility::Visible)
+{}
 
 SMVVMViewModelPanel::~SMVVMViewModelPanel()
 {
-	if (ExtensionAddeddHandle.IsValid())
+	TSharedPtr<FWidgetBlueprintEditor> WidgetBlueprintEditor = WeakBlueprintEditor.Pin();
+
+	if (ExtensionAddeddHandle.IsValid() && WidgetBlueprintEditor.IsValid())
 	{
-		if (TSharedPtr<FWidgetBlueprintEditor> WidgetBlueprintEditor = WeakBlueprintEditor.Pin())
+		if (UWidgetBlueprint* WidgetBlueprint = WidgetBlueprintEditor->GetWidgetBlueprintObj())
 		{
-			if (UWidgetBlueprint* WidgetBlueprint = WidgetBlueprintEditor->GetWidgetBlueprintObj())
-			{
-				WidgetBlueprint->OnExtensionAdded.Remove(ExtensionAddeddHandle);
-			}
+			WidgetBlueprint->OnExtensionAdded.Remove(ExtensionAddeddHandle);
 		}
 	}
 
@@ -336,11 +297,16 @@
 		}
 	}
 
-	if (WidgetBlueprintTransactionHandle.IsValid())
+	if (WidgetBlueprintTransactionHandle.IsValid() && WidgetBlueprintEditor.IsValid())
 	{
-		if (TSharedPtr<FWidgetBlueprintEditor> WidgetBlueprintEditor = WeakBlueprintEditor.Pin())
+		WidgetBlueprintEditor->GetOnWidgetBlueprintTransaction().Remove(WidgetBlueprintTransactionHandle);
+	}
+
+	if (UObjectInitialized() && WidgetBlueprintEditor.IsValid())
+	{
+		if (UWidgetBlueprint* WidgetBlueprint = WidgetBlueprintEditor->GetWidgetBlueprintObj())
 		{
-			WidgetBlueprintEditor->GetOnWidgetBlueprintTransaction().Remove(WidgetBlueprintTransactionHandle);
+			WidgetBlueprint->OnSetObjectBeingDebugged().RemoveAll(this);
 		}
 	}
 }
@@ -356,14 +322,58 @@
 	{
 		ViewModelTreeView->RemoveAll();
 
+		UUserWidget* Widget = nullptr;
+		UMVVMView* WidgetView = nullptr;
+
+		if (TSharedPtr<FWidgetBlueprintEditor> WidgetBlueprintEditor = WeakBlueprintEditor.Pin())
+		{
+			if (UWidgetBlueprint* WidgetBlueprint = WidgetBlueprintEditor->GetWidgetBlueprintObj())
+			{
+				Widget = Cast<UUserWidget>(WidgetBlueprint->GetObjectBeingDebugged());
+
+				if (Widget)
+				{
+					WidgetView = UMVVMSubsystem::GetViewFromUserWidget(Widget);
+				}
+			}
+		}
+
 		if (UMVVMBlueprintView* View = WeakBlueprintView.Get())
 		{
-			for (const FMVVMBlueprintViewModelContext& ViewModelContext : View->GetViewModels())
+			const TArrayView<const FMVVMBlueprintViewModelContext> ViewModelContexts = View->GetViewModels();
+
+			for (const FMVVMBlueprintViewModelContext& ViewModelContext : ViewModelContexts)
 			{
-				const UE::PropertyViewer::SPropertyViewer::FHandle Handle = ViewModelTreeView->AddContainer(ViewModelContext.GetViewModelClass());
-				PropertyViewerHandles.Add(Handle, ViewModelContext.GetViewModelId());
+				UE::PropertyViewer::SPropertyViewer::FHandle Handle;
+				UClass* ViewModelClass = ViewModelContext.GetViewModelClass();
+
+				if (Widget)
+				{
+					if (WidgetView)
+					{
+						if (UObject* Object = WidgetView->GetViewModel(ViewModelContext.GetViewModelName()).GetObject())
+						{
+							Handle = ViewModelTreeView->AddInstance(Object);
+						}
+					}
+				}
+				else
+				{
+					Handle = ViewModelTreeView->AddContainer(ViewModelClass);
+				}
+
+				if (Handle.IsValid())
+				{
+					PropertyViewerHandles.Add(Handle, ViewModelContext.GetViewModelId());
+				}
 			}
 		}
+
+		// Only add widget values when in editable mode.
+		if (Widget && PropertyVisibility == UE::PropertyViewer::SPropertyViewer::EPropertyVisibility::Editable)
+		{
+			ViewModelTreeView->AddInstance(Widget);
+		}
 	}
 }
 
@@ -859,6 +869,34 @@
 	}
 }
 
+void SMVVMViewModelPanel::HandleSetObjectBeingDebugged(UObject* DebugObject)
+{
+	using namespace UE::PropertyViewer;
 
+	SPropertyViewer::EPropertyVisibility RequiredPropertyVisibility = SPropertyViewer::EPropertyVisibility::Visible;
+
+	// If there's a debug object, we confirm it's a user widget and set the properties to be editable.
+	// If there's not a debug object, leave the properties as visible (not editable).
+	if (DebugObject && DebugObject->IsA<UUserWidget>())
+	{
+		if (TSharedPtr<FWidgetBlueprintEditor> WidgetBlueprintEditor = WeakBlueprintEditor.Pin())
+		{
+			if (DebugObject->GetClass()->ClassGeneratedBy == WidgetBlueprintEditor->GetWidgetBlueprintObj())
+			{
+				RequiredPropertyVisibility = SPropertyViewer::EPropertyVisibility::Editable;
+			}
+		}
+	}
+
+	// Cannot change this after creation
+	if (PropertyVisibility != RequiredPropertyVisibility)
+	{
+		CreateViewModelTreeView(RequiredPropertyVisibility);
+	}
+
+	FillViewModel();
+}
+
+
 bool SMVVMViewModelPanel::RenameViewModelProperty(FGuid ViewModelGuid, const FText& RenameTo, bool bCommit, FText& OutErrorMessage) const
 {
@@ -911,6 +949,64 @@
 	return false;
 }
 
+void SMVVMViewModelPanel::CreateViewModelTreeView(UE::PropertyViewer::SPropertyViewer::EPropertyVisibility InPropertyVisibility)
+{
+	PropertyVisibility = InPropertyVisibility;
 
+	if (!ToolbarWidget.IsValid() || SettingsWidget.IsValid())
+	{
+		FToolMenuContext GenerateWidgetContext;
+		{
+			UWidgetBlueprintToolMenuContext* WidgetBlueprintMenuContext = NewObject<UWidgetBlueprintToolMenuContext>();
+			WidgetBlueprintMenuContext->WidgetBlueprintEditor = WeakBlueprintEditor;
+			GenerateWidgetContext.AddObject(WidgetBlueprintMenuContext);
+
+			UMVVMViewModelPanelToolMenuContext* ViewModelPanelToolMenuContext = NewObject<UMVVMViewModelPanelToolMenuContext>();
+			ViewModelPanelToolMenuContext->ViewModelPanel = SharedThis(this);
+			GenerateWidgetContext.AddObject(ViewModelPanelToolMenuContext);
+		}
+
+		ToolbarWidget = UToolMenus::Get()->GenerateWidget("MVVM.Viewmodels.Toolbar", GenerateWidgetContext);
+		SettingsWidget = UToolMenus::Get()->GenerateWidget("MVVM.Viewmodels.Settings", GenerateWidgetContext);
+	}
+
+	ViewModelTreeView = SNew(UE::PropertyViewer::SPropertyViewer)
+		.PropertyVisibility(InPropertyVisibility)
+		.bShowSearchBox(true)
+		.bShowFieldIcon(true)
+		.bSanitizeName(true)
+		.FieldIterator(FieldIterator.Get())
+		.FieldExpander(FieldExpander.Get())
+		.OnGetPreSlot(this, &SMVVMViewModelPanel::HandleGetPreSlot)
+		.OnContextMenuOpening(this, &SMVVMViewModelPanel::HandleContextMenuOpening)
+		.OnSelectionChanged(this, &SMVVMViewModelPanel::HandleSelectionChanged)
+		.OnGenerateContainer(this, &SMVVMViewModelPanel::HandleGenerateContainer)
+		.OnDragDetected(this, &SMVVMViewModelPanel::HandleDragDetected)
+		.SearchBoxPreSlot()
+		[
+			ToolbarWidget.ToSharedRef()
+		]
+		.SearchBoxPostSlot()
+		[
+			SNew(SComboButton)
+			.HasDownArrow(false)
+			.ContentPadding(0)
+			.ForegroundColor(FSlateColor::UseForeground())
+			.ButtonStyle(FAppStyle::Get(), "SimpleButton")
+			.MenuContent()
+			[
+				SettingsWidget.ToSharedRef()
+			]
+			.ButtonContent()
+			[
+				SNew(SImage)
+				.Image(FAppStyle::GetBrush("DetailsView.ViewOptions"))
+			]
+		];
+
+	ViewModelTreeViewContainer->SetContent(ViewModelTreeView.ToSharedRef());
+}
+
+
 EVisibility SMVVMViewModelPanel::GetWarningPanelVisibility() const
 {
diff --git a/Engine/Plugins/Runtime/ModelViewViewModel/Source/ModelViewViewModelEditor/Private/Widgets/SMVVMViewModelPanel.h b/Engine/Plugins/Runtime/ModelViewViewModel/Source/ModelViewViewModelEditor/Private/Widgets/SMVVMViewModelPanel.h
--- a/Engine/Plugins/Runtime/ModelViewViewModel/Source/ModelViewViewModelEditor/Private/Widgets/SMVVMViewModelPanel.h
+++ b/Engine/Plugins/Runtime/ModelViewViewModel/Source/ModelViewViewModelEditor/Private/Widgets/SMVVMViewModelPanel.h
@@ -17,6 +17,7 @@ namespace UE::PropertyViewer { class FFieldExpander_Default; }
 
 class INotifyFieldValueChanged;
 class FWidgetBlueprintEditor;
+class SBox;
 class SInlineEditableTextBlock;
 class SPositiveActionButton;
 class UBlueprintExtension;
@@ -86,6 +87,7 @@ private:
 	TSharedPtr<SWidget> HandleContextMenuOpening(UE::PropertyViewer::SPropertyViewer::FHandle ContainerHandle, TArrayView<const FFieldVariant> Field) const;
 	void HandleSelectionChanged(UE::PropertyViewer::SPropertyViewer::FHandle, TArrayView<const FFieldVariant>, ESelectInfo::Type);
 	void HandleEditorSelectionChanged();
+	void HandleSetObjectBeingDebugged(UObject* DebugObject);
 	bool HandleCanRename(FGuid ViewModelGuid) const;
 	bool HandleVerifyNameTextChanged(const FText& InText, FText& OutErrorMessage, FGuid ViewModelGuid);
 	void HandleNameTextCommited(const FText& InText, ETextCommit::Type CommitInfo, FGuid ViewModelGuid);
@@ -105,6 +107,8 @@ private:
 	void FillViewModel();
 	bool RenameViewModelProperty(FGuid ViewModelGuid, const FText& RenameTo, bool bCommit, FText& OutErrorMessage) const;
 
+	void CreateViewModelTreeView(UE::PropertyViewer::SPropertyViewer::EPropertyVisibility InPropertyVisibility);
+
 	//~ Begin FNotifyHook
 	virtual void NotifyPreChange(FProperty* PropertyAboutToChange) override;
 	virtual void NotifyPostChange(const FPropertyChangedEvent& PropertyChangedEvent, FProperty* PropertyThatChanged) override;
@@ -112,6 +116,7 @@ private:
 
 private:
 	TSharedPtr<SPositiveActionButton> AddMenuButton;
+	TSharedPtr<SBox> ViewModelTreeViewContainer;
 	TSharedPtr<UE::PropertyViewer::SPropertyViewer> ViewModelTreeView;
 	TUniquePtr<FFieldIterator_Bindable> FieldIterator;
 	TUniquePtr<UE::PropertyViewer::FFieldExpander_Default> FieldExpander;
@@ -132,6 +137,10 @@ private:
 	FName PreviousViewModelPropertyName;
 	bool bIsViewModelSelecting = false;
 	bool bDisableWarningPanel = false;
+
+	UE::PropertyViewer::SPropertyViewer::EPropertyVisibility PropertyVisibility;
+	TSharedPtr<SWidget> ToolbarWidget;
+	TSharedPtr<SWidget> SettingsWidget;
 };
 
 } // namespace
diff --git a/Engine/Source/Developer/SlateReflector/Private/Models/WidgetReflectorNode.cpp b/Engine/Source/Developer/SlateReflector/Private/Models/WidgetReflectorNode.cpp
--- a/Engine/Source/Developer/SlateReflector/Private/Models/WidgetReflectorNode.cpp
+++ b/Engine/Source/Developer/SlateReflector/Private/Models/WidgetReflectorNode.cpp
@@ -249,6 +249,16 @@ bool FLiveWidgetReflectorNode::GetWidgetEnabled() const
 	return FWidgetReflectorNodeUtils::GetWidgetEnabled(Widget.Pin());
 }
 
+UObject* FLiveWidgetReflectorNode::GetWidgetSourceObject() const
+{
+	if (Widget.IsValid())
+	{
+		return FReflectionMetaData::GetWidgetSourceObject(*Widget.Pin());
+	}
+
+	return nullptr;
+}
+
 
 /**
  * -----------------------------------------------------------------------------
@@ -449,6 +459,11 @@ bool FSnapshotWidgetReflectorNode::GetWidgetEnabled() const
 	return bCachedWidgetEnabled;
 }
 
+UObject* FSnapshotWidgetReflectorNode::GetWidgetSourceObject() const
+{
+	return nullptr;
+}
+
 TSharedRef<FJsonValue> FSnapshotWidgetReflectorNode::ToJson(const TSharedRef<FSnapshotWidgetReflectorNode>& RootSnapshotNode)
 {
 	struct Internal
diff --git a/Engine/Source/Developer/SlateReflector/Private/Models/WidgetReflectorNode.h b/Engine/Source/Developer/SlateReflector/Private/Models/WidgetReflectorNode.h
--- a/Engine/Source/Developer/SlateReflector/Private/Models/WidgetReflectorNode.h
+++ b/Engine/Source/Developer/SlateReflector/Private/Models/WidgetReflectorNode.h
@@ -13,6 +13,7 @@
 class FJsonValue;
 class FWidgetPath;
 class SWidget;
+class UObject;
 
 /**
  * Enum used as crude RTTI for the widget reflector node types
@@ -194,6 +195,11 @@ public:
 	 */
 	virtual bool IsWidgetLive() const = 0;
 
+	/**
+	 * @return The source object obtained from the widget's meta data.
+	 */
+	virtual UObject* GetWidgetSourceObject() const = 0;
+
 	/**
 	 * @return The geometry of the widget.
 	 */
@@ -319,6 +325,7 @@ public:
 	virtual TPointerAsInt GetWidgetAddress() const override;
 	virtual bool GetWidgetEnabled() const override;
 	virtual bool IsWidgetLive() const override { return true; }
+	virtual UObject* GetWidgetSourceObject() const override;
 
 private:
 	/**
@@ -383,6 +390,7 @@ public:
 	virtual TPointerAsInt GetWidgetAddress() const override;
 	virtual bool GetWidgetEnabled() const override;
 	virtual bool IsWidgetLive() const override { return false; }
+	virtual UObject* GetWidgetSourceObject() const override;
 
 	/** Save this node data as a JSON object */
 	static TSharedRef<FJsonValue> ToJson(const TSharedRef<FSnapshotWidgetReflectorNode>& RootSnapshotNode);
diff --git a/Engine/Source/Developer/SlateReflector/Private/Widgets/SWidgetReflector.cpp b/Engine/Source/Developer/SlateReflector/Private/Widgets/SWidgetReflector.cpp
--- a/Engine/Source/Developer/SlateReflector/Private/Widgets/SWidgetReflector.cpp
+++ b/Engine/Source/Developer/SlateReflector/Private/Widgets/SWidgetReflector.cpp
@@ -253,7 +253,12 @@ public:
 
 	virtual void SetAssetAccessDelegate(FAccessAsset InDelegate) override
 	{
-		AsseetAccessDelegate = InDelegate;
+		AssetAccessDelegate = InDelegate;
+	}
+	
+	virtual void SetDebugObjectAccessDelegate(FAccessDebugObject InDelegate)override
+	{
+		DebugObjectAccessDelegate = InDelegate;
 	}
 
 	virtual void SetWidgetsToVisualize( const FWidgetPath& InWidgetsToVisualize ) override;
@@ -472,7 +477,8 @@ private:
 	FGuid RemoteSnapshotRequestId;
 
 	FAccessSourceCode SourceAccessDelegate;
-	FAccessAsset AsseetAccessDelegate;
+	FAccessAsset AssetAccessDelegate;
+	FAccessDebugObject DebugObjectAccessDelegate;
 
 	EWidgetReflectorUIMode CurrentUIMode;
 	EWidgetPickingMode PickingMode;
@@ -1099,7 +1105,7 @@ TSharedRef<SDockTab> SWidgetReflector::SpawnWidgetList(const FSpawnTabArgs& Args
 		[
 			SNew(SWidgetList)
 			.OnAccessSource(SourceAccessDelegate)
-			.OnAccessAsset(AsseetAccessDelegate)
+			.OnAccessAsset(AssetAccessDelegate)
 		];
 }
 
@@ -2029,7 +2035,8 @@ TSharedRef<ITableRow> SWidgetReflector::HandleReflectorTreeGenerateRow( TSharedR
 		.WidgetInfoToVisualize(InReflectorNode)
 		.ToolTip(GenerateToolTipForReflectorNode(InReflectorNode))
 		.SourceCodeAccessor(SourceAccessDelegate)
-		.AssetAccessor(AsseetAccessDelegate);
+		.AssetAccessor(AssetAccessDelegate)
+		.DebugObjectAccessor(DebugObjectAccessDelegate);
 }
 
 void SWidgetReflector::HandleReflectorTreeGetChildren(TSharedRef<FWidgetReflectorNodeBase> InReflectorNode, TArray<TSharedRef<FWidgetReflectorNodeBase>>& OutChildren)
diff --git a/Engine/Source/Developer/SlateReflector/Private/Widgets/SWidgetReflectorTreeWidgetItem.cpp b/Engine/Source/Developer/SlateReflector/Private/Widgets/SWidgetReflectorTreeWidgetItem.cpp
--- a/Engine/Source/Developer/SlateReflector/Private/Widgets/SWidgetReflectorTreeWidgetItem.cpp
+++ b/Engine/Source/Developer/SlateReflector/Private/Widgets/SWidgetReflectorTreeWidgetItem.cpp
@@ -1,6 +1,7 @@
 // Copyright Epic Games, Inc. All Rights Reserved.
 
 #include "Widgets/SWidgetReflectorTreeWidgetItem.h"
+#include "Blueprint/UserWidget.h"
 #include "SlateOptMacros.h"
 #include "Widgets/Colors/SColorBlock.h"
 #include "Widgets/Input/SHyperlink.h"
@@ -31,6 +32,7 @@ void SReflectorTreeWidgetItem::Construct(const FArguments& InArgs, const TShared
 	this->WidgetInfo = InArgs._WidgetInfoToVisualize;
 	this->OnAccessSourceCode = InArgs._SourceCodeAccessor;
 	this->OnAccessAsset = InArgs._AssetAccessor;
+	this->OnAccessDebugObject = InArgs._DebugObjectAccessor;
 	this->SetPadding(0.f);
 
 	check(WidgetInfo.IsValid());
@@ -312,20 +314,52 @@ END_SLATE_FUNCTION_BUILD_OPTIMIZATION
 
 void SReflectorTreeWidgetItem::HandleHyperlinkNavigate()
 {
-	FAssetData CachedAssetData = WidgetInfo->GetWidgetAssetData();
-	if (CachedAssetData.IsValid())
+#if WITH_EDITOR
+	bool bHasOpenedAsset = false;
+
+	if (UUserWidget* SourceObject = Cast<UUserWidget>(WidgetInfo->GetWidgetSourceObject()))
+	{
+		UClass* SourceClass = SourceObject->GetClass();
+
+		if (SourceClass->ClassGeneratedBy)
+		{
+			OnAccessAsset.Execute(SourceClass->ClassGeneratedBy);
+			bHasOpenedAsset = true;
+
+			if (OnAccessDebugObject.IsBound())
+			{
+				OnAccessDebugObject.Execute(SourceObject);
+			}
+		}
+	}
+
+	if (!bHasOpenedAsset)
 	{
-		if (OnAccessAsset.IsBound())
+#endif
+		FAssetData CachedAssetData = WidgetInfo->GetWidgetAssetData();
+
+		if (CachedAssetData.IsValid())
 		{
-			CachedAssetData.GetPackage();
-			OnAccessAsset.Execute(CachedAssetData.GetAsset());
-			return;
+			if (OnAccessAsset.IsBound())
+			{
+				CachedAssetData.GetPackage();
+				OnAccessAsset.Execute(CachedAssetData.GetAsset());
+			}
 		}
+#if WITH_EDITOR
 	}
+#endif
 
 	if (OnAccessSourceCode.IsBound())
 	{
-		OnAccessSourceCode.Execute(WidgetInfo->GetWidgetFile(), WidgetInfo->GetWidgetLineNumber(), 0);
+		const FString WidgetFile = WidgetInfo->GetWidgetFile();
+		const FString WidgetFilePath = TEXT("\\Engine\\Source\\Runtime\\UMG\\Private\\Components\\Widget.cpp");
+		const FString UserWidgetFilePath = TEXT("\\Engine\\Source\\Runtime\\UMG\\Private\\Components\\UserWidget.cpp");
+
+		if (!WidgetFile.EndsWith(WidgetFilePath) && !WidgetFile.EndsWith(UserWidgetFilePath))
+		{
+			OnAccessSourceCode.Execute(WidgetInfo->GetWidgetFile(), WidgetInfo->GetWidgetLineNumber(), 0);
+		}
 	}
 }
 
diff --git a/Engine/Source/Developer/SlateReflector/Private/Widgets/SWidgetReflectorTreeWidgetItem.h b/Engine/Source/Developer/SlateReflector/Private/Widgets/SWidgetReflectorTreeWidgetItem.h
--- a/Engine/Source/Developer/SlateReflector/Private/Widgets/SWidgetReflectorTreeWidgetItem.h
+++ b/Engine/Source/Developer/SlateReflector/Private/Widgets/SWidgetReflectorTreeWidgetItem.h
@@ -41,6 +41,7 @@ public:
 		SLATE_ARGUMENT(TSharedPtr<FWidgetReflectorNodeBase>, WidgetInfoToVisualize)
 		SLATE_ARGUMENT(FAccessSourceCode, SourceCodeAccessor)
 		SLATE_ARGUMENT(FAccessAsset, AssetAccessor)
+		SLATE_ARGUMENT(FAccessDebugObject, DebugObjectAccessor)
 
 	SLATE_END_ARGS()
 
@@ -75,4 +76,5 @@ private:
 
 	FAccessSourceCode OnAccessSourceCode;
 	FAccessAsset OnAccessAsset;
+	FAccessDebugObject OnAccessDebugObject;
 };
diff --git a/Engine/Source/Developer/SlateReflector/SlateReflector.Build.cs b/Engine/Source/Developer/SlateReflector/SlateReflector.Build.cs
--- a/Engine/Source/Developer/SlateReflector/SlateReflector.Build.cs
+++ b/Engine/Source/Developer/SlateReflector/SlateReflector.Build.cs
@@ -26,7 +26,8 @@ public class SlateReflector : ModuleRules
 				"ToolWidgets",
 				"DesktopPlatform",
 				"ImageCore",
-				"ImageWrapper"
+				"ImageWrapper",
+				"UMG"
             }
 		);
 
diff --git a/Engine/Source/Editor/Kismet/Private/SBlueprintEditorSelectedDebugObjectWidget.cpp b/Engine/Source/Editor/Kismet/Private/SBlueprintEditorSelectedDebugObjectWidget.cpp
--- a/Engine/Source/Editor/Kismet/Private/SBlueprintEditorSelectedDebugObjectWidget.cpp
+++ b/Engine/Source/Editor/Kismet/Private/SBlueprintEditorSelectedDebugObjectWidget.cpp
@@ -39,6 +39,7 @@
 #include "UObject/UObjectHash.h"
 #include "UObject/UObjectIterator.h"
 #include "UnrealEdGlobals.h"
+#include "WidgetBlueprint.h"
 #include "Widgets/Input/STextComboBox.h"
 #include "Widgets/SBoxPanel.h"
 #include "Widgets/SToolTip.h"
@@ -55,6 +56,19 @@ static TAutoConsoleVariable<int32> CVarUseFastDebugObjectDiscovery(TEXT("r.UseFa
 //////////////////////////////////////////////////////////////////////////
 // SBlueprintEditorSelectedDebugObjectWidget
 
+TSharedRef<FBlueprintDebugObjectInstance> SBlueprintEditorSelectedDebugObjectWidget::NoDebugInstance = MakeShared<FBlueprintDebugObjectInstance>(nullptr, GetNoDebugString());
+
+SBlueprintEditorSelectedDebugObjectWidget::~SBlueprintEditorSelectedDebugObjectWidget()
+{
+	if (UObjectInitialized() && GEditor)
+	{
+		if (UAssetEditorSubsystem* AssetEditorSubsystem = GEditor->GetEditorSubsystem<UAssetEditorSubsystem>())
+		{
+			AssetEditorSubsystem->GetDebugObjectDelegateForClass(UObject::StaticClass())->RemoveAll(this);
+		}
+	}
+}
+
 void SBlueprintEditorSelectedDebugObjectWidget::Construct(const FArguments& InArgs, TSharedPtr<FBlueprintEditor> InBlueprintEditor)
 {
 	BlueprintEditor = InBlueprintEditor;
@@ -63,6 +77,9 @@ void SBlueprintEditorSelectedDebugObjectWidget::Construct(const FArguments& InAr
 	GenerateDebugObjectInstances(false);
 
 	LastObjectObserved = nullptr;
+	NoDebugInstance->bPinned = true;
+
+	TSharedPtr<FBlueprintDebugObjectInstance> CurrentDebugInstance = GetDebugObjectInstance();
 
 	DebugWorldsComboBox = SNew(STextComboBox)
 		.ToolTip(IDocumentation::Get()->CreateToolTip(
@@ -84,7 +101,7 @@ void SBlueprintEditorSelectedDebugObjectWidget::Construct(const FArguments& InAr
 		TEXT("Shared/Editors/BlueprintEditor/BlueprintDebugger"),
 		TEXT("DebugObject")))
 		.OptionsSource(&DebugObjects)
-		.InitiallySelectedItem(GetDebugObjectInstance())
+		.InitiallySelectedItem(CurrentDebugInstance)
 		.OnComboBoxOpening(this, &SBlueprintEditorSelectedDebugObjectWidget::GenerateDebugObjectInstances, true)
 		.OnSelectionChanged(this, &SBlueprintEditorSelectedDebugObjectWidget::DebugObjectSelectionChanged)
 		.OnGenerateWidget(this, &SBlueprintEditorSelectedDebugObjectWidget::CreateDebugObjectItemWidget)
@@ -101,6 +118,18 @@ void SBlueprintEditorSelectedDebugObjectWidget::Construct(const FArguments& InAr
 		.UseLowDetailSlot(FMultiBoxSettings::UseSmallToolBarIcons)
 		.OnGetActiveDetailSlotContent(this, &SBlueprintEditorSelectedDebugObjectWidget::OnGetActiveDetailSlotContent)
 	];
+
+	if (UAssetEditorSubsystem* AssetEditorSubsystem = GEditor->GetEditorSubsystem<UAssetEditorSubsystem>())
+	{
+		AssetEditorSubsystem->GetDebugObjectDelegateForClass(UObject::StaticClass())
+			->AddSP(this, &SBlueprintEditorSelectedDebugObjectWidget::HandleOnDebugObjectAccess);
+	}
+
+	// If we construct this widget with an already debugged object, the object should be pinned as it was opened via debug object access.
+	if (CurrentDebugInstance.IsValid())
+	{
+		CurrentDebugInstance->bPinned = true;
+	}
 }
 
 void SBlueprintEditorSelectedDebugObjectWidget::Tick( const FGeometry& AllottedGeometry, const double InCurrentTime, const float InDeltaTime )
@@ -134,7 +163,7 @@ void SBlueprintEditorSelectedDebugObjectWidget::Tick( const FGeometry& AllottedG
 	}
 }
 
-const FString& SBlueprintEditorSelectedDebugObjectWidget::GetNoDebugString() const
+const FString& SBlueprintEditorSelectedDebugObjectWidget::GetNoDebugString()
 {
 	return NSLOCTEXT("BlueprintEditor", "DebugObjectNothingSelected", "No debug object selected").ToString();
 }
@@ -309,9 +338,20 @@ void SBlueprintEditorSelectedDebugObjectWidget::GenerateDebugObjectInstances(boo
 	// Cache the current selection as we may need to restore it
 	TSharedPtr<FBlueprintDebugObjectInstance> LastSelection = GetDebugObjectInstance();
 
+	TArray<TWeakObjectPtr<UObject>> LastPinnedObjects;
+	LastPinnedObjects.Reserve(DebugObjects.Num());
+
+	for (const TSharedPtr<FBlueprintDebugObjectInstance>& DebugObject : DebugObjects)
+	{
+		if (DebugObject.IsValid() && DebugObject->bPinned && DebugObject->ObjectPtr.IsValid())
+		{
+			LastPinnedObjects.Add(DebugObject->ObjectPtr);
+		}
+	}
+
 	// Empty the lists of actors and regenerate them
 	DebugObjects.Empty();
-	DebugObjects.Add(MakeShareable(new FBlueprintDebugObjectInstance(nullptr, GetNoDebugString())));
+	DebugObjects.Add(NoDebugInstance);
 
 	// Grab custom objects that should always be visible, regardless of the world
 	TArray<FCustomDebugObject> CustomDebugObjects;
@@ -469,6 +509,21 @@ void SBlueprintEditorSelectedDebugObjectWidget::GenerateDebugObjectInstances(boo
 		}
 	}
 
+	for (const TSharedPtr<FBlueprintDebugObjectInstance>& DebugObject : DebugObjects)
+	{
+		if (DebugObject == NoDebugInstance)
+		{
+			continue;
+		}
+
+		if (DebugObject.IsValid() && DebugObject->ObjectPtr.IsValid() && LastPinnedObjects.Contains(DebugObject->ObjectPtr))
+		{
+			DebugObject->bPinned = true;
+		}
+	}
+
+	SortDebugObjects();
+
 	if (DebugObjectsComboBox.IsValid())
 	{
 		if (bRestoreSelection)
@@ -794,9 +849,34 @@ TSharedRef<SWidget> SBlueprintEditorSelectedDebugObjectWidget::CreateDebugObject
 		ItemTooltip = InItem->ObjectPath;
 	}
 
-	return SNew(STextBlock)
-		.Text(FText::FromString(*ItemString))
-		.ToolTipText(FText::FromString(*ItemTooltip));
+	TSharedRef<SHorizontalBox> Container = SNew(SHorizontalBox)
+		+ SHorizontalBox::Slot()
+		.FillWidth(1.f)
+		.HAlign(HAlign_Left)
+		.VAlign(VAlign_Center)
+		[
+			SNew(STextBlock)
+			.Text(FText::FromString(*ItemString))
+			.ToolTipText(FText::FromString(*ItemTooltip))
+			.Clipping(EWidgetClipping::ClipToBounds)
+		];
+		
+		
+	if (!InItem->ObjectPtr.IsExplicitlyNull())
+	{
+		Container->AddSlot()
+			.AutoWidth()
+			.VAlign(VAlign_Center)
+			[
+				SNew(SImage)
+				.Image(this, &SBlueprintEditorSelectedDebugObjectWidget::GetDebugObjectItemWidgetPinBrush, InItem.ToWeakPtr())
+				.Visibility(this, &SBlueprintEditorSelectedDebugObjectWidget::GetDebugObjectItemWidgetPinVisibility, InItem.ToWeakPtr(), StaticCastWeakPtr<SWidget>(Container.ToWeakPtr()))
+				.DesiredSizeOverride(FVector2D(16.f))
+				.OnMouseButtonDown(this, &SBlueprintEditorSelectedDebugObjectWidget::OnDebugObjectItemWidgetPinMouseDown, InItem.ToWeakPtr())
+			];
+	}
+
+	return Container;
 }
 
 FText SBlueprintEditorSelectedDebugObjectWidget::GetSelectedDebugObjectTextLabel() const
@@ -823,6 +903,101 @@ FText SBlueprintEditorSelectedDebugObjectWidget::GetSelectedDebugObjectTextLabel
 	return FText::FromString(Label);
 }
 
+void SBlueprintEditorSelectedDebugObjectWidget::HandleOnDebugObjectAccess(UObject* InDebugObject)
+{
+	for (const TSharedPtr<FBlueprintDebugObjectInstance>& DebugObject : DebugObjects)
+	{
+		if (DebugObject.IsValid() && DebugObject->ObjectPtr.Get() == InDebugObject)
+		{
+			DebugObject->bPinned = true;
+			LastObjectObserved = InDebugObject;
+			break;
+		}
+	}
+}
+
+void SBlueprintEditorSelectedDebugObjectWidget::SortDebugObjects()
+{	
+	DebugObjects.StableSort(
+		[](const TSharedPtr<FBlueprintDebugObjectInstance>& A, const TSharedPtr<FBlueprintDebugObjectInstance>& B)
+		{
+			if (A == NoDebugInstance)
+			{
+				return true;
+			}
+
+			if (B == NoDebugInstance)
+			{
+				return false;
+			}
+
+			if (!B.IsValid())
+			{
+				return false;
+			}
+			
+			if (!A.IsValid())
+			{
+				return true;
+			}
+
+			if (A->bPinned == B->bPinned)
+			{
+				return false;
+			}
+
+			return A->bPinned;
+		}
+	);
+}
+
+const FSlateBrush* SBlueprintEditorSelectedDebugObjectWidget::GetDebugObjectItemWidgetPinBrush(TWeakPtr<FBlueprintDebugObjectInstance> InInstanceWeak) const
+{
+	if (TSharedPtr<FBlueprintDebugObjectInstance> Instance = InInstanceWeak.Pin())
+	{
+		if (Instance->bPinned)
+		{
+			return FAppStyle::GetBrush("Icons.Pinned");
+		}
+	}
+
+	return FAppStyle::GetBrush("Icons.Unpinned");
+}
+
+EVisibility SBlueprintEditorSelectedDebugObjectWidget::GetDebugObjectItemWidgetPinVisibility(TWeakPtr<FBlueprintDebugObjectInstance> InInstanceWeak, TWeakPtr<SWidget> InRowWidgetWeak) const
+{
+	if (TSharedPtr<FBlueprintDebugObjectInstance> Instance = InInstanceWeak.Pin())
+	{
+		if (Instance->bPinned)
+		{
+			return EVisibility::Visible;
+		}
+
+		if (TSharedPtr<SWidget> RowWidget = InRowWidgetWeak.Pin())
+		{
+			if (RowWidget->GetParentWidget().IsValid())
+			{
+				if (RowWidget->GetParentWidget()->IsHovered())
+				{
+					return EVisibility::Visible;
+				}
+			}
+		}
+	}
+
+	return EVisibility::Hidden;
+}
+
+FReply SBlueprintEditorSelectedDebugObjectWidget::OnDebugObjectItemWidgetPinMouseDown(const FGeometry& InGeometry, const FPointerEvent& InPointerEvent, TWeakPtr<FBlueprintDebugObjectInstance> InInstanceWeak)
+{
+	if (TSharedPtr<FBlueprintDebugObjectInstance> Instance = InInstanceWeak.Pin())
+	{
+		Instance->bPinned = !Instance->bPinned;
+	}
+
+	return FReply::Handled();
+}
+
 //////////////////////////////////////////////////////////////////////////
 
 #undef LOCTEXT_NAMESPACE
diff --git a/Engine/Source/Editor/Kismet/Public/SBlueprintEditorSelectedDebugObjectWidget.h b/Engine/Source/Editor/Kismet/Public/SBlueprintEditorSelectedDebugObjectWidget.h
--- a/Engine/Source/Editor/Kismet/Public/SBlueprintEditorSelectedDebugObjectWidget.h
+++ b/Engine/Source/Editor/Kismet/Public/SBlueprintEditorSelectedDebugObjectWidget.h
@@ -32,6 +32,7 @@ class UBlueprint;
 class UObject;
 class UWorld;
 struct FGeometry;
+struct FSlateBrush;
 
 struct FBlueprintDebugObjectInstance
 {
@@ -47,9 +48,12 @@ struct FBlueprintDebugObjectInstance
 	/** Object path to object in the editor, will only be set for static objects */
 	FString EditorObjectPath;
 
+	bool bPinned;
+
 	FBlueprintDebugObjectInstance(TWeakObjectPtr<UObject> InPtr, const FString& InLabel)
 		: ObjectPtr(InPtr)
 		, ObjectLabel(InLabel)
+		, bPinned(false)
 	{
 	}
 
@@ -78,6 +82,8 @@ public:
 	SLATE_BEGIN_ARGS(SBlueprintEditorSelectedDebugObjectWidget){}
 	SLATE_END_ARGS()
 
+	UE_API virtual ~SBlueprintEditorSelectedDebugObjectWidget();
+
 	UE_API void Construct(const FArguments& InArgs, TSharedPtr<FBlueprintEditor> InBlueprintEditor);
 
 	// SWidget interface
@@ -119,7 +125,7 @@ private:
 	UE_API EVisibility IsDebugWorldComboVisible() const;
 
 	/* Returns the string to indicate no debug object is selected */
-	UE_API const FString& GetNoDebugString() const;
+	UE_API static const FString& GetNoDebugString();
 
 	UE_API const FString& GetDebugAllWorldsString() const;
 
@@ -137,10 +143,27 @@ private:
 	/** Returns the combo button label to use for the currently-selected debug object item */
 	UE_API FText GetSelectedDebugObjectTextLabel() const;
 
+	/** Handle selection of debug object through debug object access. */
+	UE_API void HandleOnDebugObjectAccess(UObject* InDebugObject);
+
+	/** Sorts the debug objects based on pinned status. */
+	UE_API void SortDebugObjects();
+
+	/** Returns either a pinned or unpinned image based on object pin status. */
+	const FSlateBrush* GetDebugObjectItemWidgetPinBrush(TWeakPtr<FBlueprintDebugObjectInstance> InInstanceWeak) const;
+
+	/** Returns Visible if the item is pinned or hovered. */
+	EVisibility GetDebugObjectItemWidgetPinVisibility(TWeakPtr<FBlueprintDebugObjectInstance> InInstanceWeak, TWeakPtr<SWidget> InRowWidgetWeak) const;
+
+	/** Pins or unpins the debug object. */
+	FReply OnDebugObjectItemWidgetPinMouseDown(const FGeometry& InGeometry, const FPointerEvent& InPointerEvent, TWeakPtr<FBlueprintDebugObjectInstance> InInstanceWeak);
+
 private:
 	/** Pointer back to the blueprint editor tool that owns us */
 	TWeakPtr<FBlueprintEditor> BlueprintEditor;
 
+	static TSharedRef<FBlueprintDebugObjectInstance> NoDebugInstance;
+
 	/** Lists of actors of a given blueprint type and their names */
 	TArray<TSharedPtr<FBlueprintDebugObjectInstance>> DebugObjects;
 
diff --git a/Engine/Source/Editor/UMGEditor/Private/WidgetBlueprintEditor.cpp b/Engine/Source/Editor/UMGEditor/Private/WidgetBlueprintEditor.cpp
--- a/Engine/Source/Editor/UMGEditor/Private/WidgetBlueprintEditor.cpp
+++ b/Engine/Source/Editor/UMGEditor/Private/WidgetBlueprintEditor.cpp
@@ -157,6 +157,14 @@ FWidgetBlueprintEditor::~FWidgetBlueprintEditor()
 	{
 		return SequencerAddTrackExtenderHandle == Extender.GetHandle();
 	});
+
+	if (Blueprint)
+	{
+		if (UAssetEditorSubsystem* AssetEditorSubsystem = GEditor->GetEditorSubsystem<UAssetEditorSubsystem>())
+		{
+			AssetEditorSubsystem->GetDebugObjectDelegateForClass(Blueprint->GeneratedClass)->RemoveAll(this);
+		}
+	}
 }
 
 void FWidgetBlueprintEditor::InitWidgetBlueprintEditor(const EToolkitMode::Type Mode, const TSharedPtr< IToolkitHost >& InitToolkitHost, const TArray<UBlueprint*>& InBlueprints, bool bShouldOpenInDefaultsMode)
@@ -243,6 +251,12 @@ void FWidgetBlueprintEditor::InitWidgetBlueprintEditor(const EToolkitMode::Type
 	TSharedPtr<class IToolkitHost> PinnedToolkitHost = ToolkitHost.Pin();
 	check(PinnedToolkitHost.IsValid());
 	ModeUILayer = MakeShared<FWidgetEditorModeUILayer>(PinnedToolkitHost.Get());
+
+	if (UAssetEditorSubsystem* AssetEditorSubsystem = GEditor->GetEditorSubsystem<UAssetEditorSubsystem>())
+	{
+		AssetEditorSubsystem->GetDebugObjectDelegateForClass(Blueprint->GeneratedClass)
+			->AddSP(this, &FWidgetBlueprintEditor::OnDebugObjectAccess);
+	}
 }
 
 void FWidgetBlueprintEditor::InitalizeExtenders()
@@ -2660,6 +2674,21 @@ void FWidgetBlueprintEditor::OnFinishedChangingDynamicBindingProperties(const FP
 	}
 }
 
+void FWidgetBlueprintEditor::OnDebugObjectAccess(UObject* DebugObject)
+{
+	if (DebugObject)
+	{
+		if (UWidgetBlueprint* WidgetBlueprint = GetWidgetBlueprintObj())
+		{
+			if (WidgetBlueprint->GetObjectBeingDebugged() != DebugObject)
+			{
+				WidgetBlueprint->SetObjectBeingDebugged(DebugObject);
+				RefreshMyBlueprint();
+			}
+		}
+	}
+}
+
 void FWidgetBlueprintEditor::AddSlotTrack( UPanelSlot* Slot )
 {
 	TSharedPtr<ISequencer>& ActiveSequencer = GetSequencer();
diff --git a/Engine/Source/Editor/UMGEditor/Public/WidgetBlueprintEditor.h b/Engine/Source/Editor/UMGEditor/Public/WidgetBlueprintEditor.h
--- a/Engine/Source/Editor/UMGEditor/Public/WidgetBlueprintEditor.h
+++ b/Engine/Source/Editor/UMGEditor/Public/WidgetBlueprintEditor.h
@@ -445,9 +445,12 @@ private:
 	/** Update the name of a track to reflect changes in bindings */
 	UE_API void UpdateTrackName(FGuid ObjectId);
 
-	/* Called after changing dynamic binding properties*/
+	/** Called after changing dynamic binding properties*/
 	UE_API void OnFinishedChangingDynamicBindingProperties(const FPropertyChangedEvent& ChangeEvent, TSharedPtr<FStructOnScope> ValueStruct, FGuid ObjectId);
 
+	/** If editing the matching widget blueprint, sets the blueprint's current debug instance to the given object. */
+	void OnDebugObjectAccess(UObject* DebugObject);
+
 private:
 	/** The preview scene that owns the preview GUI */
 	FPreviewScene PreviewScene;
diff --git a/Engine/Source/Editor/UnrealEd/Classes/Editor/EditorEngine.h b/Engine/Source/Editor/UnrealEd/Classes/Editor/EditorEngine.h
--- a/Engine/Source/Editor/UnrealEd/Classes/Editor/EditorEngine.h
+++ b/Engine/Source/Editor/UnrealEd/Classes/Editor/EditorEngine.h
@@ -3174,6 +3174,9 @@ protected:
 	// Handle requests from slate application to open assets.
 	UNREALED_API bool HandleOpenAsset(UObject* Asset);
 
+	// Handle requests from slate application to open debug objects.
+	UNREALED_API bool HandleOpenDebugObject(UObject* DebugObject);
+
 private:
 	/** Handler for when any asset is loaded in the editor */
 	UNREALED_API void OnAssetLoaded( UObject* Asset );
diff --git a/Engine/Source/Editor/UnrealEd/Private/EditorEngine.cpp b/Engine/Source/Editor/UnrealEd/Private/EditorEngine.cpp
--- a/Engine/Source/Editor/UnrealEd/Private/EditorEngine.cpp
+++ b/Engine/Source/Editor/UnrealEd/Private/EditorEngine.cpp
@@ -1037,6 +1037,9 @@ void UEditorEngine::InitEditor(IEngineLoop* InEngineLoop)
 	{
 		// Setup a delegate to handle requests for opening assets
 		FSlateApplication::Get().SetWidgetReflectorAssetAccessDelegate(FAccessAsset::CreateUObject(this, &UEditorEngine::HandleOpenAsset));
+
+		// Setup a delegate to handle requests for opening debug objects
+		FSlateApplication::Get().SetWidgetReflectorDebugAccessDelegate(FAccessAsset::CreateUObject(this, &UEditorEngine::HandleOpenDebugObject));
 	}
 
 	IHotReloadModule& HotReloadModule = IHotReloadModule::Get();
@@ -1210,6 +1213,11 @@ bool UEditorEngine::HandleOpenAsset(UObject* Asset)
 	return GEditor->GetEditorSubsystem<UAssetEditorSubsystem>()->OpenEditorForAsset(Asset);
 }
 
+bool UEditorEngine::HandleOpenDebugObject(UObject* Object)
+{
+	return GEditor->GetEditorSubsystem<UAssetEditorSubsystem>()->OpenDebugObjectInEditor(Object);
+}
+
 void UEditorEngine::HandleSettingChanged( FName Name )
 {
 	// When settings are reset to default, the property name will be "None" so make sure that case is handled.
diff --git a/Engine/Source/Editor/UnrealEd/Private/Subsystems/AssetEditorSubsystem.cpp b/Engine/Source/Editor/UnrealEd/Private/Subsystems/AssetEditorSubsystem.cpp
--- a/Engine/Source/Editor/UnrealEd/Private/Subsystems/AssetEditorSubsystem.cpp
+++ b/Engine/Source/Editor/UnrealEd/Private/Subsystems/AssetEditorSubsystem.cpp
@@ -774,6 +774,36 @@ bool UAssetEditorSubsystem::OpenEditorForAssets_Advanced(const TArray <UObject*>
 	return NumNullAssets == 0;
 }
 
+bool UAssetEditorSubsystem::OpenDebugObjectInEditor(UObject* DebugObject)
+{
+	if (DebugObject)
+	{
+		DebugObjectDelegate.Broadcast(DebugObject);
+
+		for (UClass* DebugClass = DebugObject->GetClass(); !!DebugClass; DebugClass = DebugClass->GetSuperClass())
+		{
+			if (const FAssetEditorOpenDebugObject* Delegate = DebugObjectDelegates.Find(DebugClass->GetFName()))
+			{
+				Delegate->Broadcast(DebugObject);
+			}
+		}
+
+		return true;
+	}
+
+	return false;
+}
+
+FAssetEditorOpenDebugObject* UAssetEditorSubsystem::GetDebugObjectDelegateForClass(UClass* InClass)
+{
+	if (InClass)
+	{
+		return &DebugObjectDelegates.FindOrAdd(InClass->GetFName());
+	}
+
+	return nullptr;
+}
+
 bool UAssetEditorSubsystem::OpenEditorForAssets(const TArray<UObject*>& Assets, const EAssetTypeActivationOpenedMethod OpenedMethod)
 {
 	return OpenEditorForAssets_Advanced(Assets, EToolkitMode::Standalone, TSharedPtr<IToolkitHost>(), OpenedMethod);
diff --git a/Engine/Source/Editor/UnrealEd/Public/Subsystems/AssetEditorSubsystem.h b/Engine/Source/Editor/UnrealEd/Public/Subsystems/AssetEditorSubsystem.h
--- a/Engine/Source/Editor/UnrealEd/Public/Subsystems/AssetEditorSubsystem.h
+++ b/Engine/Source/Editor/UnrealEd/Public/Subsystems/AssetEditorSubsystem.h
@@ -102,6 +102,9 @@ struct FRegisteredModeInfo
 };
 using RegisteredModeInfoMap = TMap<FEditorModeID, FRegisteredModeInfo>;
 
+DECLARE_MULTICAST_DELEGATE_OneParam(FAssetEditorOpenDebugObject, UObject* /* DebugObject */);
+DECLARE_DYNAMIC_MULTICAST_DELEGATE_OneParam(FAssetEditorOpenDebugObjectDynamic, UObject*, DebugObject);
+
 /**
  * UAssetEditorSubsystem
  */
@@ -226,6 +229,18 @@ public:
 
 	UNREALED_API void RegisterUAssetEditor(UAssetEditor* NewAssetEditor);
 	UNREALED_API void UnregisterUAssetEditor(UAssetEditor* RemovedAssetEditor);
+
+	/**
+	 * Attempts to open a debug object in an already-open asset editor. Will not open the asset editor responsible for this
+	 * type of object, only cause the editor to debug the specified object.
+	 * Returns true if an asset editor accepted the debug object.
+	 */
+	UNREALED_API bool OpenDebugObjectInEditor(UObject* DebugObject);
+
+	/**
+	 * Returns the debug object delegate for a specific class.
+	 */
+	UNREALED_API FAssetEditorOpenDebugObject* GetDebugObjectDelegateForClass(UClass* InClass);
 	
 	/**
 	 * Creates a scriptable editor mode based on ID name, which will be owned by the given Owner, if that name exists in the map of editor modes found at system startup.
@@ -519,4 +534,11 @@ private:
 
 	/** External filters that are run to check if an asset is openable in read only mode */
 	TMap<FName, FReadOnlyAssetFilter> ReadOnlyAssetFilters;
+
+	/** Callbacks to open specific debug objects based on their type. */
+	TMap<FName, FAssetEditorOpenDebugObject> DebugObjectDelegates;
+
+	/** Callbacks to open specific debug objects exposed to blueprint. */
+	UPROPERTY(BlueprintAssignable, meta = (AllowPrivateAccess = "true", Category = "Debug"))
+	FAssetEditorOpenDebugObjectDynamic DebugObjectDelegate;
 };
diff --git a/Engine/Source/Runtime/Slate/Private/Framework/Application/SlateApplication.cpp b/Engine/Source/Runtime/Slate/Private/Framework/Application/SlateApplication.cpp
--- a/Engine/Source/Runtime/Slate/Private/Framework/Application/SlateApplication.cpp
+++ b/Engine/Source/Runtime/Slate/Private/Framework/Application/SlateApplication.cpp
@@ -7587,6 +7587,11 @@ void FSlateApplication::SetWidgetReflector(const TSharedRef<IWidgetReflector>& W
 		WidgetReflector->SetAssetAccessDelegate(AssetAccessDelegate);
 	}
 
+	if ( DebugObjectAccessDelegate.IsBound() )
+	{
+		WidgetReflector->SetDebugObjectAccessDelegate(DebugObjectAccessDelegate);
+	}
+
 	WidgetReflectorPtr = WidgetReflector;
 }
 
diff --git a/Engine/Source/Runtime/Slate/Public/Framework/Application/IWidgetReflector.h b/Engine/Source/Runtime/Slate/Public/Framework/Application/IWidgetReflector.h
--- a/Engine/Source/Runtime/Slate/Public/Framework/Application/IWidgetReflector.h
+++ b/Engine/Source/Runtime/Slate/Public/Framework/Application/IWidgetReflector.h
@@ -14,6 +14,9 @@ DECLARE_DELEGATE_RetVal_ThreeParams(bool, FAccessSourceCode, const FString& /*Fi
 /** A Delegate for an asset object to access */
 DECLARE_DELEGATE_RetVal_OneParam(bool, FAccessAsset, UObject* /*InAsset*/);
 
+/** A Delegate for an debug object to access */
+DECLARE_DELEGATE_RetVal_OneParam(bool, FAccessDebugObject, UObject* /*InObject*/);
+
 
 /**
  * Interface for widget reflectors.
@@ -82,6 +85,11 @@ public:
 	* @param InDelegate A delegate to access assets with.
 	*/
 	virtual void SetAssetAccessDelegate(FAccessAsset InDelegate) = 0;
+
+	/**
+	* @param InDelegate A delegate to access debug objects with.
+	*/
+	virtual void SetDebugObjectAccessDelegate(FAccessDebugObject InDelegate) = 0;
 	
 	/**
 	 * @param ThisWindow    Do we want to draw something for this window?
diff --git a/Engine/Source/Runtime/Slate/Public/Framework/Application/SlateApplication.h b/Engine/Source/Runtime/Slate/Public/Framework/Application/SlateApplication.h
--- a/Engine/Source/Runtime/Slate/Public/Framework/Application/SlateApplication.h
+++ b/Engine/Source/Runtime/Slate/Public/Framework/Application/SlateApplication.h
@@ -871,6 +871,12 @@ public:
 		AssetAccessDelegate = AccessDelegate;
 	}
 
+	/** @param AccessDelegate The delegate to pass along to the widget reflector */
+	void SetWidgetReflectorDebugAccessDelegate(FAccessDebugObject AccessDelegate)
+	{
+		DebugObjectAccessDelegate = AccessDelegate;
+	}
+
 	/** @param Scale  Sets the ratio SlateUnit / ScreenPixel */
 	void SetApplicationScale( float InScale ){ Scale = InScale; }
 
@@ -2064,6 +2070,9 @@ private:
 	/** Delegate for accessing assets, to pass to any widget inspectors. */
 	FAccessAsset AssetAccessDelegate;
 
+	/** Delegate for accessing debug objects, to pass to any widget inspectors. */
+	FAccessDebugObject DebugObjectAccessDelegate;
+
 	/** Allows us to track the number of non-slate modal windows active. */
 	int32 NumExternalModalWindowsActive;
 
diff --git a/Engine/Source/Runtime/Slate/Public/Framework/SlateDelegates.h b/Engine/Source/Runtime/Slate/Public/Framework/SlateDelegates.h
--- a/Engine/Source/Runtime/Slate/Public/Framework/SlateDelegates.h
+++ b/Engine/Source/Runtime/Slate/Public/Framework/SlateDelegates.h
@@ -255,6 +255,9 @@ DECLARE_DELEGATE_RetVal_ThreeParams(bool, FAccessSourceCode, const FString& /*Fi
 /** A Delegate for an asset object to access */
 DECLARE_DELEGATE_RetVal_OneParam(bool, FAccessAsset, UObject* /*InAsset*/);
 
+/** A Delegate for an debug object to access */
+DECLARE_DELEGATE_RetVal_OneParam(bool, FAccessDebugObject, UObject* /*InObject*/);
+
 /** A delegate for a callback during OnKeyDown from an SWidget */
 DECLARE_DELEGATE_RetVal_TwoParams(FReply, FOnKeyDown, const FGeometry&, const FKeyEvent&);
 
diff --git a/Engine/Source/Runtime/SlateCore/Private/Types/ReflectionMetadata.cpp b/Engine/Source/Runtime/SlateCore/Private/Types/ReflectionMetadata.cpp
--- a/Engine/Source/Runtime/SlateCore/Private/Types/ReflectionMetadata.cpp
+++ b/Engine/Source/Runtime/SlateCore/Private/Types/ReflectionMetadata.cpp
@@ -91,6 +91,28 @@ FString FReflectionMetaData::GetWidgetDebugInfo(const SWidget& InWidget)
 	return InWidget.ToString();
 }
 
+UObject* FReflectionMetaData::GetWidgetSourceObject(const SWidget* InWidget)
+{
+	if (!InWidget)
+	{
+		return nullptr;
+	}
+
+	return GetWidgetSourceObject(*InWidget);
+}
+
+UObject* FReflectionMetaData::GetWidgetSourceObject(const SWidget& InWidget)
+{
+	// UMG widgets have meta-data to help track them
+	TSharedPtr<FReflectionMetaData> MetaData = InWidget.GetMetaData<FReflectionMetaData>();
+	if (MetaData.IsValid())
+	{
+		return MetaData->SourceObject.Get();
+	}
+
+	return nullptr;
+}
+
 TSharedPtr<FReflectionMetaData> FReflectionMetaData::GetWidgetOrParentMetaData(const SWidget* InWidget)
 {
 	const SWidget* CurrentWidget = InWidget;
diff --git a/Engine/Source/Runtime/SlateCore/Public/Types/ReflectionMetadata.h b/Engine/Source/Runtime/SlateCore/Public/Types/ReflectionMetadata.h
--- a/Engine/Source/Runtime/SlateCore/Public/Types/ReflectionMetadata.h
+++ b/Engine/Source/Runtime/SlateCore/Public/Types/ReflectionMetadata.h
@@ -48,5 +48,8 @@ public:
 	static SLATECORE_API FString GetWidgetDebugInfo(const SWidget* InWidget);
 	static SLATECORE_API FString GetWidgetDebugInfo(const SWidget& InWidget);
 
+	static SLATECORE_API UObject* GetWidgetSourceObject(const SWidget* InWidget);
+	static SLATECORE_API UObject* GetWidgetSourceObject(const SWidget& InWidget);
+
 	static SLATECORE_API TSharedPtr<FReflectionMetaData> GetWidgetOrParentMetaData(const SWidget* InWidget);
 };
