UMVVMPanelWidgetViewExtension:
- Added OnItemsChanged delegate. Broadcast at the end of ReplaceAllSlots.

#Synced-CL 51408817

diff --git a/Engine/Plugins/Runtime/ModelViewViewModel/Source/ModelViewViewModel/Private/Extensions/MVVMViewPanelWidgetExtension.cpp b/Engine/Plugins/Runtime/ModelViewViewModel/Source/ModelViewViewModel/Private/Extensions/MVVMViewPanelWidgetExtension.cpp
--- a/Engine/Plugins/Runtime/ModelViewViewModel/Source/ModelViewViewModel/Private/Extensions/MVVMViewPanelWidgetExtension.cpp
+++ b/Engine/Plugins/Runtime/ModelViewViewModel/Source/ModelViewViewModel/Private/Extensions/MVVMViewPanelWidgetExtension.cpp
@@ -151,6 +151,8 @@ void UMVVMPanelWidgetViewExtension::ReplaceAllSlots(TArrayView<TTuple<UPanelSlot
 			UPanelSlot* NewSlot = PanelWidget->AddChild(NewSlots[SlotIndex].Value, NewSlots[SlotIndex].Key);
 		}
 	}
+
+	OnItemsChanged.Broadcast();
 }
 
 UUserWidget* UMVVMPanelWidgetViewExtension::GetUserWidget() const
diff --git a/Engine/Plugins/Runtime/ModelViewViewModel/Source/ModelViewViewModel/Public/Extensions/MVVMViewPanelWidgetExtension.h b/Engine/Plugins/Runtime/ModelViewViewModel/Source/ModelViewViewModel/Public/Extensions/MVVMViewPanelWidgetExtension.h
--- a/Engine/Plugins/Runtime/ModelViewViewModel/Source/ModelViewViewModel/Public/Extensions/MVVMViewPanelWidgetExtension.h
+++ b/Engine/Plugins/Runtime/ModelViewViewModel/Source/ModelViewViewModel/Public/Extensions/MVVMViewPanelWidgetExtension.h
@@ -24,6 +24,11 @@ public:
 	UFUNCTION(BlueprintCallable, Category = PanelWidget, meta = (AllowPrivateAccess = true, DisplayName = "Set Items", ViewmodelBlueprintWidgetExtension = "EntryViewModel"))
 	UE_API void BP_SetItems(const TArray<UObject*>& InItems);
 
+	DECLARE_DYNAMIC_MULTICAST_DELEGATE(FOnItemsChanged);
+
+	UPROPERTY(BlueprintAssignable)
+	FOnItemsChanged OnItemsChanged;
+
 private:
 	void SetViewModelOnEntryWidget(UUserWidget* EntryWidget, UObject* ViewModelObject, UUserWidget* OwningUserWidget);
 	void ReplaceAllSlots(TArrayView<TTuple<UPanelSlot*, UWidget*>> NewSlots);
