[UMG] Remove reference so utility can be used in code more easily.

#tests used locally, ran preflights, doesn't seem to cause any content issues
#Synced-CL 51470662

diff --git a/Engine/Plugins/Runtime/ModelViewViewModel/Source/ModelViewViewModel/Private/MVVMBlueprintLibrary.cpp b/Engine/Plugins/Runtime/ModelViewViewModel/Source/ModelViewViewModel/Private/MVVMBlueprintLibrary.cpp
--- a/Engine/Plugins/Runtime/ModelViewViewModel/Source/ModelViewViewModel/Private/MVVMBlueprintLibrary.cpp
+++ b/Engine/Plugins/Runtime/ModelViewViewModel/Source/ModelViewViewModel/Private/MVVMBlueprintLibrary.cpp
@@ -2,6 +2,7 @@
 
 #include "MVVMBlueprintLibrary.h"
 
+#include "Blueprint/UserWidget.h"
 #include "MVVMMessageLog.h"
 #include "View/MVVMView.h"
 
@@ -9,7 +10,7 @@
 
 #define LOCTEXT_NAMESPACE "MVVMBlueprintLibrary"
 
-void UMVVMBlueprintLibrary::SetViewModelByClass(UUserWidget*& Widget, TScriptInterface<INotifyFieldValueChanged> ViewModel)
+void UMVVMBlueprintLibrary::SetViewModelByClass(UUserWidget* Widget, TScriptInterface<INotifyFieldValueChanged> ViewModel)
 {
 	if (UMVVMView* View = Widget ? Widget->GetExtension<UMVVMView>() : nullptr)
 	{
diff --git a/Engine/Plugins/Runtime/ModelViewViewModel/Source/ModelViewViewModel/Public/MVVMBlueprintLibrary.h b/Engine/Plugins/Runtime/ModelViewViewModel/Source/ModelViewViewModel/Public/MVVMBlueprintLibrary.h
--- a/Engine/Plugins/Runtime/ModelViewViewModel/Source/ModelViewViewModel/Public/MVVMBlueprintLibrary.h
+++ b/Engine/Plugins/Runtime/ModelViewViewModel/Source/ModelViewViewModel/Public/MVVMBlueprintLibrary.h
@@ -12,5 +12,5 @@ class UMVVMBlueprintLibrary : public UBlueprintFunctionLibrary
 	GENERATED_BODY()
 public:
 	UFUNCTION(BlueprintCallable, Category="MVVM | Utilities")
-	static void SetViewModelByClass(UPARAM(ref) UUserWidget*& Widget, TScriptInterface<INotifyFieldValueChanged> ViewModel);
+	static MODELVIEWVIEWMODEL_API void SetViewModelByClass(UUserWidget* Widget, TScriptInterface<INotifyFieldValueChanged> ViewModel);
 };
\ No newline at end of file
