[UMG] Rename Ancesorts resolver type to Context which can pull from Parents of an external context extension on any user widgets in the hierarchy

#tests Used locally to adapt external systems to pass data along to MVVM
#Synced-CL 51547522

diff --git a/Engine/Plugins/Runtime/ModelViewViewModel/Source/ModelViewViewModel/Private/View/MVVMView.cpp b/Engine/Plugins/Runtime/ModelViewViewModel/Source/ModelViewViewModel/Private/View/MVVMView.cpp
--- a/Engine/Plugins/Runtime/ModelViewViewModel/Source/ModelViewViewModel/Private/View/MVVMView.cpp
+++ b/Engine/Plugins/Runtime/ModelViewViewModel/Source/ModelViewViewModel/Private/View/MVVMView.cpp
@@ -353,35 +353,32 @@ void UMVVMView::InitializeSourceParentBinding(FMVVMView_SourceKey SourceKey)
 
 	if (ClassSource.IsParentContextReference())
 	{
-		const TPair<const UMVVMView*, FMVVMView_SourceKey> ParentInfo = ClassSource.GetParentSupplyingSourceClass(GetUserWidget());
-		if (const UMVVMView* ParentView = ParentInfo.Key)
+		FMVVMViewClass_Source::FParentSourceInfo ParentInfo = ClassSource.GetParentSupplyingSourceClass(GetUserWidget());
+		if (ParentInfo.ParentView && ParentInfo.ParentUserWidget)
 		{
-			if (UUserWidget* ParentWidget = ParentView->GetOuterUUserWidget())
-			{
-				const FMVVMView_SourceKey ParentSourcKey = ParentInfo.Value;
-				const FMVVMViewClass_Source& ParentClassSource = ParentView->GetViewClass()->GetSource(ParentSourcKey.GetIndex());
+			const FMVVMView_SourceKey ParentSourceKey = ParentInfo.ParentSourceKey;
+			const FMVVMViewClass_Source& ParentClassSource = ParentInfo.ParentView->GetViewClass()->GetSource(ParentInfo.ParentSourceKey.GetIndex());
 
-				TWeakObjectPtr<const UMVVMView> WeakView = ParentView;
-				UE::FieldNotification::FFieldMulticastDelegate::FDelegate Delegate = UE::FieldNotification::FFieldMulticastDelegate::FDelegate::CreateWeakLambda(this, [this, SourceKey, ParentSourcKey, WeakView](UObject*, UE::FieldNotification::FFieldId)
+			TWeakObjectPtr<const UMVVMView> WeakView = ParentInfo.ParentView;
+			UE::FieldNotification::FFieldMulticastDelegate::FDelegate Delegate = UE::FieldNotification::FFieldMulticastDelegate::FDelegate::CreateWeakLambda(this, [this, SourceKey, ParentSourceKey, WeakView](UObject*, UE::FieldNotification::FFieldId)
+				{
+					if (const UMVVMView* ParentView = WeakView.Get())
 					{
-						if (const UMVVMView* ParentView = WeakView.Get())
-						{
-							SetSourceInternal(FMVVMViewClass_SourceKey(SourceKey.GetIndex()), ParentView->GetSource(ParentSourcKey).Source, true);
-						}
-					});
-
-				ParentWidget->GetFieldNotificationDescriptor();
-
-				// This is a big of a hack. The VM variable is not marked as field notify and we don't want it to be, but we also don't want to add bespoke
-				// notification logic. We always know that there will be at least one field notify in a user widget (UWidget has a few) and field notifies
-				// use the index to enable callbacks then the names to filter them. Adding this will enable all calbacks bound to index 0 (which may be multiple)
-				// and get called for any callbacks that share the name, regardless of index. If there happens to be a duplicate name in the delegates, then
-				// it likely just means that some callbacks are firing off more than they should, i.e. being notified of a change when there wasn't any,
-				// so there should be very low risk to this.
-				//
-				// If any of this changes make sure to also update UMVVMView::InitializeSourceInternal and UMVVMView::UninitializeSourceParentBinding
-				ParentWidget->AddFieldValueChangedDelegate(UE::FieldNotification::FFieldId(ParentClassSource.GetName(), 0), MoveTemp(Delegate));
-			}
+						SetSourceInternal(FMVVMViewClass_SourceKey(SourceKey.GetIndex()), ParentView->GetSource(ParentSourceKey).Source, true);
+					}
+				});
+
+			ParentInfo.ParentUserWidget->GetFieldNotificationDescriptor();
+
+			// This is a big of a hack. The VM variable is not marked as field notify and we don't want it to be, but we also don't want to add bespoke
+			// notification logic. We always know that there will be at least one field notify in a user widget (UWidget has a few) and field notifies
+			// use the index to enable callbacks then the names to filter them. Adding this will enable all calbacks bound to index 0 (which may be multiple)
+			// and get called for any callbacks that share the name, regardless of index. If there happens to be a duplicate name in the delegates, then
+			// it likely just means that some callbacks are firing off more than they should, i.e. being notified of a change when there wasn't any,
+			// so there should be very low risk to this.
+			//
+			// If any of this changes make sure to also update UMVVMView::InitializeSourceInternal and UMVVMView::UninitializeSourceParentBinding
+			ParentInfo.ParentUserWidget->AddFieldValueChangedDelegate(UE::FieldNotification::FFieldId(ParentClassSource.GetName(), 0), MoveTemp(Delegate));
 		}
 	}
 }
@@ -394,13 +391,10 @@ void UMVVMView::UninitializeSourceParentBinding(FMVVMView_SourceKey SourceKey)
 	if (ClassSource.IsParentContextReference())
 	{
 		// See comment in UMVVMView::InitializeSourceParentBinding for why we do this
-		const TPair<const UMVVMView*, FMVVMView_SourceKey> ParentInfo = ClassSource.GetParentSupplyingSourceClass(GetUserWidget());
-		if (const UMVVMView* ParentView = ParentInfo.Key)
+		FMVVMViewClass_Source::FParentSourceInfo ParentInfo = ClassSource.GetParentSupplyingSourceClass(GetUserWidget());
+		if (ParentInfo.ParentView && ParentInfo.ParentUserWidget)
 		{
-			if (UUserWidget* ParentWidget = ParentView->GetOuterUUserWidget())
-			{
-				ParentWidget->RemoveAllFieldValueChangedDelegates(this);
-			}
+			ParentInfo.ParentUserWidget->RemoveAllFieldValueChangedDelegates(this);
 		}
 	}
 }
diff --git a/Engine/Plugins/Runtime/ModelViewViewModel/Source/ModelViewViewModel/Private/View/MVVMViewClass.cpp b/Engine/Plugins/Runtime/ModelViewViewModel/Source/ModelViewViewModel/Private/View/MVVMViewClass.cpp
--- a/Engine/Plugins/Runtime/ModelViewViewModel/Source/ModelViewViewModel/Private/View/MVVMViewClass.cpp
+++ b/Engine/Plugins/Runtime/ModelViewViewModel/Source/ModelViewViewModel/Private/View/MVVMViewClass.cpp
@@ -12,6 +12,7 @@
 #include "Engine/Engine.h"
 #include "Engine/GameInstance.h"
 #include "Engine/World.h"
+#include "Extensions/ContextDataWidgetExtension.h"
 #include "MVVMGameSubsystem.h"
 #include "MVVMMessageLog.h"
 #include "MVVMSubsystem.h"
@@ -117,12 +118,18 @@ UObject* FMVVMViewClass_Source::GetOrCreateInstance(const UMVVMViewClass* InView
 
 		if (ExpectedSourceType.Get() != nullptr)
 		{
-			const TPair<const UMVVMView*, FMVVMView_SourceKey> ParentInfo = GetParentSupplyingSourceClass(InUserWidget);
-			if (const UMVVMView* ParentView = ParentInfo.Key)
+			const FParentSourceInfo ParentInfo = GetParentSupplyingSourceClass(InUserWidget);
+			if (ParentInfo.ParentView)
 			{
-				Result = ParentView->GetSource(ParentInfo.Value).Source;
+				Result = ParentInfo.Source;
 
-				Log.Info(FText::Format(LOCTEXT("CreateInstanceContextParentFound", "The source '{0}' was found on parent widget '{1}'"), FText::FromName(PropertyName), FText::FromString(GetNameSafe(ParentView->GetOuterUUserWidget()))));
+				Log.Info(FText::Format(LOCTEXT("CreateInstanceContextParentFound", "The source '{0}' was found on view from parent widget '{1}'"), FText::FromName(PropertyName), FText::FromString(GetNameSafe(ParentInfo.ParentUserWidget))));
+			}
+			else if (ParentInfo.Source)
+			{
+				Result = ParentInfo.Source;
+
+				Log.Info(FText::Format(LOCTEXT("CreateInstanceContextDataFound", "The source '{0}' was found on context data from widget '{1}'"), FText::FromName(PropertyName), FText::FromString(GetNameSafe(ParentInfo.ParentUserWidget))));
 			}
 			else
 			{
@@ -204,27 +211,46 @@ TScriptInterface<INotifyFieldValueChanged> FMVVMViewClass_Source::GetGlobalColle
 	return FoundViewModelInstance ? TScriptInterface<INotifyFieldValueChanged>(FoundViewModelInstance) : TScriptInterface<INotifyFieldValueChanged>();
 }
 
-TPair<const UMVVMView*, FMVVMView_SourceKey> FMVVMViewClass_Source::GetParentSupplyingSourceClass(UUserWidget* UserWidget) const
+FMVVMViewClass_Source::FParentSourceInfo FMVVMViewClass_Source::GetParentSupplyingSourceClass(UUserWidget* UserWidget) const
 {
 	// Find an existing instance from a parent source and/or bind to when that paren't source changes
-	for (const UUserWidget* Parent = UserWidget->GetTypedOuter<UUserWidget>(); Parent; Parent = Parent->GetTypedOuter<UUserWidget>())
+	for (UUserWidget* ContextWidget = UserWidget; ContextWidget; ContextWidget = ContextWidget->GetTypedOuter<UUserWidget>())
 	{
-		if (const UMVVMView* ParentView = UMVVMSubsystem::GetViewFromUserWidget(Parent); ParentView && ParentView->GetViewClass())
+		if (const UContextDataWidgetExtension* ContextDataExtension = ContextWidget->GetExtension<UContextDataWidgetExtension>())
 		{
-			const int32 FoundIndex = ParentView->GetViewClass()->GetSources().IndexOfByPredicate([this](const FMVVMViewClass_Source& Other)
-				{
-					return Other.IsViewModel() && Other.GetSourceClass()->IsChildOf(ExpectedSourceType) &&
-						(GetName().IsNone() ? Other.GetName().IsNone() : GetName() == Other.GetName());
-				});
+			if (UObject* ContextData = ContextDataExtension->FindData(ExpectedSourceType))
+			{
+				FParentSourceInfo Result;
+				Result.ParentUserWidget = ContextWidget;
+				Result.Source = ContextData;
+				return Result;
+			}
+		}
 
-			if (ParentView->GetSources().IsValidIndex(FoundIndex))
+		if (ContextWidget != UserWidget)
+		{
+			if (const UMVVMView* ParentView = UMVVMSubsystem::GetViewFromUserWidget(ContextWidget); ParentView && ParentView->GetViewClass())
 			{
-				return { ParentView, FoundIndex };
+				const int32 FoundIndex = ParentView->GetViewClass()->GetSources().IndexOfByPredicate([this](const FMVVMViewClass_Source& Other)
+					{
+						return Other.IsViewModel() && Other.GetSourceClass()->IsChildOf(ExpectedSourceType) &&
+							(GetName().IsNone() ? Other.GetName().IsNone() : GetName() == Other.GetName());
+					});
+
+				if (ParentView->GetSources().IsValidIndex(FoundIndex))
+				{
+					FParentSourceInfo Result;
+					Result.ParentView = ParentView;
+					Result.ParentSourceKey = FoundIndex;
+					Result.ParentUserWidget = ContextWidget;
+					Result.Source = ParentView->GetSource(Result.ParentSourceKey).Source;
+					return Result;
+				}
 			}
 		}
 	}
 
-	return { nullptr, INDEX_NONE };
+	return FParentSourceInfo();
 }
 
 #if UE_WITH_MVVM_DEBUGGING
diff --git a/Engine/Plugins/Runtime/ModelViewViewModel/Source/ModelViewViewModel/Public/View/MVVMViewClass.h b/Engine/Plugins/Runtime/ModelViewViewModel/Source/ModelViewViewModel/Public/View/MVVMViewClass.h
--- a/Engine/Plugins/Runtime/ModelViewViewModel/Source/ModelViewViewModel/Public/View/MVVMViewClass.h
+++ b/Engine/Plugins/Runtime/ModelViewViewModel/Source/ModelViewViewModel/Public/View/MVVMViewClass.h
@@ -279,8 +279,22 @@ public:
 	/** Returns the viewmodel inside the global viewmodel collection. */
 	MODELVIEWVIEWMODEL_API TScriptInterface<INotifyFieldValueChanged> GetGlobalCollectionViewModel(UUserWidget* UserWidget) const;
 
+	// Helper struct to provide info about where context data is coming from
+	struct FParentSourceInfo
+	{
+		// Parent view info will only be set if a parent widget was found that has the same class/name in it's view
+		const UMVVMView* ParentView = nullptr;
+		FMVVMView_SourceKey ParentSourceKey;
+
+		// Parent User Widget will always be valid if Source is valid or if the ParentView is valid
+		UUserWidget* ParentUserWidget = nullptr;
+
+		// May be valid or not if the parent view is found, otherwise will only be valid if and external source was provided
+		UObject* Source = nullptr;
+	};
+
 	/** Get the parent vm view that has a matching vm class source */
-	TPair<const UMVVMView*, FMVVMView_SourceKey> GetParentSupplyingSourceClass(UUserWidget* UserWidget) const;
+	FParentSourceInfo GetParentSupplyingSourceClass(UUserWidget* UserWidget) const;
 
 	/** The expected class of the source. */
 	UClass* GetSourceClass() const
diff --git a/Engine/Plugins/Runtime/ModelViewViewModel/Source/ModelViewViewModelBlueprint/Private/MVVMDeveloperProjectSettings.cpp b/Engine/Plugins/Runtime/ModelViewViewModel/Source/ModelViewViewModelBlueprint/Private/MVVMDeveloperProjectSettings.cpp
--- a/Engine/Plugins/Runtime/ModelViewViewModel/Source/ModelViewViewModelBlueprint/Private/MVVMDeveloperProjectSettings.cpp
+++ b/Engine/Plugins/Runtime/ModelViewViewModel/Source/ModelViewViewModelBlueprint/Private/MVVMDeveloperProjectSettings.cpp
@@ -37,7 +37,7 @@ UMVVMDeveloperProjectSettings::UMVVMDeveloperProjectSettings()
 	AllowedContextCreationType.Add(EMVVMBlueprintViewModelContextCreationType::PropertyPath);
 	AllowedContextCreationType.Add(EMVVMBlueprintViewModelContextCreationType::PropertyPath);
 	AllowedContextCreationType.Add(EMVVMBlueprintViewModelContextCreationType::Resolver);
-	AllowedContextCreationType.Add(EMVVMBlueprintViewModelContextCreationType::Ancestors);
+	AllowedContextCreationType.Add(EMVVMBlueprintViewModelContextCreationType::Context);
 
 	FTopLevelAssetPath BlueprintFunctionLibrary = FTopLevelAssetPath("/Script/Engine", "BlueprintFunctionLibrary");
 	FTopLevelAssetPath FormatText = FTopLevelAssetPath("/Script/BlueprintGraph", "K2Node_FormatText");
diff --git a/Engine/Plugins/Runtime/ModelViewViewModel/Source/ModelViewViewModelBlueprint/Private/MVVMViewBlueprintCompiler.cpp b/Engine/Plugins/Runtime/ModelViewViewModel/Source/ModelViewViewModelBlueprint/Private/MVVMViewBlueprintCompiler.cpp
--- a/Engine/Plugins/Runtime/ModelViewViewModel/Source/ModelViewViewModelBlueprint/Private/MVVMViewBlueprintCompiler.cpp
+++ b/Engine/Plugins/Runtime/ModelViewViewModel/Source/ModelViewViewModelBlueprint/Private/MVVMViewBlueprintCompiler.cpp
@@ -963,7 +963,7 @@ void FMVVMViewBlueprintCompiler::CreateRequiredProperties(const FWidgetBlueprint
 			SourceContext->Type = FCompilerBindingSource::EType::ViewModel;
 			SourceContext->bIsOptional = ViewModelContext.bOptional;
 
-			if (ViewModelContext.CreationType == EMVVMBlueprintViewModelContextCreationType::Manual || ViewModelContext.CreationType == EMVVMBlueprintViewModelContextCreationType::Ancestors)
+			if (ViewModelContext.CreationType == EMVVMBlueprintViewModelContextCreationType::Manual || ViewModelContext.CreationType == EMVVMBlueprintViewModelContextCreationType::Context)
 			{
 				SourceContext->bIsOptional = true;
 			}
@@ -2570,7 +2570,7 @@ void FMVVMViewBlueprintCompiler::PreCompileViewModelCreatorContexts(UWidgetBluep
 			}
 		}
 
-		if (ViewModelContext.CreationType == EMVVMBlueprintViewModelContextCreationType::Manual || ViewModelContext.CreationType == EMVVMBlueprintViewModelContextCreationType::Ancestors)
+		if (ViewModelContext.CreationType == EMVVMBlueprintViewModelContextCreationType::Manual || ViewModelContext.CreationType == EMVVMBlueprintViewModelContextCreationType::Context)
 		{
 		}
 		else if (ViewModelContext.CreationType == EMVVMBlueprintViewModelContextCreationType::CreateInstance)
@@ -2857,7 +2857,7 @@ void FMVVMViewBlueprintCompiler::CompileSources(const FCompiledBindingLibraryCom
 			CompiledSourceCreator.Resolver = Resolver;
 			ensure(bIsOptional == ViewModelContext.bOptional);
 		}
-		else if (ViewModelContext.CreationType == EMVVMBlueprintViewModelContextCreationType::Ancestors)
+		else if (ViewModelContext.CreationType == EMVVMBlueprintViewModelContextCreationType::Context)
 		{
 			ensure(bCanBeSet == false);
 			ensure(bIsOptional == true);
diff --git a/Engine/Plugins/Runtime/ModelViewViewModel/Source/ModelViewViewModelBlueprint/Public/MVVMBlueprintViewModelContext.h b/Engine/Plugins/Runtime/ModelViewViewModel/Source/ModelViewViewModelBlueprint/Public/MVVMBlueprintViewModelContext.h
--- a/Engine/Plugins/Runtime/ModelViewViewModel/Source/ModelViewViewModelBlueprint/Public/MVVMBlueprintViewModelContext.h
+++ b/Engine/Plugins/Runtime/ModelViewViewModel/Source/ModelViewViewModelBlueprint/Public/MVVMBlueprintViewModelContext.h
@@ -28,8 +28,8 @@ enum class EMVVMBlueprintViewModelContextCreationType : uint8
 	PropertyPath,
 	// The viewmodel will be fetched by evaluating the resolver object.
 	Resolver,
-	// The viewmodel will be fetched from any parent view that contains the same class and name
-	Ancestors,
+	// The viewmodel will be fetched from any local context provider, can also be set manually
+	Context,
 };
 
 namespace UE::MVVM
diff --git a/Engine/Plugins/Runtime/ModelViewViewModel/Source/ModelViewViewModelEditor/Private/Customizations/MVVMBlueprintViewModelContextCustomization.cpp b/Engine/Plugins/Runtime/ModelViewViewModel/Source/ModelViewViewModelEditor/Private/Customizations/MVVMBlueprintViewModelContextCustomization.cpp
--- a/Engine/Plugins/Runtime/ModelViewViewModel/Source/ModelViewViewModelEditor/Private/Customizations/MVVMBlueprintViewModelContextCustomization.cpp
+++ b/Engine/Plugins/Runtime/ModelViewViewModel/Source/ModelViewViewModelEditor/Private/Customizations/MVVMBlueprintViewModelContextCustomization.cpp
@@ -390,7 +390,7 @@ void FBlueprintViewModelContextDetailCustomization::CustomizeChildren(TSharedRef
 				ChildBuilder.AddProperty(CreateSetterFunctionHandle.ToSharedRef())
 					.IsEnabled(MakeAttributeLambda([bCanEdit, ContextPtr]()
 						{
-							bool bResult = ContextPtr->CreationType != EMVVMBlueprintViewModelContextCreationType::Manual && ContextPtr->CreationType != EMVVMBlueprintViewModelContextCreationType::Ancestors;
+							bool bResult = ContextPtr->CreationType != EMVVMBlueprintViewModelContextCreationType::Manual && ContextPtr->CreationType != EMVVMBlueprintViewModelContextCreationType::Context;
 							return bResult && bCanEdit;
 						}));
 			}
@@ -526,7 +526,7 @@ void FBlueprintViewModelContextDetailCustomization::HandleCreationTypeChanged()
 		{
 			const EMVVMBlueprintViewModelContextCreationType CreationType = (EMVVMBlueprintViewModelContextCreationType)NewValue;
 			const bool bIsManual = CreationType == EMVVMBlueprintViewModelContextCreationType::Manual;
-			const bool bIsAlwaysOptional = CreationType == EMVVMBlueprintViewModelContextCreationType::Manual || CreationType == EMVVMBlueprintViewModelContextCreationType::Ancestors;
+			const bool bIsAlwaysOptional = CreationType == EMVVMBlueprintViewModelContextCreationType::Manual || CreationType == EMVVMBlueprintViewModelContextCreationType::Context;
 			if (ContextPtr->bOptional != bIsAlwaysOptional)
 			{
 				OptionalHandle->SetValue(bIsAlwaysOptional);
diff --git a/Engine/Source/Runtime/UMG/Private/UserWidget.cpp b/Engine/Source/Runtime/UMG/Private/UserWidget.cpp
--- a/Engine/Source/Runtime/UMG/Private/UserWidget.cpp
+++ b/Engine/Source/Runtime/UMG/Private/UserWidget.cpp
@@ -2945,6 +2945,15 @@ UUserWidgetExtension* UUserWidget::AddExtension(TSubclassOf<UUserWidgetExtension
 	return Extension;
 }
 
+UUserWidgetExtension* UUserWidget::GetOrAddExtension(TSubclassOf<UUserWidgetExtension> InExtensionType)
+{
+	if (UUserWidgetExtension* ExistingExtension = GetExtension(InExtensionType))
+	{
+		return ExistingExtension;
+	}
+	return AddExtension(InExtensionType);
+}
+
 void UUserWidget::RemoveExtension(UUserWidgetExtension* InExtension)
 {
 	if (InExtension)
diff --git a/Engine/Source/Runtime/UMG/Public/Blueprint/UserWidget.h b/Engine/Source/Runtime/UMG/Public/Blueprint/UserWidget.h
--- a/Engine/Source/Runtime/UMG/Public/Blueprint/UserWidget.h
+++ b/Engine/Source/Runtime/UMG/Public/Blueprint/UserWidget.h
@@ -1356,6 +1356,17 @@ public:
 	UFUNCTION(BlueprintCallable, Category = "User Interface|Extension", Meta = (DeterminesOutputType = "InExtensionType"))
 	UMG_API UUserWidgetExtension* AddExtension(TSubclassOf<UUserWidgetExtension> InExtensionType);
 
+	/** Find an existing extension of the request type, or create it if not found */
+	template<typename ExtensionType>
+	ExtensionType& GetOrAddExtension()
+	{
+		return *CastChecked<ExtensionType>(GetOrAddExtension(ExtensionType::StaticClass()));
+	}
+
+	/** Find an existing extension of the request type, or create it if not found */
+	UFUNCTION(BlueprintCallable, Category = "User Interface|Extension", Meta = (DeterminesOutputType = "InExtensionType"))
+	UMG_API UUserWidgetExtension* GetOrAddExtension(TSubclassOf<UUserWidgetExtension> InExtensionType);
+
 	/** Remove the extension. */
 	UFUNCTION(BlueprintCallable, Category = "User Interface|Extension")
 	UMG_API void RemoveExtension(UUserWidgetExtension* InExtension);
diff --git a/Engine/Source/Runtime/UMG/Public/Extensions/ContextDataWidgetExtension.h b/Engine/Source/Runtime/UMG/Public/Extensions/ContextDataWidgetExtension.h
new file mode 100644
--- /dev/null
+++ b/Engine/Source/Runtime/UMG/Public/Extensions/ContextDataWidgetExtension.h
@@ -0,0 +1,48 @@
+// Copyright Epic Games, Inc. All Rights Reserved.
+
+#pragma once
+
+#include "UserWidgetExtension.h"
+
+#include "ContextDataWidgetExtension.generated.h"
+
+/**
+ * This extention is just a data container that can be used to attach arbitrary data to a widget
+ */
+UCLASS()
+class UContextDataWidgetExtension : public UUserWidgetExtension
+{
+	GENERATED_BODY()
+	
+public:
+	template <typename DataType>
+	DataType* FindData() const
+	{
+		return CastChecked<DataType>(FindData(DataType::StaticClass()), ECastCheckedType::NullAllowed);
+	}
+
+	UObject* FindData(const UClass* Class) const
+	{
+		for (UObject* Data : ContextData)
+		{
+			if (Data && Data->IsA(Class)) // Can bedome null from GC
+			{
+				return Data;
+			}
+		}
+		return nullptr;
+	}
+
+	void AddData(UObject* Data)
+	{
+		// Todo: Support multiple instances? What about exsiting data that is of this type?
+		if (ensure(Data) && ensure(!FindData(Data->GetClass())))
+		{
+			ContextData.Add(Data);
+		}
+	}
+
+private:
+	UPROPERTY(Transient)
+	TArray<TObjectPtr<UObject>> ContextData;
+};
