MVVM: Guard against null conditions and events

#rnx

#Synced-CL 51839901

diff --git a/Engine/Plugins/Runtime/ModelViewViewModel/Source/ModelViewViewModelEditor/Private/Widgets/SMVVMViewBindingListView.cpp b/Engine/Plugins/Runtime/ModelViewViewModel/Source/ModelViewViewModelEditor/Private/Widgets/SMVVMViewBindingListView.cpp
--- a/Engine/Plugins/Runtime/ModelViewViewModel/Source/ModelViewViewModelEditor/Private/Widgets/SMVVMViewBindingListView.cpp
+++ b/Engine/Plugins/Runtime/ModelViewViewModel/Source/ModelViewViewModelEditor/Private/Widgets/SMVVMViewBindingListView.cpp
@@ -594,6 +594,11 @@ void SBindingsList::Refresh()
 
 		for (UMVVMBlueprintViewEvent* Event : BlueprintView->GetEvents())
 		{
+			if (Event == nullptr)
+			{
+				continue;
+			}
+
 			// Make sure the graph is up to date
 			Event->GetOrCreateWrapperGraph();
 
@@ -690,6 +695,11 @@ void SBindingsList::Refresh()
 		FGuid GroupViewModelId;
 		for (UMVVMBlueprintViewCondition* Condition : BlueprintView->GetConditions())
 		{
+			if (Condition == nullptr)
+			{
+				continue;
+			}
+
 			// Make sure the graph is up to date
 			Condition->GetOrCreateWrapperGraph();
 
