Itemization UI: force re-binding of all the widgets whenever a new property interface replicates to the client, or when the list of property interfaces used by the inventory change. Otherwise we get cases where the bindings fail and never re-evaluate.

#rnx

#Synced-CL 51999454

diff --git a/Engine/Plugins/Runtime/ModelViewViewModel/Source/ModelViewViewModel/Public/View/MVVMView.h b/Engine/Plugins/Runtime/ModelViewViewModel/Source/ModelViewViewModel/Public/View/MVVMView.h
--- a/Engine/Plugins/Runtime/ModelViewViewModel/Source/ModelViewViewModel/Public/View/MVVMView.h
+++ b/Engine/Plugins/Runtime/ModelViewViewModel/Source/ModelViewViewModel/Public/View/MVVMView.h
@@ -78,6 +78,12 @@ public:
 	UE_API virtual void Destruct() override;
 	//~ End UUserWidgetExtension implementation
 
+	/** Is the view in a valid constructed state? */
+	bool IsConstructed() const
+	{
+		return bConstructed;
+	}
+
 	/**
 	 * Initialize the sources if they are not already initialized.
 	 * Initializing the sources will also initialize the bindings if the option bInitializeBindingsOnConstruct is enabled.
