UMG Editor: Popout Viewmodel editor.

#Synced-CL 52052026

diff --git a/Engine/Plugins/Runtime/ModelViewViewModel/Source/ModelViewViewModelEditor/Private/MVVMEditorCommands.cpp b/Engine/Plugins/Runtime/ModelViewViewModel/Source/ModelViewViewModelEditor/Private/MVVMEditorCommands.cpp
--- a/Engine/Plugins/Runtime/ModelViewViewModel/Source/ModelViewViewModelEditor/Private/MVVMEditorCommands.cpp
+++ b/Engine/Plugins/Runtime/ModelViewViewModel/Source/ModelViewViewModelEditor/Private/MVVMEditorCommands.cpp
@@ -21,6 +21,8 @@ void FMVVMEditorCommands::RegisterCommands()
 #else
 	UI_COMMAND(ToggleMVVMDrawer, "Toggle MVVM Drawer", "Toggle MVVM Drawer", EUserInterfaceActionType::ToggleButton, FInputChord(EModifierKey::Control | EModifierKey::Shift, EKeys::B));
 #endif
+
+	UI_COMMAND(PopoutEditor, "Popout Editor", "Popout a standalone Property Editor for this object.", EUserInterfaceActionType::Button, FInputChord());
 }
 
 const FMVVMEditorCommands& FMVVMEditorCommands::Get()
diff --git a/Engine/Plugins/Runtime/ModelViewViewModel/Source/ModelViewViewModelEditor/Private/ModelViewViewModelEditorModule.cpp b/Engine/Plugins/Runtime/ModelViewViewModel/Source/ModelViewViewModelEditor/Private/ModelViewViewModelEditorModule.cpp
--- a/Engine/Plugins/Runtime/ModelViewViewModel/Source/ModelViewViewModelEditor/Private/ModelViewViewModelEditorModule.cpp
+++ b/Engine/Plugins/Runtime/ModelViewViewModel/Source/ModelViewViewModelEditor/Private/ModelViewViewModelEditorModule.cpp
@@ -32,6 +32,7 @@
 #include "UObject/AssetRegistryTagsContext.h"
 #include "WidgetBlueprintEditor.h"
 #include "WidgetDrawerConfig.h"
+#include "Widgets/SMVVMPopoutEditor.h"
 #include "Widgets/SMVVMViewBindingPanel.h"
 #include "Widgets/SMVVMViewModelPanel.h"
 
@@ -133,6 +134,11 @@ void FModelViewViewModelEditorModule::ShutdownModule()
 	FMVVMEditorCommands::Unregister();
 }
 
+void FModelViewViewModelEditorModule::OpenPopoutEditor(TNotNull<UObject*> InObject, bool bInReadyOnly)
+{
+	UE::MVVM::SMVVMPopoutEditor::OpenPopoutEditor(InObject, bInReadyOnly);
+}
+
 
 void FModelViewViewModelEditorModule::HandleRegisterBlueprintEditorTab(const FWidgetBlueprintApplicationMode& ApplicationMode, FWorkflowAllowedTabSet& TabFactories)
 {
diff --git a/Engine/Plugins/Runtime/ModelViewViewModel/Source/ModelViewViewModelEditor/Private/ModelViewViewModelEditorModule.h b/Engine/Plugins/Runtime/ModelViewViewModel/Source/ModelViewViewModelEditor/Private/ModelViewViewModelEditorModule.h
--- a/Engine/Plugins/Runtime/ModelViewViewModel/Source/ModelViewViewModelEditor/Private/ModelViewViewModelEditorModule.h
+++ b/Engine/Plugins/Runtime/ModelViewViewModel/Source/ModelViewViewModelEditor/Private/ModelViewViewModelEditorModule.h
@@ -2,7 +2,7 @@
 
 #pragma once
 
-#include "Modules/ModuleInterface.h"
+#include "IModelViewViewModelEditorModule.h"
 #include "UObject/Object.h"
 
 namespace UE::MVVM
@@ -27,7 +27,7 @@ class UWidgetBlueprintGeneratedClass;
 /**
  *
  */
-class FModelViewViewModelEditorModule : public IModuleInterface
+class FModelViewViewModelEditorModule : public IModelViewViewModelEditorModule
 {
 public:
 	FModelViewViewModelEditorModule() = default;
@@ -37,6 +37,10 @@ public:
 	 virtual void ShutdownModule() override;
 	//~ End IModuleInterface interface
 
+	//~ Begin IModelViewViewModelEditorModule
+	virtual void OpenPopoutEditor(TNotNull<UObject*> InObject, bool bInReadyOnly) override;
+	//~ End IModelViewViewModelEditorModule
+
 private:
 	void HandleRegisterBlueprintEditorTab(const FWidgetBlueprintApplicationMode& ApplicationMode, FWorkflowAllowedTabSet& TabFactories);
 	void HandleRenameFieldReferences(UBlueprint* Blueprint, UClass* FieldOwnerClass, const FName& OldVarName, const FName& NewVarName);
diff --git a/Engine/Plugins/Runtime/ModelViewViewModel/Source/ModelViewViewModelEditor/Private/Widgets/SMVVMPopoutEditor.cpp b/Engine/Plugins/Runtime/ModelViewViewModel/Source/ModelViewViewModelEditor/Private/Widgets/SMVVMPopoutEditor.cpp
new file mode 100644
--- /dev/null
+++ b/Engine/Plugins/Runtime/ModelViewViewModel/Source/ModelViewViewModelEditor/Private/Widgets/SMVVMPopoutEditor.cpp
@@ -0,0 +1,201 @@
+// Copyright Epic Games, Inc. All Rights Reserved.
+
+#include "Widgets/SMVVMPopoutEditor.h"
+
+#include "ClassIconFinder.h"
+#include "Framework/Docking/TabManager.h"
+#include "Framework/PropertyViewer/IFieldExpander.h"
+#include "Internationalization/Text.h"
+#include "MVVMSubsystem.h"
+#include "Styling/MVVMEditorStyle.h"
+#include "Types/MVVMAvailableBinding.h"
+#include "Widgets/Docking/SDockTab.h"
+#include "Widgets/Images/SImage.h"
+#include "Widgets/Layout/SBox.h"
+#include "Widgets/PropertyViewer/SFieldIcon.h"
+#include "Widgets/SMVVMViewModelBindingListWidget.h"
+#include "Widgets/SWindow.h"
+#include "Widgets/Text/STextBlock.h"
+
+#define LOCTEXT_NAMESPACE "SMVVMPopoutEditor"
+
+namespace UE::MVVM
+{
+
+void SMVVMPopoutEditor::OpenPopoutEditor(TNotNull<UObject*> InObjectToInspect, bool bInReadOnly)
+{
+	const FName TabId = *(FString(TEXT("SMVVMPopoutEditor:")) + InObjectToInspect->GetPathName());
+
+	const TSharedRef<FGlobalTabmanager> GlobalTabManager = FGlobalTabmanager::Get();
+
+	if (GlobalTabManager->FindExistingLiveTab(TabId))
+	{
+		GlobalTabManager->TryInvokeTab(TabId);
+		return;
+	}
+
+	GlobalTabManager->UnregisterNomadTabSpawner(TabId);
+
+	GlobalTabManager->RegisterNomadTabSpawner(
+		TabId,
+		FOnSpawnTab::CreateLambda(
+			[ObjectWeak = TWeakObjectPtr<UObject>(InObjectToInspect), bInReadOnly](const FSpawnTabArgs& InSpawnArgs)
+			{
+				TSharedRef<SDockTab> Tab = SNew(SDockTab)
+					.TabRole(ETabRole::NomadTab)
+					.OnTabClosed_Lambda(
+						[TabId = InSpawnArgs.GetTabId()](TSharedRef<SDockTab> InTab)
+						{
+							FGlobalTabmanager::Get()->UnregisterNomadTabSpawner(TabId.TabType);
+						}
+					);
+
+				if (UObject* Object = ObjectWeak.Get())
+				{
+					const FText Label = FText::Format(
+						LOCTEXT("LabelFormat", "{0} [{1}]"),
+						FText::FromName(Object->GetFName()),
+						FText::FromString(Object->GetPathName())
+					);
+
+					Tab->SetLabel(Label);
+					Tab->SetTabIcon(FClassIconFinder::FindThumbnailForClass(Object->GetClass()));
+
+					Tab->SetContent(
+						SNew(SMVVMPopoutEditor, Object)
+						.bReadOnly(bInReadOnly)
+					);
+				}
+
+				return Tab;
+			}
+		)
+	);
+
+	if (!GlobalTabManager->TryInvokeTab(TabId).IsValid())
+	{
+		GlobalTabManager->UnregisterNomadTabSpawner(TabId);
+	}
+}
+
+void SMVVMPopoutEditor::Construct(const FArguments& InArgs, UObject* InObjectToInspect)
+{
+	ObjectWeak = InObjectToInspect;
+	FieldIterator = MakeUnique<UE::PropertyViewer::FFieldIterator_BlueprintVisible>();
+	FieldExpander = MakeUnique<FFieldExpander_Bindable>();
+
+	PropertyViewer = SNew(UE::PropertyViewer::SPropertyViewer)
+		.PropertyVisibility(
+			InArgs._bReadOnly
+				? UE::PropertyViewer::SPropertyViewer::EPropertyVisibility::Visible
+				: UE::PropertyViewer::SPropertyViewer::EPropertyVisibility::Editable
+		)
+		.bShowSearchBox(true)
+		.bShowFieldIcon(true)
+		.bSanitizeName(true)
+		.FieldIterator(FieldIterator.Get())
+		.FieldExpander(FieldExpander.Get())
+		.OnGetPreSlot(this, &SMVVMPopoutEditor::HandleGetPreSlot)
+		.OnGenerateContainer(this, &SMVVMPopoutEditor::HandleGenerateContainer);
+
+	PropertyViewer->AddInstance(InObjectToInspect);
+
+	ChildSlot
+	[
+		PropertyViewer.ToSharedRef()
+	];
+}
+
+TSharedPtr<SWidget> SMVVMPopoutEditor::HandleGetPreSlot(UE::PropertyViewer::SPropertyViewer::FHandle Handle, TArrayView<const FFieldVariant> FieldPath)
+{
+	if (FieldPath.Num() > 0)
+	{
+		return ConstructFieldPreSlot(ObjectWeak.Get(), Handle, FieldPath.Last());
+	}
+
+	return TSharedPtr<SWidget>();
+}
+
+TSharedRef<SWidget> SMVVMPopoutEditor::HandleGenerateContainer(UE::PropertyViewer::SPropertyViewer::FHandle ContainerHandle, TOptional<FText> DisplayName)
+{
+	if (UObject* Object = ObjectWeak.Get())
+	{
+		return SNew(SHorizontalBox)
+			+ SHorizontalBox::Slot()
+			.AutoWidth()
+			.HAlign(HAlign_Right)
+			.VAlign(VAlign_Center)
+			[
+				SNew(UE::PropertyViewer::SFieldIcon, Object->GetClass())
+			]
+			+ SHorizontalBox::Slot()
+			.Padding(4.0f)
+			[
+				SNew(STextBlock)
+				.Text(FText::FromName(Object->GetFName()))
+			];
+	}
+
+	return SNew(STextBlock)
+		.Text(LOCTEXT("ObjectError", "Error. Invalid object."));
+}
+
+TSharedRef<SWidget> SMVVMPopoutEditor::ConstructFieldPreSlot(const UObject* InObject, UE::PropertyViewer::SPropertyViewer::FHandle Handle, const FFieldVariant FieldPath, const bool bIsForEvent)
+{
+	TSharedRef<SWidget> ImageWidget = SNullWidget::NullWidget;
+	const UClass* ObjectClass = InObject->GetClass();
+	FMVVMAvailableBinding Binding = bIsForEvent ? UMVVMSubsystem::GetAvailableBindingForEvent(FMVVMConstFieldVariant(FieldPath), ObjectClass) : UMVVMSubsystem::GetAvailableBindingForField(FMVVMConstFieldVariant(FieldPath), ObjectClass);
+	if (Binding.IsValid())
+	{
+		const FSlateBrush* Brush = nullptr;
+		if (Binding.HasNotify())
+		{
+			if (Binding.IsReadable() && Binding.IsWritable())
+			{
+				Brush = FMVVMEditorStyle::Get().GetBrush("BindingMode.TwoWay");
+			}
+			else if (Binding.IsReadable())
+			{
+				Brush = FMVVMEditorStyle::Get().GetBrush("BindingMode.OneWayToSource");
+			}
+			else if (Binding.IsWritable())
+			{
+				Brush = FMVVMEditorStyle::Get().GetBrush("BindingMode.OneWay");
+			}
+		}
+		else
+		{
+			if (Binding.IsReadable() && Binding.IsWritable())
+			{
+				Brush = FMVVMEditorStyle::Get().GetBrush("BindingMode.OneTimeTwoWay");
+			}
+			else if (Binding.IsReadable())
+			{
+				Brush = FMVVMEditorStyle::Get().GetBrush("BindingMode.OneTimeOneWay");
+			}
+			else if (Binding.IsWritable())
+			{
+				Brush = FMVVMEditorStyle::Get().GetBrush("BindingMode.OneTimeOneWayToSource");
+			}
+		}
+
+		if (Brush)
+		{
+			ImageWidget = SNew(SImage)
+				.Image(Brush);
+		}
+	}
+
+	return SNew(SBox)
+		.HAlign(HAlign_Center)
+		.VAlign(VAlign_Center)
+		.WidthOverride(16.0f)
+		.HeightOverride(16.0f)
+		[
+			ImageWidget
+		];
+}
+
+}
+
+#undef LOCTEXT_NAMESPACE
diff --git a/Engine/Plugins/Runtime/ModelViewViewModel/Source/ModelViewViewModelEditor/Private/Widgets/SMVVMPopoutEditor.h b/Engine/Plugins/Runtime/ModelViewViewModel/Source/ModelViewViewModelEditor/Private/Widgets/SMVVMPopoutEditor.h
new file mode 100644
--- /dev/null
+++ b/Engine/Plugins/Runtime/ModelViewViewModel/Source/ModelViewViewModelEditor/Private/Widgets/SMVVMPopoutEditor.h
@@ -0,0 +1,49 @@
+// Copyright Epic Games, Inc. All Rights Reserved.
+
+#pragma once
+
+#include "Widgets/SCompoundWidget.h"
+
+#include "Misc/NotNull.h"
+#include "UObject/WeakObjectPtr.h"
+#include "Widgets/PropertyViewer/SPropertyViewer.h"
+
+class UObject;
+
+namespace UE::PropertyViewer
+{
+	class FFieldExpander_Default;
+	class FFieldIterator_BlueprintVisible;
+}
+
+namespace UE::MVVM
+{
+
+class FFieldIterator_Bindable;
+
+class SMVVMPopoutEditor : public SCompoundWidget
+{
+public:
+	SLATE_BEGIN_ARGS(SMVVMPopoutEditor)
+		: _bReadOnly(true)
+		{}
+		SLATE_ARGUMENT(bool, bReadOnly)
+	SLATE_END_ARGS()
+
+	static void OpenPopoutEditor(TNotNull<UObject*> InObjectToInspect, bool bInReadOnly = true);
+
+	void Construct(const FArguments& InArgs, UObject* InObjectToInspect);
+
+private:
+	TWeakObjectPtr<UObject> ObjectWeak;
+	TUniquePtr<UE::PropertyViewer::FFieldIterator_BlueprintVisible> FieldIterator;
+	TUniquePtr<UE::PropertyViewer::FFieldExpander_Default> FieldExpander;
+	TSharedPtr<UE::PropertyViewer::SPropertyViewer> PropertyViewer;
+
+	TSharedPtr<SWidget> HandleGetPreSlot(UE::PropertyViewer::SPropertyViewer::FHandle Handle, TArrayView<const FFieldVariant> FieldPath);
+	TSharedRef<SWidget> HandleGenerateContainer(UE::PropertyViewer::SPropertyViewer::FHandle ContainerHandle, TOptional<FText> DisplayName);
+
+	TSharedRef<SWidget> ConstructFieldPreSlot(const UObject* InObject, UE::PropertyViewer::SPropertyViewer::FHandle Handle, const FFieldVariant FieldPath, const bool bIsForEvent = false);
+};
+
+} // UE::MVVM
\ No newline at end of file
diff --git a/Engine/Plugins/Runtime/ModelViewViewModel/Source/ModelViewViewModelEditor/Private/Widgets/SMVVMViewModelPanel.cpp b/Engine/Plugins/Runtime/ModelViewViewModel/Source/ModelViewViewModelEditor/Private/Widgets/SMVVMViewModelPanel.cpp
--- a/Engine/Plugins/Runtime/ModelViewViewModel/Source/ModelViewViewModelEditor/Private/Widgets/SMVVMViewModelPanel.cpp
+++ b/Engine/Plugins/Runtime/ModelViewViewModel/Source/ModelViewViewModelEditor/Private/Widgets/SMVVMViewModelPanel.cpp
@@ -4,10 +4,13 @@
 
 #include "DetailsViewArgs.h"
 #include "Dialogs/Dialogs.h"
+#include "IModelViewViewModelEditorModule.h"
 #include "IStructureDetailsView.h"
-#include "MVVMBlueprintView.h"
+#include "Modules/ModuleManager.h"
 #include "MVVMBlueprintInstancedViewModel.h"
+#include "MVVMBlueprintView.h"
 #include "MVVMDeveloperProjectSettings.h"
+#include "MVVMEditorCommands.h"
 #include "MVVMEditorSubsystem.h"
 #include "MVVMSubsystem.h"
 #include "MVVMWidgetBlueprintExtension_View.h"
@@ -32,6 +35,7 @@
 #include "Widgets/Layout/SBorder.h"
 #include "Widgets/Layout/SSplitter.h"
 #include "Widgets/PropertyViewer/SFieldIcon.h"
+#include "Widgets/SMVVMPopoutEditor.h"
 #include "Widgets/SMVVMSelectViewModel.h"
 #include "Widgets/SMVVMViewModelBindingListWidget.h"
 #include "Widgets/Text/SInlineEditableTextBlock.h"
@@ -431,6 +435,84 @@ void SMVVMViewModelPanel::HandleViewModelsUpdated()
 	FillViewModel();
 }
 
+TSet<UObject*> SMVVMViewModelPanel::GetSelectedEditableObjects() const
+{
+	if (!ViewModelTreeView.IsValid())
+	{
+		return {};
+	}
+
+	TSharedPtr<FWidgetBlueprintEditor> WidgetBlueprintEditor = WeakBlueprintEditor.Pin();
+	if (!WidgetBlueprintEditor)
+	{
+		return {};
+	}
+
+	UWidgetBlueprint* WidgetBP = WidgetBlueprintEditor->GetWidgetBlueprintObj();
+	if (!WidgetBP)
+	{
+		return {};
+	}
+
+	UMVVMBlueprintView* BlueprintView = WeakBlueprintView.Get();
+	if (!BlueprintView)
+	{
+		return {};
+	}
+
+	UUserWidget* Widget = Widget = Cast<UUserWidget>(WidgetBP->GetObjectBeingDebugged());
+	if (!Widget)
+	{
+		return {};
+	}
+
+	UMVVMView* WidgetView = UMVVMSubsystem::GetViewFromUserWidget(Widget);
+
+	TArray<UE::PropertyViewer::SPropertyViewer::FSelectedItem> Items = ViewModelTreeView->GetSelectedItems();
+	TSet<UObject*> Selected;
+
+	for (const UE::PropertyViewer::SPropertyViewer::FSelectedItem& Item : Items)
+	{
+		if (Item.bIsContainerSelected)
+		{
+			if (const FGuid* VMGuidPtr = PropertyViewerHandles.Find(Item.Handle))
+			{
+				// The widget instance itself is 0,0,0,0 = invalid.
+				if (!VMGuidPtr->IsValid())
+				{
+					Selected.Add(Widget);
+				}
+				else if (WidgetView)
+				{
+					if (const FMVVMBlueprintViewModelContext* ViewModelContext = BlueprintView->FindViewModel(*VMGuidPtr))
+					{
+						if (UObject* Object = WidgetView->GetViewModel(ViewModelContext->GetViewModelName()).GetObject())
+						{
+							Selected.Add(Object);
+						}
+					}
+				}
+			}
+		}
+	}
+
+	return Selected;
+}
+
+bool SMVVMViewModelPanel::HandleCanPopoutEditor() const
+{
+	return !GetSelectedEditableObjects().IsEmpty();
+}
+
+void SMVVMViewModelPanel::HandlePopoutEditor()
+{
+	IModelViewViewModelEditorModule& Module = FModuleManager::Get().GetModuleChecked<IModelViewViewModelEditorModule>("ModelViewViewModelEditor");
+
+	for (UObject* Object : GetSelectedEditableObjects())
+	{
+		Module.OpenPopoutEditor(Object, /* Read Only */ false);
+	}
+}
 
 void SMVVMViewModelPanel::OpenAddViewModelMenu()
 {
@@ -659,11 +741,18 @@ void SMVVMViewModelPanel::CreateCommandList()
 		FExecuteAction::CreateSP(this, &SMVVMViewModelPanel::HandleDeleteViewModel),
 		FCanExecuteAction::CreateSP(this, &SMVVMViewModelPanel::HandleCanDeleteViewModel)
 	);
+
 	CommandList->MapAction(
 		FGenericCommands::Get().Rename,
 		FExecuteAction::CreateSP(this, &SMVVMViewModelPanel::HandleRenameViewModel),
 		FCanExecuteAction::CreateSP(this, &SMVVMViewModelPanel::HandleCanRenameViewModel)
 	);
+
+	CommandList->MapAction(
+		FMVVMEditorCommands::Get().PopoutEditor,
+		FExecuteAction::CreateSP(this, &SMVVMViewModelPanel::HandlePopoutEditor),
+		FCanExecuteAction::CreateSP(this, &SMVVMViewModelPanel::HandleCanPopoutEditor)
+	);
 }
 
 
@@ -803,6 +892,14 @@ TSharedPtr<SWidget> SMVVMViewModelPanel::HandleContextMenuOpening(UE::PropertyVi
 		{
 			MenuBuilder.AddMenuEntry(FGenericCommands::Get().Delete);
 			MenuBuilder.AddMenuEntry(FGenericCommands::Get().Rename);
+			
+			MenuBuilder.AddMenuEntry(
+				FMVVMEditorCommands::Get().PopoutEditor,
+				NAME_None,
+				TAttribute<FText>(),
+				TAttribute<FText>(),
+				FSlateIcon(FAppStyle::GetAppStyleSetName(), "Icons.OpenInExternalEditor")
+			);
 		}
 		MenuBuilder.EndSection();
 
diff --git a/Engine/Plugins/Runtime/ModelViewViewModel/Source/ModelViewViewModelEditor/Private/Widgets/SMVVMViewModelPanel.h b/Engine/Plugins/Runtime/ModelViewViewModel/Source/ModelViewViewModelEditor/Private/Widgets/SMVVMViewModelPanel.h
--- a/Engine/Plugins/Runtime/ModelViewViewModel/Source/ModelViewViewModelEditor/Private/Widgets/SMVVMViewModelPanel.h
+++ b/Engine/Plugins/Runtime/ModelViewViewModel/Source/ModelViewViewModelEditor/Private/Widgets/SMVVMViewModelPanel.h
@@ -6,8 +6,10 @@
 #include "Misc/NotifyHook.h"
 #include "Widgets/PropertyViewer/SPropertyViewer.h"
 
-#include "UObject/StrongObjectPtr.h"
+#include "Containers/ContainersFwd.h"
 #include "MVVMBlueprintViewModelContext.h"
+#include "UObject/StrongObjectPtr.h"
+
 #include "SMVVMViewModelPanel.generated.h"
 
 namespace ETextCommit { enum Type : int; }
@@ -76,6 +78,9 @@ private:
 
 	void HandleViewUpdated(UBlueprintExtension* Extension);
 	void HandleViewModelsUpdated();
+	TSet<UObject*> GetSelectedEditableObjects() const;
+	bool HandleCanPopoutEditor() const;
+	void HandlePopoutEditor();
 
 	TSharedRef<SWidget> MakeAddMenu();
 	void HandleCancelAddMenu();
diff --git a/Engine/Plugins/Runtime/ModelViewViewModel/Source/ModelViewViewModelEditor/Public/IModelViewViewModelEditorModule.h b/Engine/Plugins/Runtime/ModelViewViewModel/Source/ModelViewViewModelEditor/Public/IModelViewViewModelEditorModule.h
new file mode 100644
--- /dev/null
+++ b/Engine/Plugins/Runtime/ModelViewViewModel/Source/ModelViewViewModelEditor/Public/IModelViewViewModelEditorModule.h
@@ -0,0 +1,17 @@
+// Copyright Epic Games, Inc. All Rights Reserved.
+
+#pragma once
+
+#include "Modules/ModuleInterface.h"
+#include "Misc/NotNull.h"
+
+class UObject;
+
+class IModelViewViewModelEditorModule : public IModuleInterface
+{
+public:
+	/**
+	 * Pops open a new window with a standalone object editor in either read only or edit mode.
+	 */
+	virtual void OpenPopoutEditor(TNotNull<UObject*> InObject, bool bInReadyOnly = true) = 0;
+};
diff --git a/Engine/Plugins/Runtime/ModelViewViewModel/Source/ModelViewViewModelEditor/Public/MVVMEditorCommands.h b/Engine/Plugins/Runtime/ModelViewViewModel/Source/ModelViewViewModelEditor/Public/MVVMEditorCommands.h
--- a/Engine/Plugins/Runtime/ModelViewViewModel/Source/ModelViewViewModelEditor/Public/MVVMEditorCommands.h
+++ b/Engine/Plugins/Runtime/ModelViewViewModel/Source/ModelViewViewModelEditor/Public/MVVMEditorCommands.h
@@ -13,5 +13,8 @@ public:
 
 	/** Command to toggle MVVM drawer in widget editor */
 	TSharedPtr<FUICommandInfo> ToggleMVVMDrawer;
+
+	/** Pops out a standalone property editor window. */
+	TSharedPtr<FUICommandInfo> PopoutEditor;
 };
 
