MVVM: fix default resolver being outered to the package
this caused the resolver object to be shown in source control changed assets

#Synced-CL 52103956
#Synced-CL 52103956

diff --git a/Engine/Plugins/Runtime/ModelViewViewModel/Source/ModelViewViewModelBlueprint/Private/MVVMBlueprintView.cpp b/Engine/Plugins/Runtime/ModelViewViewModel/Source/ModelViewViewModelBlueprint/Private/MVVMBlueprintView.cpp
--- a/Engine/Plugins/Runtime/ModelViewViewModel/Source/ModelViewViewModelBlueprint/Private/MVVMBlueprintView.cpp
+++ b/Engine/Plugins/Runtime/ModelViewViewModel/Source/ModelViewViewModelBlueprint/Private/MVVMBlueprintView.cpp
@@ -640,11 +640,11 @@ void UMVVMBlueprintView::PostLoad()
 
 	GetOuterUMVVMWidgetBlueprintExtension_View()->ConditionalPostLoad();
 	GetOuterUMVVMWidgetBlueprintExtension_View()->GetWidgetBlueprint()->ConditionalPostLoad();
+	UWidgetBlueprint* ThisWidgetBlueprint = GetOuterUMVVMWidgetBlueprintExtension_View()->GetWidgetBlueprint();
 
 	// Make sure all bindings uses the skeletal class
 	if (GetLinkerCustomVersion(FFortniteMainBranchObjectVersion::GUID) < FFortniteMainBranchObjectVersion::MVVMConvertPropertyPathToSkeletalClass)
 	{
-		const UWidgetBlueprint* ThisWidgetBlueprint = GetOuterUMVVMWidgetBlueprintExtension_View()->GetWidgetBlueprint();
 		for (FMVVMBlueprintViewBinding& Binding : Bindings)
 		{
 			for (const FMVVMBlueprintFieldPath& FieldPath : Binding.SourcePath.GetFieldPaths())
@@ -665,7 +665,6 @@ void UMVVMBlueprintView::PostLoad()
 
 	if (GetLinkerCustomVersion(FFortniteMainBranchObjectVersion::GUID) < FFortniteMainBranchObjectVersion::MVVMPropertyPathSelf)
 	{
-		const UWidgetBlueprint* ThisWidgetBlueprint = GetOuterUMVVMWidgetBlueprintExtension_View()->GetWidgetBlueprint();
 		auto DeprecateSelfPath = [ThisWidgetBlueprint](FMVVMBlueprintPropertyPath& Path) -> bool
 		{
 			Path.GetSource(ThisWidgetBlueprint);
@@ -676,9 +675,16 @@ void UMVVMBlueprintView::PostLoad()
 
 	for (FMVVMBlueprintViewModelContext& ViewModelContext : AvailableViewModels)
 	{
-		if (ViewModelContext.Resolver && !ViewModelContext.Resolver->HasAnyFlags(RF_Public))
+		if (ViewModelContext.Resolver)
 		{
-			ViewModelContext.Resolver->SetFlags(RF_Public);
+			if (!ViewModelContext.Resolver->HasAnyFlags(RF_Public))
+			{
+				ViewModelContext.Resolver->SetFlags(RF_Public);
+			}
+			if (ViewModelContext.Resolver->GetOuter() == ThisWidgetBlueprint->GetPackage())
+			{
+				ViewModelContext.Resolver->Rename(nullptr, ThisWidgetBlueprint, REN_None);
+			}
 		}
 	}
 }
diff --git a/Engine/Plugins/Runtime/ModelViewViewModel/Source/ModelViewViewModelBlueprint/Private/MVVMBlueprintViewModelContext.cpp b/Engine/Plugins/Runtime/ModelViewViewModel/Source/ModelViewViewModelBlueprint/Private/MVVMBlueprintViewModelContext.cpp
--- a/Engine/Plugins/Runtime/ModelViewViewModel/Source/ModelViewViewModelBlueprint/Private/MVVMBlueprintViewModelContext.cpp
+++ b/Engine/Plugins/Runtime/ModelViewViewModel/Source/ModelViewViewModelBlueprint/Private/MVVMBlueprintViewModelContext.cpp
@@ -32,7 +32,7 @@ FText FMVVMBlueprintViewModelContext::GetDisplayName() const
 }
 
 #if WITH_EDITOR
-TObjectPtr<UMVVMViewModelContextResolver> FMVVMBlueprintViewModelContext::CreateDefaultResolver(UPackage* Package) const
+TObjectPtr<UMVVMViewModelContextResolver> FMVVMBlueprintViewModelContext::CreateDefaultResolver(UObject* Outer) const
 {
 	TObjectPtr<UMVVMViewModelContextResolver> DefaultResolver;
 
@@ -44,7 +44,7 @@ TObjectPtr<UMVVMViewModelContextResolver> FMVVMBlueprintViewModelContext::Create
 			{
 				if (DefaultResolverClass->DoesSupportViewModelClass(GetViewModelClass()))
 				{
-					DefaultResolver = NewObject<UMVVMViewModelContextResolver>(Package, DefaultClass, NAME_None, RF_Public);
+					DefaultResolver = NewObject<UMVVMViewModelContextResolver>(Outer, DefaultClass, NAME_None, RF_Public);
 				}
 			}
 		}
diff --git a/Engine/Plugins/Runtime/ModelViewViewModel/Source/ModelViewViewModelBlueprint/Public/MVVMBlueprintViewModelContext.h b/Engine/Plugins/Runtime/ModelViewViewModel/Source/ModelViewViewModelBlueprint/Public/MVVMBlueprintViewModelContext.h
--- a/Engine/Plugins/Runtime/ModelViewViewModel/Source/ModelViewViewModelBlueprint/Public/MVVMBlueprintViewModelContext.h
+++ b/Engine/Plugins/Runtime/ModelViewViewModel/Source/ModelViewViewModelBlueprint/Public/MVVMBlueprintViewModelContext.h
@@ -107,7 +107,7 @@ public:
 	}
 
 #if WITH_EDITOR
-	[[nodiscard]] UE_API TObjectPtr<UMVVMViewModelContextResolver> CreateDefaultResolver(UPackage* Package) const;
+	[[nodiscard]] UE_API TObjectPtr<UMVVMViewModelContextResolver> CreateDefaultResolver(UObject* Outer) const;
 #endif
 
 private:
diff --git a/Engine/Plugins/Runtime/ModelViewViewModel/Source/ModelViewViewModelEditor/Private/MVVMEditorSubsystem.cpp b/Engine/Plugins/Runtime/ModelViewViewModel/Source/ModelViewViewModelEditor/Private/MVVMEditorSubsystem.cpp
--- a/Engine/Plugins/Runtime/ModelViewViewModel/Source/ModelViewViewModelEditor/Private/MVVMEditorSubsystem.cpp
+++ b/Engine/Plugins/Runtime/ModelViewViewModel/Source/ModelViewViewModelEditor/Private/MVVMEditorSubsystem.cpp
@@ -286,7 +286,7 @@ FGuid UMVVMEditorSubsystem::AddViewModel(UWidgetBlueprint* WidgetBlueprint, cons
 			FMVVMBlueprintViewModelContext Context = FMVVMBlueprintViewModelContext(ViewModelClass, *ViewModelName);
 			if (Context.IsValid())
 			{
-				Context.Resolver = Context.CreateDefaultResolver(WidgetBlueprint->GetPackage());
+				Context.Resolver = Context.CreateDefaultResolver(WidgetBlueprint);
 				Context.bUseAsInterface = bCanAddAsInterface;
 
 				const FScopedTransaction Transaction(LOCTEXT("AddViewModel", "Add viewmodel"));
