MVVM: fix static analysis error
#Synced-CL 52074045

#Synced-CL 52074045

diff --git a/Engine/Plugins/Runtime/ModelViewViewModel/Source/ModelViewViewModelEditor/Private/Widgets/SMVVMViewModelPanel.cpp b/Engine/Plugins/Runtime/ModelViewViewModel/Source/ModelViewViewModelEditor/Private/Widgets/SMVVMViewModelPanel.cpp
--- a/Engine/Plugins/Runtime/ModelViewViewModel/Source/ModelViewViewModelEditor/Private/Widgets/SMVVMViewModelPanel.cpp
+++ b/Engine/Plugins/Runtime/ModelViewViewModel/Source/ModelViewViewModelEditor/Private/Widgets/SMVVMViewModelPanel.cpp
@@ -460,7 +460,7 @@ TSet<UObject*> SMVVMViewModelPanel::GetSelectedEditableObjects() const
 		return {};
 	}
 
-	UUserWidget* Widget = Widget = Cast<UUserWidget>(WidgetBP->GetObjectBeingDebugged());
+	UUserWidget* Widget = Cast<UUserWidget>(WidgetBP->GetObjectBeingDebugged());
 	if (!Widget)
 	{
 		return {};
