[MVVM] Remove one-time binding error
* Fix some spelling errors

#tests Compiled blueprints
#Synced-CL 52176868

diff --git a/Engine/Plugins/Runtime/ModelViewViewModel/Source/ModelViewViewModelBlueprint/Private/MVVMViewBlueprintCompiler.cpp b/Engine/Plugins/Runtime/ModelViewViewModel/Source/ModelViewViewModelBlueprint/Private/MVVMViewBlueprintCompiler.cpp
--- a/Engine/Plugins/Runtime/ModelViewViewModel/Source/ModelViewViewModelBlueprint/Private/MVVMViewBlueprintCompiler.cpp
+++ b/Engine/Plugins/Runtime/ModelViewViewModel/Source/ModelViewViewModelBlueprint/Private/MVVMViewBlueprintCompiler.cpp
@@ -2174,7 +2174,7 @@ void FMVVMViewBlueprintCompiler::CreateReadFieldContexts(UWidgetBlueprintGenerat
 			return MakeValue(FieldContextResult.StealValue());
 		};
 
-		auto CreateFieldId = [Self = this, Class, &Binding, bIsOneTimeBinding = ValidBinding->bIsOneTimeBinding](const FMVVMBlueprintPropertyPath& PropertyPath, TSharedPtr<FGeneratedReadFieldPathContext>& ReadFieldContext, const FMVVMBlueprintPinId& PinId) -> TValueOrError<void, void>
+		auto CreateFieldId = [Self = this, Class, &Binding](const FMVVMBlueprintPropertyPath& PropertyPath, TSharedPtr<FGeneratedReadFieldPathContext>& ReadFieldContext, const FMVVMBlueprintPinId& PinId) -> TValueOrError<void, void>
 		{
 			if (!ReadFieldContext->NotificationField.IsValid())
 			{
@@ -3070,12 +3070,17 @@ void FMVVMViewBlueprintCompiler::PreCompileBindings(UWidgetBlueprintGeneratedCla
 		bool bRequiresValidFieldId = !ValidBinding.bIsOneTimeBinding;
 		if (bRequiresValidFieldId && !bHasFieldId)
 		{
-			Self->AddMessageForBinding(Binding
-				, LOCTEXT("CouldNotCreateSourceFields", "There is no field to bind to. The binding must be a OneTime binding.")
-				, Compiler::EMessageType::Error
-				, FMVVMBlueprintPinId()
-			);
-			return false;
+			const bool bShowOneTimeWarning = GetDefault<UMVVMDeveloperProjectSettings>()->bShowOneTimeBindingWarning &&
+				(!IsRunningCookCommandlet() || !GetDefault<UMVVMDeveloperProjectSettings>()->bShowOneTimeBindingWarningEditorOnly);
+
+			if (bShowOneTimeWarning)
+			{
+				Self->AddMessageForBinding(Binding
+					, LOCTEXT("CouldNotCreateSourceFields", "There is no field to bind to. The binding must be a OneTime binding.")
+					, Compiler::EMessageType::Warning
+					, FMVVMBlueprintPinId()
+				);
+			}
 		}
 
 		return true;
diff --git a/Engine/Plugins/Runtime/ModelViewViewModel/Source/ModelViewViewModelBlueprint/Public/MVVMDeveloperProjectSettings.h b/Engine/Plugins/Runtime/ModelViewViewModel/Source/ModelViewViewModelBlueprint/Public/MVVMDeveloperProjectSettings.h
--- a/Engine/Plugins/Runtime/ModelViewViewModel/Source/ModelViewViewModelBlueprint/Public/MVVMDeveloperProjectSettings.h
+++ b/Engine/Plugins/Runtime/ModelViewViewModel/Source/ModelViewViewModelBlueprint/Public/MVVMDeveloperProjectSettings.h
@@ -219,6 +219,14 @@ public:
 	UPROPERTY()
 	bool bAllowWidgetInConditionSource = false;
 
+	/** Show compiler warning if a one way binding is used on a property that has no field notify */
+	UPROPERTY(EditAnywhere, config, Category = "Features")
+	bool bShowOneTimeBindingWarning = true;
+
+	/** Restrict the one time binding warning to editor only so it doesn't cause issues in cook logs */
+	UPROPERTY(EditAnywhere, config, Category = "Features")
+	bool bShowOneTimeBindingWarningEditorOnly = true;
+
 private:
 	/** MVVM Generated function permission list */
 	FPathPermissionList GeneratedFunctionPermissions;
diff --git a/Engine/Source/Runtime/UMG/Public/Extensions/ContextDataWidgetExtension.h b/Engine/Source/Runtime/UMG/Public/Extensions/ContextDataWidgetExtension.h
--- a/Engine/Source/Runtime/UMG/Public/Extensions/ContextDataWidgetExtension.h
+++ b/Engine/Source/Runtime/UMG/Public/Extensions/ContextDataWidgetExtension.h
@@ -25,7 +25,7 @@ public:
 	{
 		for (UObject* Data : ContextData)
 		{
-			if (Data && Data->IsA(Class)) // Can bedome null from GC
+			if (Data && Data->IsA(Class)) // Can become null from GC
 			{
 				return Data;
 			}
@@ -35,7 +35,7 @@ public:
 
 	void AddData(UObject* Data)
 	{
-		// Todo: Support multiple instances? What about exsiting data that is of this type?
+		// Todo: Support multiple instances? What about existing data that is of this type?
 		if (ensure(Data) && ensure(!FindData(Data->GetClass())))
 		{
 			ContextData.Add(Data);
