[UMG] Fix mvvm view to only remove specific delegate registrations tied to a source

#tests Used to fix local issue of multiple bindings to the same source object through different binding paths
#Synced-CL 52322456

diff --git a/Engine/Plugins/Runtime/ModelViewViewModel/Source/ModelViewViewModel/Private/View/MVVMView.cpp b/Engine/Plugins/Runtime/ModelViewViewModel/Source/ModelViewViewModel/Private/View/MVVMView.cpp
--- a/Engine/Plugins/Runtime/ModelViewViewModel/Source/ModelViewViewModel/Private/View/MVVMView.cpp
+++ b/Engine/Plugins/Runtime/ModelViewViewModel/Source/ModelViewViewModel/Private/View/MVVMView.cpp
@@ -507,14 +507,18 @@ void UMVVMView::InitializeSourceBindings(FMVVMView_SourceKey SourceKey, bool bRu
 			TScriptInterface<INotifyFieldValueChanged> NotifyFieldValueChanged = ViewSource.Source->Implements<UNotifyFieldValueChanged>() ? TScriptInterface<INotifyFieldValueChanged>(ViewSource.Source) : TScriptInterface<INotifyFieldValueChanged>();
 			if (NotifyFieldValueChanged.GetInterface() != nullptr)
 			{
+				int32 HandleIndex = 0;
 				for (const FMVVMViewClass_FieldId& FieldId : ClassSource.GetFieldIds())
 				{
 					if (ensureMsgf(FieldId.IsValid(), TEXT("Invalid field. It failed somewhere in the compiler.")))
 					{
 						UE::FieldNotification::FFieldMulticastDelegate::FDelegate Delegate = UE::FieldNotification::FFieldMulticastDelegate::FDelegate::CreateUObject(this, &UMVVMView::HandledLibraryBindingValueChanged, SourceKey);
-						NotifyFieldValueChanged->AddFieldValueChangedDelegate(FieldId.GetFieldId(), MoveTemp(Delegate));
+						FDelegateHandle Handle = NotifyFieldValueChanged->AddFieldValueChangedDelegate(FieldId.GetFieldId(), MoveTemp(Delegate));
+						ensure(Handle.IsValid());
+						RegisteredBindings.Add( { SourceKey.GetIndex(), HandleIndex }, Handle);
 						++ViewSource.RegisteredCount;
 					}
+					++HandleIndex;
 				}
 			}
 		}
@@ -595,7 +599,22 @@ void UMVVMView::UninitializeSourceBindings(FMVVMViewClass_SourceKey SourceKey, c
 		{
 			TScriptInterface<INotifyFieldValueChanged> SourceAsInterface = ViewSource.Source;
 			checkf(SourceAsInterface.GetInterface(), TEXT("It was added as a INotifyFieldValueChanged. It should still be."));
-			SourceAsInterface->RemoveAllFieldValueChangedDelegates(this);
+
+			int32 HandleIndex = 0;
+			int32 HandlesRemoved = 0;
+			for (const FMVVMViewClass_FieldId& FieldId : ClassSource.GetFieldIds())
+			{
+				if (FieldId.IsValid())
+				{
+					FDelegateHandle Handle;
+					ensure(RegisteredBindings.RemoveAndCopyValue({ SourceKey.GetIndex(), HandleIndex }, Handle));
+					ensure(SourceAsInterface->RemoveFieldValueChangedDelegate(FieldId.GetFieldId(), Handle));
+					++HandlesRemoved;
+				}
+
+				++HandleIndex;
+			}
+			ensure(HandlesRemoved == ViewSource.RegisteredCount);
 			ViewSource.RegisteredCount = 0;
 		}
 		else
diff --git a/Engine/Plugins/Runtime/ModelViewViewModel/Source/ModelViewViewModel/Public/View/MVVMView.h b/Engine/Plugins/Runtime/ModelViewViewModel/Source/ModelViewViewModel/Public/View/MVVMView.h
--- a/Engine/Plugins/Runtime/ModelViewViewModel/Source/ModelViewViewModel/Public/View/MVVMView.h
+++ b/Engine/Plugins/Runtime/ModelViewViewModel/Source/ModelViewViewModel/Public/View/MVVMView.h
@@ -246,6 +246,8 @@ private:
 	/** The event that are registered by the view to the sources. */
 	TArray<FBoundEvent> BoundEvents;
 
+	TMap<TPair<int32 /* Source Index */, int32 /* Field Index */>, FDelegateHandle> RegisteredBindings;
+
 	UPROPERTY(VisibleAnywhere, Transient, Category = "View")
 	TArray<TObjectPtr<UMVVMViewExtension>> Extensions;
 
