[MVVM]: Changed one time binding warning to info only. It will no longer prevent compilation, and will only display an icon,  with a mention, when hovering it, that binding may not trigger in cases you might expect it would .

#tests Created "problematic" binding (source without field notify + one way binding), info is displayed as expected, and cook still succeed.
#Synced-CL 52488042

diff --git a/Engine/Plugins/Runtime/ModelViewViewModel/Source/ModelViewViewModelBlueprint/Private/MVVMBlueprintView.cpp b/Engine/Plugins/Runtime/ModelViewViewModel/Source/ModelViewViewModelBlueprint/Private/MVVMBlueprintView.cpp
--- a/Engine/Plugins/Runtime/ModelViewViewModel/Source/ModelViewViewModelBlueprint/Private/MVVMBlueprintView.cpp
+++ b/Engine/Plugins/Runtime/ModelViewViewModel/Source/ModelViewViewModelBlueprint/Private/MVVMBlueprintView.cpp
@@ -417,6 +417,11 @@ TArray<FText> UMVVMBlueprintView::GetBindingMessages(FGuid Id, UE::MVVM::EBindin
 	return Results;
 }
 
+bool UMVVMBlueprintView::HasBindingMessage(FGuid Id) const
+{
+	return BindingMessages.Find(Id) != nullptr;
+}
+
 bool UMVVMBlueprintView::HasBindingMessage(FGuid Id, UE::MVVM::EBindingMessageType InMessageType) const
 {
 	if (const TArray<UE::MVVM::FBindingMessage>* FoundBindingMessages = BindingMessages.Find(Id))
diff --git a/Engine/Plugins/Runtime/ModelViewViewModel/Source/ModelViewViewModelBlueprint/Private/MVVMViewBlueprintCompiler.cpp b/Engine/Plugins/Runtime/ModelViewViewModel/Source/ModelViewViewModelBlueprint/Private/MVVMViewBlueprintCompiler.cpp
--- a/Engine/Plugins/Runtime/ModelViewViewModel/Source/ModelViewViewModelBlueprint/Private/MVVMViewBlueprintCompiler.cpp
+++ b/Engine/Plugins/Runtime/ModelViewViewModel/Source/ModelViewViewModelBlueprint/Private/MVVMViewBlueprintCompiler.cpp
@@ -3069,17 +3069,11 @@ void FMVVMViewBlueprintCompiler::PreCompileBindings(UWidgetBlueprintGeneratedCla
 		bool bRequiresValidFieldId = !ValidBinding.bIsOneTimeBinding;
 		if (bRequiresValidFieldId && !bHasFieldId)
 		{
-			const bool bShowOneTimeWarning = GetDefault<UMVVMDeveloperProjectSettings>()->bShowOneTimeBindingWarning &&
-				(!IsRunningCookCommandlet() || !GetDefault<UMVVMDeveloperProjectSettings>()->bShowOneTimeBindingWarningEditorOnly);
-
-			if (bShowOneTimeWarning)
-			{
-				Self->AddMessageForBinding(Binding
-					, LOCTEXT("CouldNotCreateSourceFields", "There is no field to bind to. The binding must be a OneTime binding.")
-					, Compiler::EMessageType::Warning
-					, FMVVMBlueprintPinId()
-				);
-			}
+			Self->AddMessageForBinding(Binding
+				, LOCTEXT("CouldNotCreateSourceFields", "This binding will only trigger ONCE whenever the value updates (typically when the value is first initialized).\n To ensure it triggers EVERY time the value updates, turn on Field Notify on this variable.")
+				, Compiler::EMessageType::Info
+				, FMVVMBlueprintPinId()
+			);
 		}
 
 		return true;
diff --git a/Engine/Plugins/Runtime/ModelViewViewModel/Source/ModelViewViewModelBlueprint/Public/MVVMBlueprintView.h b/Engine/Plugins/Runtime/ModelViewViewModel/Source/ModelViewViewModelBlueprint/Public/MVVMBlueprintView.h
--- a/Engine/Plugins/Runtime/ModelViewViewModel/Source/ModelViewViewModelBlueprint/Public/MVVMBlueprintView.h
+++ b/Engine/Plugins/Runtime/ModelViewViewModel/Source/ModelViewViewModelBlueprint/Public/MVVMBlueprintView.h
@@ -175,6 +175,7 @@ public:
 	}
 
 	UE_API TArray<FText> GetBindingMessages(FGuid Id, UE::MVVM::EBindingMessageType InMessageType) const;
+	UE_API bool HasBindingMessage(FGuid Id) const;
 	UE_API bool HasBindingMessage(FGuid Id, UE::MVVM::EBindingMessageType InMessageType) const;
 	UE_API void AddMessageToBinding(FGuid Id, UE::MVVM::FBindingMessage MessageToAdd);
 	UE_API void ResetBindingMessages();
diff --git a/Engine/Plugins/Runtime/ModelViewViewModel/Source/ModelViewViewModelBlueprint/Public/MVVMDeveloperProjectSettings.h b/Engine/Plugins/Runtime/ModelViewViewModel/Source/ModelViewViewModelBlueprint/Public/MVVMDeveloperProjectSettings.h
--- a/Engine/Plugins/Runtime/ModelViewViewModel/Source/ModelViewViewModelBlueprint/Public/MVVMDeveloperProjectSettings.h
+++ b/Engine/Plugins/Runtime/ModelViewViewModel/Source/ModelViewViewModelBlueprint/Public/MVVMDeveloperProjectSettings.h
@@ -219,14 +219,6 @@ public:
 	UPROPERTY()
 	bool bAllowWidgetInConditionSource = false;
 
-	/** Show compiler warning if a one way binding is used on a property that has no field notify */
-	UPROPERTY(EditAnywhere, config, Category = "Features")
-	bool bShowOneTimeBindingWarning = true;
-
-	/** Restrict the one time binding warning to editor only so it doesn't cause issues in cook logs */
-	UPROPERTY(EditAnywhere, config, Category = "Features")
-	bool bShowOneTimeBindingWarningEditorOnly = true;
-
 private:
 	/** MVVM Generated function permission list */
 	FPathPermissionList GeneratedFunctionPermissions;
diff --git a/Engine/Plugins/Runtime/ModelViewViewModel/Source/ModelViewViewModelEditor/Private/Widgets/BindingEntry/SMVVMBindingRow.cpp b/Engine/Plugins/Runtime/ModelViewViewModel/Source/ModelViewViewModelEditor/Private/Widgets/BindingEntry/SMVVMBindingRow.cpp
--- a/Engine/Plugins/Runtime/ModelViewViewModel/Source/ModelViewViewModelEditor/Private/Widgets/BindingEntry/SMVVMBindingRow.cpp
+++ b/Engine/Plugins/Runtime/ModelViewViewModel/Source/ModelViewViewModelEditor/Private/Widgets/BindingEntry/SMVVMBindingRow.cpp
@@ -281,9 +281,8 @@ EVisibility SBindingRow::GetErrorButtonVisibility() const
 	if (const FMVVMBlueprintViewBinding* ViewModelBinding = GetThisViewBinding())
 	{
 		UMVVMBlueprintView* BlueprintView = GetBlueprintView();
-		bool HasBindingError = BlueprintView->HasBindingMessage(ViewModelBinding->BindingId, EBindingMessageType::Error)
-			|| BlueprintView->HasBindingMessage(ViewModelBinding->BindingId, EBindingMessageType::Warning);
-		return HasBindingError ? EVisibility::Visible : EVisibility::Hidden;
+		const bool HasBindingMessage = BlueprintView->HasBindingMessage(ViewModelBinding->BindingId);
+		return HasBindingMessage ? EVisibility::Visible : EVisibility::Hidden;
 	}
 	return EVisibility::Collapsed;
 }
@@ -296,10 +295,12 @@ FText SBindingRow::GetErrorButtonToolTip() const
 		UMVVMBlueprintView* BlueprintView = GetBlueprintView();
 		TArray<FText> BindingErrorList = BlueprintView->GetBindingMessages(ViewModelBinding->BindingId, EBindingMessageType::Error);
 		TArray<FText> BindingWarningList = BlueprintView->GetBindingMessages(ViewModelBinding->BindingId, EBindingMessageType::Warning);
+		TArray<FText> BindingInfoList = BlueprintView->GetBindingMessages(ViewModelBinding->BindingId, EBindingMessageType::Info);
 		BindingErrorList.Append(BindingWarningList);
+		BindingErrorList.Append(BindingInfoList);
 
 		static const FText NewLineText = FText::FromString(TEXT("\n"));
-		FText HintText = LOCTEXT("ErrorButtonText", "Errors and Warnings: (Click to show in a separate window)");
+		FText HintText = LOCTEXT("ErrorButtonText", "Messages: (Click to show in a separate window)");
 		FText ErrorsText = FText::Join(NewLineText, BindingErrorList);
 		return FText::Join(NewLineText, HintText, ErrorsText);
 	}
@@ -325,9 +326,14 @@ FReply SBindingRow::OnErrorButtonClicked()
 			ErrorItems.Add(MakeShared<FText>(WarningText));
 		}
 
+		for (const FText& InfoText : View->GetBindingMessages(ViewBinding->BindingId, EBindingMessageType::Info))
+		{
+			ErrorItems.Add(MakeShared<FText>(InfoText));
+		}
+
 		const FText BindingDisplayName = FText::FromString(ViewBinding->GetDisplayNameString(WidgetBlueprint));
 		TSharedRef<SCustomDialog> ErrorDialog = SNew(SCustomDialog)
-			.Title(FText::Format(LOCTEXT("Compilation Errors and Warnings", "Compilation Errors and Warnings for {0}"), BindingDisplayName))
+			.Title(FText::Format(LOCTEXT("Compilation Messages", "Compilation messages for {0}"), BindingDisplayName))
 			.Buttons({
 				SCustomDialog::FButton(LOCTEXT("OK", "OK"))
 				})
