PR #14649: MVVMObjectVariant.h: Fix comparison to use GetUObject method.

#rnx
#Synced-CL 52587275

diff --git a/Engine/Plugins/Runtime/ModelViewViewModel/Source/ModelViewViewModel/Public/Types/MVVMObjectVariant.h b/Engine/Plugins/Runtime/ModelViewViewModel/Source/ModelViewViewModel/Public/Types/MVVMObjectVariant.h
--- a/Engine/Plugins/Runtime/ModelViewViewModel/Source/ModelViewViewModel/Public/Types/MVVMObjectVariant.h
+++ b/Engine/Plugins/Runtime/ModelViewViewModel/Source/ModelViewViewModel/Public/Types/MVVMObjectVariant.h
@@ -165,7 +165,7 @@ namespace UE::MVVM
 			}
 			if (IsUObject())
 			{
-				return GetUObject() == B.GetObject();
+				return GetUObject() == B.GetUObject();
 			}
 			return GetUScriptStructData() == B.GetUScriptStructData();
 		}
